MovieScenes: Renamed module; reorganized files; cleaned up include dependencies.

diff --git a/Engine/Source/Editor/MovieSceneTools/MovieSceneTools.Build.cs b/Engine/Source/Editor/MovieSceneTools/MovieSceneTools.Build.cs
--- a/Engine/Source/Editor/MovieSceneTools/MovieSceneTools.Build.cs
+++ b/Engine/Source/Editor/MovieSceneTools/MovieSceneTools.Build.cs
@@ -28,7 +28,7 @@ public class MovieSceneTools : ModuleRules
 		PrivateDependencyModuleNames.AddRange(
 			new string[] {
 				"MovieSceneCore",
-				"MovieSceneCoreTypes",
+				"MovieSceneTracks",
 				"BlueprintGraph",
 				"Slate",
 				"SlateCore",
diff --git a/Engine/Source/Editor/MovieSceneTools/Private/TrackEditors/AudioTrackEditor.cpp b/Engine/Source/Editor/MovieSceneTools/Private/TrackEditors/AudioTrackEditor.cpp
--- a/Engine/Source/Editor/MovieSceneTools/Private/TrackEditors/AudioTrackEditor.cpp
+++ b/Engine/Source/Editor/MovieSceneTools/Private/TrackEditors/AudioTrackEditor.cpp
@@ -22,7 +22,7 @@
 #include "Runtime/Engine/Public/Slate/SlateTextures.h"
 #include "ObjectTools.h"
 #include "Runtime/Engine/Public/Slate/SceneViewport.h"
-#include "Runtime/MovieSceneCoreTypes/Classes/MovieSceneShotSection.h"
+#include "Runtime/MovieSceneTracks/Public/Sections/MovieSceneShotSection.h"
 #include "Runtime/Engine/Public/AudioDecompress.h"
 #include "Sound/SoundBase.h"
 #include "Sound/SoundWave.h"
diff --git a/Engine/Source/Editor/MovieSceneTools/Private/TrackEditors/DirectorTrackEditor.cpp b/Engine/Source/Editor/MovieSceneTools/Private/TrackEditors/DirectorTrackEditor.cpp
--- a/Engine/Source/Editor/MovieSceneTools/Private/TrackEditors/DirectorTrackEditor.cpp
+++ b/Engine/Source/Editor/MovieSceneTools/Private/TrackEditors/DirectorTrackEditor.cpp
@@ -14,7 +14,7 @@
 #include "Runtime/Engine/Public/Slate/SlateTextures.h"
 #include "ObjectTools.h"
 #include "Runtime/Engine/Public/Slate/SceneViewport.h"
-#include "Runtime/MovieSceneCoreTypes/Classes/MovieSceneShotSection.h"
+#include "Runtime/MovieSceneTracks/Public/Sections/MovieSceneShotSection.h"
 #include "IKeyArea.h"
 #include "MovieSceneToolHelpers.h"
 #include "MovieSceneTrackEditor.h"
diff --git a/Engine/Source/Editor/Sequencer/Sequencer.Build.cs b/Engine/Source/Editor/Sequencer/Sequencer.Build.cs
--- a/Engine/Source/Editor/Sequencer/Sequencer.Build.cs
+++ b/Engine/Source/Editor/Sequencer/Sequencer.Build.cs
@@ -19,7 +19,7 @@ public class Sequencer : ModuleRules
                 "EditorStyle",
 				"UnrealEd", 
 				"MovieSceneCore", 
-				"MovieSceneCoreTypes", 
+				"MovieSceneTracks", 
 				"MovieSceneTools", 
 				"EditorWidgets", 
 				"SequencerWidgets",
diff --git a/Engine/Source/Runtime/Engine/Engine.Build.cs b/Engine/Source/Runtime/Engine/Engine.Build.cs
--- a/Engine/Source/Runtime/Engine/Engine.Build.cs
+++ b/Engine/Source/Runtime/Engine/Engine.Build.cs
@@ -86,7 +86,7 @@ public class Engine : ModuleRules
 			new string[]
 			{
 				"MovieSceneCore",
-				"MovieSceneCoreTypes",
+				"MovieSceneTracks",
 				"HeadMountedDisplay",
 				"StreamingPauseRendering",
 			}
diff --git a/Engine/Source/Runtime/MovieSceneCore/Public/IMovieSceneTrackInstance.h b/Engine/Source/Runtime/MovieSceneCore/Public/IMovieSceneTrackInstance.h
--- a/Engine/Source/Runtime/MovieSceneCore/Public/IMovieSceneTrackInstance.h
+++ b/Engine/Source/Runtime/MovieSceneCore/Public/IMovieSceneTrackInstance.h
@@ -2,23 +2,29 @@
 
 #pragma once
 
+
+class IMovieScenePlayer;
+
+
 /**
  * Base class for an animated type instance in a Movie Scene
  */
 class IMovieSceneTrackInstance
 {
 public:
-	virtual ~IMovieSceneTrackInstance() {}
+
+	/** Virtual destructor. */
+	virtual ~IMovieSceneTrackInstance() { }
 
 	/**
-	* Save state of objects that this instance will be editing.
-	*/
-	virtual void SaveState ( const TArray<UObject*>& RuntimeObjects ) {}
+	 * Save state of objects that this instance will be editing.
+	 */
+	virtual void SaveState(const TArray<UObject*>& RuntimeObjects) { }
 
 	/**
-	* Restore state of objects that this instance edited.
-	*/
-	virtual void RestoreState ( const TArray<UObject*>& RuntimeObjects ) {}
+	 * Restore state of objects that this instance edited.
+	 */
+	virtual void RestoreState(const TArray<UObject*>& RuntimeObjects) { }
 
 	/**
 	 * Main update function for track instances.  Called in game and in editor when we update a moviescene.
@@ -28,11 +34,10 @@ public:
 	 * @param RuntimeObjects	Runtime objects bound to this instance (if any)
 	 * @param Player			The playback interface.  Contains state and some other functionality for runtime playback
 	 */
-	virtual void Update( float Position, float LastPosition, const TArray<UObject*>& RuntimeObjects, class IMovieScenePlayer& Player ) = 0;
+	virtual void Update(float Position, float LastPosition, const TArray<UObject*>& RuntimeObjects, IMovieScenePlayer& Player) = 0;
 	
-
 	/**
 	 * Refreshes the current instance
 	 */
-	virtual void RefreshInstance( const TArray<UObject*>& RuntimeObjects, class IMovieScenePlayer& Player ) = 0;
+	virtual void RefreshInstance(const TArray<UObject*>& RuntimeObjects, IMovieScenePlayer& Player) = 0;
 };
diff --git a/Engine/Source/Runtime/MovieSceneCoreTypes/Classes/MovieScene3DTransformSection.h b/Engine/Source/Runtime/MovieSceneCoreTypes/Classes/MovieScene3DTransformSection.h
deleted file mode 100644
diff --git a/Engine/Source/Runtime/MovieSceneCoreTypes/Classes/MovieScene3DTransformTrack.h b/Engine/Source/Runtime/MovieSceneCoreTypes/Classes/MovieScene3DTransformTrack.h
deleted file mode 100644
diff --git a/Engine/Source/Runtime/MovieSceneCoreTypes/Classes/MovieSceneAnimationSection.h b/Engine/Source/Runtime/MovieSceneCoreTypes/Classes/MovieSceneAnimationSection.h
deleted file mode 100644
diff --git a/Engine/Source/Runtime/MovieSceneCoreTypes/Classes/MovieSceneAnimationTrack.h b/Engine/Source/Runtime/MovieSceneCoreTypes/Classes/MovieSceneAnimationTrack.h
deleted file mode 100644
diff --git a/Engine/Source/Runtime/MovieSceneCoreTypes/Classes/MovieSceneAudioSection.h b/Engine/Source/Runtime/MovieSceneCoreTypes/Classes/MovieSceneAudioSection.h
deleted file mode 100644
diff --git a/Engine/Source/Runtime/MovieSceneCoreTypes/Classes/MovieSceneAudioTrack.h b/Engine/Source/Runtime/MovieSceneCoreTypes/Classes/MovieSceneAudioTrack.h
deleted file mode 100644
diff --git a/Engine/Source/Runtime/MovieSceneCoreTypes/Classes/MovieSceneBoolSection.h b/Engine/Source/Runtime/MovieSceneCoreTypes/Classes/MovieSceneBoolSection.h
deleted file mode 100644
diff --git a/Engine/Source/Runtime/MovieSceneCoreTypes/Classes/MovieSceneBoolTrack.h b/Engine/Source/Runtime/MovieSceneCoreTypes/Classes/MovieSceneBoolTrack.h
deleted file mode 100644
diff --git a/Engine/Source/Runtime/MovieSceneCoreTypes/Classes/MovieSceneByteSection.h b/Engine/Source/Runtime/MovieSceneCoreTypes/Classes/MovieSceneByteSection.h
deleted file mode 100644
diff --git a/Engine/Source/Runtime/MovieSceneCoreTypes/Classes/MovieSceneByteTrack.h b/Engine/Source/Runtime/MovieSceneCoreTypes/Classes/MovieSceneByteTrack.h
deleted file mode 100644
diff --git a/Engine/Source/Runtime/MovieSceneCoreTypes/Classes/MovieSceneColorSection.h b/Engine/Source/Runtime/MovieSceneCoreTypes/Classes/MovieSceneColorSection.h
deleted file mode 100644
diff --git a/Engine/Source/Runtime/MovieSceneCoreTypes/Classes/MovieSceneColorTrack.h b/Engine/Source/Runtime/MovieSceneCoreTypes/Classes/MovieSceneColorTrack.h
deleted file mode 100644
diff --git a/Engine/Source/Runtime/MovieSceneCoreTypes/Classes/MovieSceneDirectorTrack.h b/Engine/Source/Runtime/MovieSceneCoreTypes/Classes/MovieSceneDirectorTrack.h
deleted file mode 100644
diff --git a/Engine/Source/Runtime/MovieSceneCoreTypes/Classes/MovieSceneFloatSection.h b/Engine/Source/Runtime/MovieSceneCoreTypes/Classes/MovieSceneFloatSection.h
deleted file mode 100644
diff --git a/Engine/Source/Runtime/MovieSceneCoreTypes/Classes/MovieSceneFloatTrack.h b/Engine/Source/Runtime/MovieSceneCoreTypes/Classes/MovieSceneFloatTrack.h
deleted file mode 100644
diff --git a/Engine/Source/Runtime/MovieSceneCoreTypes/Classes/MovieSceneParticleSection.h b/Engine/Source/Runtime/MovieSceneCoreTypes/Classes/MovieSceneParticleSection.h
deleted file mode 100644
diff --git a/Engine/Source/Runtime/MovieSceneCoreTypes/Classes/MovieSceneParticleTrack.h b/Engine/Source/Runtime/MovieSceneCoreTypes/Classes/MovieSceneParticleTrack.h
deleted file mode 100644
diff --git a/Engine/Source/Runtime/MovieSceneCoreTypes/Classes/MovieSceneShotSection.h b/Engine/Source/Runtime/MovieSceneCoreTypes/Classes/MovieSceneShotSection.h
deleted file mode 100644
diff --git a/Engine/Source/Runtime/MovieSceneCoreTypes/Classes/MovieSceneVectorSection.h b/Engine/Source/Runtime/MovieSceneCoreTypes/Classes/MovieSceneVectorSection.h
deleted file mode 100644
diff --git a/Engine/Source/Runtime/MovieSceneCoreTypes/Classes/MovieSceneVectorTrack.h b/Engine/Source/Runtime/MovieSceneCoreTypes/Classes/MovieSceneVectorTrack.h
deleted file mode 100644
diff --git a/Engine/Source/Runtime/MovieSceneCoreTypes/Classes/SubMovieSceneSection.h b/Engine/Source/Runtime/MovieSceneCoreTypes/Classes/SubMovieSceneSection.h
deleted file mode 100644
diff --git a/Engine/Source/Runtime/MovieSceneCoreTypes/Classes/SubMovieSceneTrack.h b/Engine/Source/Runtime/MovieSceneCoreTypes/Classes/SubMovieSceneTrack.h
deleted file mode 100644
diff --git a/Engine/Source/Runtime/MovieSceneCoreTypes/MovieSceneCoreTypes.Build.cs b/Engine/Source/Runtime/MovieSceneCoreTypes/MovieSceneCoreTypes.Build.cs
deleted file mode 100644
diff --git a/Engine/Source/Runtime/MovieSceneCoreTypes/Private/Instances/MovieScene3DTransformTrackInstance.cpp b/Engine/Source/Runtime/MovieSceneCoreTypes/Private/Instances/MovieScene3DTransformTrackInstance.cpp
deleted file mode 100644
diff --git a/Engine/Source/Runtime/MovieSceneCoreTypes/Private/Instances/MovieScene3DTransformTrackInstance.h b/Engine/Source/Runtime/MovieSceneCoreTypes/Private/Instances/MovieScene3DTransformTrackInstance.h
deleted file mode 100644
diff --git a/Engine/Source/Runtime/MovieSceneCoreTypes/Private/Instances/MovieSceneAnimationTrackInstance.cpp b/Engine/Source/Runtime/MovieSceneCoreTypes/Private/Instances/MovieSceneAnimationTrackInstance.cpp
deleted file mode 100644
diff --git a/Engine/Source/Runtime/MovieSceneCoreTypes/Private/Instances/MovieSceneAnimationTrackInstance.h b/Engine/Source/Runtime/MovieSceneCoreTypes/Private/Instances/MovieSceneAnimationTrackInstance.h
deleted file mode 100644
diff --git a/Engine/Source/Runtime/MovieSceneCoreTypes/Private/Instances/MovieSceneAudioTrackInstance.cpp b/Engine/Source/Runtime/MovieSceneCoreTypes/Private/Instances/MovieSceneAudioTrackInstance.cpp
deleted file mode 100644
diff --git a/Engine/Source/Runtime/MovieSceneCoreTypes/Private/Instances/MovieSceneAudioTrackInstance.h b/Engine/Source/Runtime/MovieSceneCoreTypes/Private/Instances/MovieSceneAudioTrackInstance.h
deleted file mode 100644
diff --git a/Engine/Source/Runtime/MovieSceneCoreTypes/Private/Instances/MovieSceneBoolTrackInstance.cpp b/Engine/Source/Runtime/MovieSceneCoreTypes/Private/Instances/MovieSceneBoolTrackInstance.cpp
deleted file mode 100644
diff --git a/Engine/Source/Runtime/MovieSceneCoreTypes/Private/Instances/MovieSceneBoolTrackInstance.h b/Engine/Source/Runtime/MovieSceneCoreTypes/Private/Instances/MovieSceneBoolTrackInstance.h
deleted file mode 100644
diff --git a/Engine/Source/Runtime/MovieSceneCoreTypes/Private/Instances/MovieSceneByteTrackInstance.cpp b/Engine/Source/Runtime/MovieSceneCoreTypes/Private/Instances/MovieSceneByteTrackInstance.cpp
deleted file mode 100644
diff --git a/Engine/Source/Runtime/MovieSceneCoreTypes/Private/Instances/MovieSceneByteTrackInstance.h b/Engine/Source/Runtime/MovieSceneCoreTypes/Private/Instances/MovieSceneByteTrackInstance.h
deleted file mode 100644
diff --git a/Engine/Source/Runtime/MovieSceneCoreTypes/Private/Instances/MovieSceneColorTrackInstance.cpp b/Engine/Source/Runtime/MovieSceneCoreTypes/Private/Instances/MovieSceneColorTrackInstance.cpp
deleted file mode 100644
diff --git a/Engine/Source/Runtime/MovieSceneCoreTypes/Private/Instances/MovieSceneColorTrackInstance.h b/Engine/Source/Runtime/MovieSceneCoreTypes/Private/Instances/MovieSceneColorTrackInstance.h
deleted file mode 100644
diff --git a/Engine/Source/Runtime/MovieSceneCoreTypes/Private/Instances/MovieSceneDirectorTrackInstance.cpp b/Engine/Source/Runtime/MovieSceneCoreTypes/Private/Instances/MovieSceneDirectorTrackInstance.cpp
deleted file mode 100644
diff --git a/Engine/Source/Runtime/MovieSceneCoreTypes/Private/Instances/MovieSceneDirectorTrackInstance.h b/Engine/Source/Runtime/MovieSceneCoreTypes/Private/Instances/MovieSceneDirectorTrackInstance.h
deleted file mode 100644
diff --git a/Engine/Source/Runtime/MovieSceneCoreTypes/Private/Instances/MovieSceneFloatTrackInstance.cpp b/Engine/Source/Runtime/MovieSceneCoreTypes/Private/Instances/MovieSceneFloatTrackInstance.cpp
deleted file mode 100644
diff --git a/Engine/Source/Runtime/MovieSceneCoreTypes/Private/Instances/MovieSceneFloatTrackInstance.h b/Engine/Source/Runtime/MovieSceneCoreTypes/Private/Instances/MovieSceneFloatTrackInstance.h
deleted file mode 100644
diff --git a/Engine/Source/Runtime/MovieSceneCoreTypes/Private/Instances/MovieSceneParticleTrackInstance.cpp b/Engine/Source/Runtime/MovieSceneCoreTypes/Private/Instances/MovieSceneParticleTrackInstance.cpp
deleted file mode 100644
diff --git a/Engine/Source/Runtime/MovieSceneCoreTypes/Private/Instances/MovieSceneParticleTrackInstance.h b/Engine/Source/Runtime/MovieSceneCoreTypes/Private/Instances/MovieSceneParticleTrackInstance.h
deleted file mode 100644
diff --git a/Engine/Source/Runtime/MovieSceneCoreTypes/Private/Instances/MovieSceneVectorTrackInstance.cpp b/Engine/Source/Runtime/MovieSceneCoreTypes/Private/Instances/MovieSceneVectorTrackInstance.cpp
deleted file mode 100644
diff --git a/Engine/Source/Runtime/MovieSceneCoreTypes/Private/Instances/MovieSceneVectorTrackInstance.h b/Engine/Source/Runtime/MovieSceneCoreTypes/Private/Instances/MovieSceneVectorTrackInstance.h
deleted file mode 100644
diff --git a/Engine/Source/Runtime/MovieSceneCoreTypes/Private/Instances/SubMovieSceneTrackInstance.cpp b/Engine/Source/Runtime/MovieSceneCoreTypes/Private/Instances/SubMovieSceneTrackInstance.cpp
deleted file mode 100644
diff --git a/Engine/Source/Runtime/MovieSceneCoreTypes/Private/Instances/SubMovieSceneTrackInstance.h b/Engine/Source/Runtime/MovieSceneCoreTypes/Private/Instances/SubMovieSceneTrackInstance.h
deleted file mode 100644
diff --git a/Engine/Source/Runtime/MovieSceneCoreTypes/Private/MovieScene3DTransformSection.cpp b/Engine/Source/Runtime/MovieSceneCoreTypes/Private/MovieScene3DTransformSection.cpp
deleted file mode 100644
diff --git a/Engine/Source/Runtime/MovieSceneCoreTypes/Private/MovieScene3DTransformTrack.cpp b/Engine/Source/Runtime/MovieSceneCoreTypes/Private/MovieScene3DTransformTrack.cpp
deleted file mode 100644
diff --git a/Engine/Source/Runtime/MovieSceneCoreTypes/Private/MovieSceneAnimationSection.cpp b/Engine/Source/Runtime/MovieSceneCoreTypes/Private/MovieSceneAnimationSection.cpp
deleted file mode 100644
diff --git a/Engine/Source/Runtime/MovieSceneCoreTypes/Private/MovieSceneAnimationTrack.cpp b/Engine/Source/Runtime/MovieSceneCoreTypes/Private/MovieSceneAnimationTrack.cpp
deleted file mode 100644
diff --git a/Engine/Source/Runtime/MovieSceneCoreTypes/Private/MovieSceneAudioSection.cpp b/Engine/Source/Runtime/MovieSceneCoreTypes/Private/MovieSceneAudioSection.cpp
deleted file mode 100644
diff --git a/Engine/Source/Runtime/MovieSceneCoreTypes/Private/MovieSceneAudioTrack.cpp b/Engine/Source/Runtime/MovieSceneCoreTypes/Private/MovieSceneAudioTrack.cpp
deleted file mode 100644
diff --git a/Engine/Source/Runtime/MovieSceneCoreTypes/Private/MovieSceneBoolSection.cpp b/Engine/Source/Runtime/MovieSceneCoreTypes/Private/MovieSceneBoolSection.cpp
deleted file mode 100644
diff --git a/Engine/Source/Runtime/MovieSceneCoreTypes/Private/MovieSceneBoolTrack.cpp b/Engine/Source/Runtime/MovieSceneCoreTypes/Private/MovieSceneBoolTrack.cpp
deleted file mode 100644
diff --git a/Engine/Source/Runtime/MovieSceneCoreTypes/Private/MovieSceneByteSection.cpp b/Engine/Source/Runtime/MovieSceneCoreTypes/Private/MovieSceneByteSection.cpp
deleted file mode 100644
diff --git a/Engine/Source/Runtime/MovieSceneCoreTypes/Private/MovieSceneByteTrack.cpp b/Engine/Source/Runtime/MovieSceneCoreTypes/Private/MovieSceneByteTrack.cpp
deleted file mode 100644
diff --git a/Engine/Source/Runtime/MovieSceneCoreTypes/Private/MovieSceneColorSection.cpp b/Engine/Source/Runtime/MovieSceneCoreTypes/Private/MovieSceneColorSection.cpp
deleted file mode 100644
diff --git a/Engine/Source/Runtime/MovieSceneCoreTypes/Private/MovieSceneColorTrack.cpp b/Engine/Source/Runtime/MovieSceneCoreTypes/Private/MovieSceneColorTrack.cpp
deleted file mode 100644
diff --git a/Engine/Source/Runtime/MovieSceneCoreTypes/Private/MovieSceneCoreTypesModule.cpp b/Engine/Source/Runtime/MovieSceneCoreTypes/Private/MovieSceneCoreTypesModule.cpp
deleted file mode 100644
diff --git a/Engine/Source/Runtime/MovieSceneCoreTypes/Private/MovieSceneCoreTypesPCH.h b/Engine/Source/Runtime/MovieSceneCoreTypes/Private/MovieSceneCoreTypesPCH.h
deleted file mode 100644
diff --git a/Engine/Source/Runtime/MovieSceneCoreTypes/Private/MovieSceneDirectorTrack.cpp b/Engine/Source/Runtime/MovieSceneCoreTypes/Private/MovieSceneDirectorTrack.cpp
deleted file mode 100644
diff --git a/Engine/Source/Runtime/MovieSceneCoreTypes/Private/MovieSceneFloatSection.cpp b/Engine/Source/Runtime/MovieSceneCoreTypes/Private/MovieSceneFloatSection.cpp
deleted file mode 100644
diff --git a/Engine/Source/Runtime/MovieSceneCoreTypes/Private/MovieSceneFloatTrack.cpp b/Engine/Source/Runtime/MovieSceneCoreTypes/Private/MovieSceneFloatTrack.cpp
deleted file mode 100644
diff --git a/Engine/Source/Runtime/MovieSceneCoreTypes/Private/MovieSceneParticleSection.cpp b/Engine/Source/Runtime/MovieSceneCoreTypes/Private/MovieSceneParticleSection.cpp
deleted file mode 100644
diff --git a/Engine/Source/Runtime/MovieSceneCoreTypes/Private/MovieSceneParticleTrack.cpp b/Engine/Source/Runtime/MovieSceneCoreTypes/Private/MovieSceneParticleTrack.cpp
deleted file mode 100644
diff --git a/Engine/Source/Runtime/MovieSceneCoreTypes/Private/MovieSceneShotSection.cpp b/Engine/Source/Runtime/MovieSceneCoreTypes/Private/MovieSceneShotSection.cpp
deleted file mode 100644
diff --git a/Engine/Source/Runtime/MovieSceneCoreTypes/Private/MovieSceneVectorSection.cpp b/Engine/Source/Runtime/MovieSceneCoreTypes/Private/MovieSceneVectorSection.cpp
deleted file mode 100644
diff --git a/Engine/Source/Runtime/MovieSceneCoreTypes/Private/MovieSceneVectorTrack.cpp b/Engine/Source/Runtime/MovieSceneCoreTypes/Private/MovieSceneVectorTrack.cpp
deleted file mode 100644
diff --git a/Engine/Source/Runtime/MovieSceneCoreTypes/Private/SubMovieSceneSection.cpp b/Engine/Source/Runtime/MovieSceneCoreTypes/Private/SubMovieSceneSection.cpp
deleted file mode 100644
diff --git a/Engine/Source/Runtime/MovieSceneCoreTypes/Private/SubMovieSceneTrack.cpp b/Engine/Source/Runtime/MovieSceneCoreTypes/Private/SubMovieSceneTrack.cpp
deleted file mode 100644
diff --git a/Engine/Source/Runtime/MovieSceneCoreTypes/Public/IMovieSceneCoreTypes.h b/Engine/Source/Runtime/MovieSceneCoreTypes/Public/IMovieSceneCoreTypes.h
deleted file mode 100644
diff --git a/Engine/Source/Runtime/MovieSceneTracks/MovieSceneTracks.Build.cs b/Engine/Source/Runtime/MovieSceneTracks/MovieSceneTracks.Build.cs
new file mode 100644
--- /dev/null
+++ b/Engine/Source/Runtime/MovieSceneTracks/MovieSceneTracks.Build.cs
@@ -0,0 +1,32 @@
+// Copyright 1998-2015 Epic Games, Inc. All Rights Reserved.
+
+using UnrealBuildTool;
+
+public class MovieSceneTracks : ModuleRules
+{
+	public MovieSceneTracks(TargetInfo Target)
+	{
+		PrivateIncludePaths.AddRange(new string[]
+        { 
+            "Runtime/MovieSceneTracks/Private",
+			"Runtime/MovieSceneTracks/Private/Sections",
+			"Runtime/MovieSceneTracks/Private/Tracks",
+            "Runtime/MovieSceneTracks/Private/TrackInstances",
+        });
+
+		PublicDependencyModuleNames.AddRange(
+			new string[] {
+				"Core",
+				"CoreUObject",
+				"Engine"
+			}
+		);
+
+		PrivateDependencyModuleNames.AddRange(
+			new string[] {
+				"MovieSceneCore",
+				"SlateCore"
+			}
+		);
+	}
+}
diff --git a/Engine/Source/Runtime/MovieSceneTracks/Private/MovieSceneTracksModule.cpp b/Engine/Source/Runtime/MovieSceneTracks/Private/MovieSceneTracksModule.cpp
new file mode 100644
--- /dev/null
+++ b/Engine/Source/Runtime/MovieSceneTracks/Private/MovieSceneTracksModule.cpp
@@ -0,0 +1,15 @@
+// Copyright 1998-2015 Epic Games, Inc. All Rights Reserved.
+
+#include "MovieSceneTracksPrivatePCH.h"
+#include "ModuleManager.h"
+
+
+/**
+ * Implements the MovieSceneTracks module.
+ */
+class FMovieSceneTracksModule
+	: public IMovieSceneTracksModule
+{ };
+
+
+IMPLEMENT_MODULE(FMovieSceneTracksModule, MovieSceneTracks);
diff --git a/Engine/Source/Runtime/MovieSceneTracks/Private/MovieSceneTracksPrivatePCH.h b/Engine/Source/Runtime/MovieSceneTracks/Private/MovieSceneTracksPrivatePCH.h
new file mode 100644
--- /dev/null
+++ b/Engine/Source/Runtime/MovieSceneTracks/Private/MovieSceneTracksPrivatePCH.h
@@ -0,0 +1,26 @@
+// Copyright 1998-2015 Epic Games, Inc. All Rights Reserved.
+
+#ifndef _MOVIE_SCENE_CORE_TYPES_PRIVATE_
+#define _MOVIE_SCENE_CORE_TYPES_PRIVATE_
+
+/* Private dependencies
+ *****************************************************************************/
+
+#include "Core.h"
+#include "Engine.h"
+#include "MovieScene.h"
+#include "MovieSceneCommonHelpers.h"
+#include "MovieScenePropertyTrack.h"
+#include "MovieSceneSection.h"
+#include "MovieSceneTrack.h"
+#include "IMovieSceneTrackInstance.h"
+#include "IMovieSceneTracksModule.h"
+
+
+/* Private includes
+ *****************************************************************************/
+
+#include "MovieSceneTracksClasses.h"
+
+
+#endif
diff --git a/Engine/Source/Runtime/MovieSceneTracks/Private/Sections/MovieScene3DTransformSection.cpp b/Engine/Source/Runtime/MovieSceneTracks/Private/Sections/MovieScene3DTransformSection.cpp
new file mode 100644
--- /dev/null
+++ b/Engine/Source/Runtime/MovieSceneTracks/Private/Sections/MovieScene3DTransformSection.cpp
@@ -0,0 +1,234 @@
+// Copyright 1998-2015 Epic Games, Inc. All Rights Reserved.
+
+#include "MovieSceneTracksPrivatePCH.h"
+#include "MovieScene3DTransformSection.h"
+#include "MovieScene3DTransformTrack.h"
+
+UMovieScene3DTransformSection::UMovieScene3DTransformSection( const FObjectInitializer& ObjectInitializer )
+	: Super( ObjectInitializer )
+{
+
+}
+
+void UMovieScene3DTransformSection::MoveSection( float DeltaTime, TSet<FKeyHandle>& KeyHandles )
+{
+	Super::MoveSection( DeltaTime, KeyHandles );
+
+	// Move all the curves in this section
+	for( int32 Axis = 0; Axis < 3; ++Axis )
+	{
+		Translation[Axis].ShiftCurve( DeltaTime, KeyHandles );
+		Rotation[Axis].ShiftCurve( DeltaTime, KeyHandles );
+		Scale[Axis].ShiftCurve( DeltaTime, KeyHandles );
+	}
+}
+
+void UMovieScene3DTransformSection::DilateSection( float DilationFactor, float Origin, TSet<FKeyHandle>& KeyHandles )
+{
+	Super::DilateSection(DilationFactor, Origin, KeyHandles);
+
+	for( int32 Axis = 0; Axis < 3; ++Axis )
+	{
+		Translation[Axis].ScaleCurve( Origin, DilationFactor, KeyHandles );
+		Rotation[Axis].ScaleCurve( Origin, DilationFactor, KeyHandles );
+		Scale[Axis].ScaleCurve( Origin, DilationFactor, KeyHandles );
+	}
+}
+
+void UMovieScene3DTransformSection::GetKeyHandles(TSet<FKeyHandle>& KeyHandles) const
+{
+	for (int32 Axis = 0; Axis < 3; ++Axis)
+	{
+		for (auto It(Translation[Axis].GetKeyHandleIterator()); It; ++It)
+		{
+			float Time = Translation[Axis].GetKeyTime(It.Key());
+			if (IsTimeWithinSection(Time))
+			{
+				KeyHandles.Add(It.Key());
+			}
+		}
+		for (auto It(Rotation[Axis].GetKeyHandleIterator()); It; ++It)
+		{
+			float Time = Rotation[Axis].GetKeyTime(It.Key());
+			if (IsTimeWithinSection(Time))
+			{
+				KeyHandles.Add(It.Key());
+			}
+		}
+		for (auto It(Scale[Axis].GetKeyHandleIterator()); It; ++It)
+		{
+			float Time = Scale[Axis].GetKeyTime(It.Key());
+			if (IsTimeWithinSection(Time))
+			{
+				KeyHandles.Add(It.Key());
+			}
+		}
+	}
+}	
+
+bool UMovieScene3DTransformSection::EvalTranslation( float Time, FVector& OutTranslation ) const
+{
+	bool bHasTranslationKeys = Translation[0].GetNumKeys() > 0 || Translation[1].GetNumKeys() > 0 || Translation[2].GetNumKeys() > 0;
+
+	if( bHasTranslationKeys )
+	{
+		// Evaluate each component of translation if there are any keys
+		OutTranslation.X = Translation[0].Eval( Time );
+		OutTranslation.Y = Translation[1].Eval( Time );
+		OutTranslation.Z = Translation[2].Eval( Time );
+	}
+
+	return bHasTranslationKeys;
+}
+
+bool UMovieScene3DTransformSection::EvalRotation( float Time, FRotator& OutRotation ) const
+{
+	bool bHasRotationKeys = Rotation[0].GetNumKeys() > 0 || Rotation[1].GetNumKeys() > 0 || Rotation[2].GetNumKeys() > 0;
+
+	if( bHasRotationKeys )
+	{
+		// Evaluate each component of rotation if there are any keys
+		OutRotation.Roll = Rotation[0].Eval( Time );
+		OutRotation.Pitch = Rotation[1].Eval( Time );
+		OutRotation.Yaw = Rotation[2].Eval( Time );
+	}
+
+	return bHasRotationKeys;
+
+}
+
+bool UMovieScene3DTransformSection::EvalScale( float Time, FVector& OutScale ) const
+{
+	bool bHasScaleKeys = Scale[0].GetNumKeys() > 0 || Scale[1].GetNumKeys() > 0 || Scale[2].GetNumKeys() > 0;
+	
+	if( bHasScaleKeys )
+	{
+		// Evaluate each component of scale if there are any keys
+		OutScale.X = Scale[0].Eval( Time );
+		OutScale.Y = Scale[1].Eval( Time );
+		OutScale.Z = Scale[2].Eval( Time );
+	}
+
+	return bHasScaleKeys;
+
+}
+
+/**
+ * Chooses an appropriate curve from an axis and a set of curves
+ */
+static FRichCurve* ChooseCurve( EAxis::Type Axis, FRichCurve* Curves )
+{
+	switch( Axis )
+	{
+	case EAxis::X:
+		return &Curves[0];
+		break;
+	case EAxis::Y:
+		return &Curves[1];
+		break;
+	case EAxis::Z:
+		return &Curves[2];
+		break;
+	default:
+		check( false );
+		return NULL;
+		break;
+	}
+}
+
+FRichCurve& UMovieScene3DTransformSection::GetTranslationCurve( EAxis::Type Axis ) 
+{
+	return *ChooseCurve( Axis, Translation );
+}
+
+FRichCurve& UMovieScene3DTransformSection::GetRotationCurve( EAxis::Type Axis )
+{
+	return *ChooseCurve( Axis, Rotation );
+}
+
+FRichCurve& UMovieScene3DTransformSection::GetScaleCurve( EAxis::Type Axis )
+{
+	return *ChooseCurve( Axis, Scale );
+}
+
+
+
+void UMovieScene3DTransformSection::AddTranslationKeys( const FTransformKey& TransformKey )
+{
+	const float Time = TransformKey.GetKeyTime();
+
+	const bool bUnwindRotation = false; // Not needed for translation
+
+	// Key each component.  Note: we always key each component if no key exists
+	if( Translation[0].GetNumKeys() == 0 || TransformKey.ShouldKeyTranslation( EAxis::X ) )
+	{
+		AddKeyToCurve( Translation[0], Time, TransformKey.GetTranslationValue().X );
+	}
+
+	if( Translation[1].GetNumKeys() == 0 || TransformKey.ShouldKeyTranslation( EAxis::Y ) )
+	{ 
+		AddKeyToCurve( Translation[1], Time, TransformKey.GetTranslationValue().Y );
+	}
+
+	if( Translation[2].GetNumKeys() == 0 || TransformKey.ShouldKeyTranslation( EAxis::Z ) )
+	{ 
+		AddKeyToCurve( Translation[2], Time, TransformKey.GetTranslationValue().Z );
+	}
+}
+
+void UMovieScene3DTransformSection::AddRotationKeys( const FTransformKey& TransformKey, const bool bUnwindRotation )
+{
+	const float Time = TransformKey.GetKeyTime();
+
+	// Key each component.  Note: we always key each component if no key exists
+	if( Rotation[0].GetNumKeys() == 0 || TransformKey.ShouldKeyRotation( EAxis::X ) )
+	{
+		AddKeyToCurve( Rotation[0], Time, TransformKey.GetRotationValue().Roll );
+	}
+
+	if( Rotation[1].GetNumKeys() == 0 || TransformKey.ShouldKeyRotation( EAxis::Y ) )
+	{ 
+		AddKeyToCurve( Rotation[1], Time, TransformKey.GetRotationValue().Pitch );
+	}
+
+	if( Rotation[2].GetNumKeys() == 0 || TransformKey.ShouldKeyRotation( EAxis::Z ) )
+	{ 
+		AddKeyToCurve( Rotation[2], Time, TransformKey.GetRotationValue().Yaw );
+	}
+}
+
+void UMovieScene3DTransformSection::AddScaleKeys( const FTransformKey& TransformKey )
+{
+	const float Time = TransformKey.GetKeyTime();
+
+	const bool bUnwindRotation = false; // Not needed for scale
+
+	// Key each component.  Note: we always key each component if no key exists
+	if( Scale[0].GetNumKeys() == 0 || TransformKey.ShouldKeyScale( EAxis::X ) )
+	{
+		AddKeyToCurve( Scale[0], Time, TransformKey.GetScaleValue().X );
+	}
+
+	if( Scale[1].GetNumKeys() == 0 || TransformKey.ShouldKeyScale( EAxis::Y ) )
+	{ 
+		AddKeyToCurve( Scale[1], Time, TransformKey.GetScaleValue().Y );
+	}
+
+	if( Scale[2].GetNumKeys() == 0 || TransformKey.ShouldKeyScale( EAxis::Z ) )
+	{ 
+		AddKeyToCurve( Scale[2], Time, TransformKey.GetScaleValue().Z );
+	}
+}
+
+bool UMovieScene3DTransformSection::NewKeyIsNewData(const FTransformKey& TransformKey) const
+{
+	bool bHasEmptyKeys = false;
+	for (int32 i = 0; i < 3; ++i)
+	{
+		bHasEmptyKeys = bHasEmptyKeys ||
+			Translation[i].GetNumKeys() == 0 ||
+			Rotation[i].GetNumKeys() == 0 ||
+			Scale[i].GetNumKeys() == 0;
+	}
+	return bHasEmptyKeys || TransformKey.ShouldKeyAny();
+}
diff --git a/Engine/Source/Runtime/MovieSceneTracks/Private/Sections/MovieSceneAnimationSection.cpp b/Engine/Source/Runtime/MovieSceneTracks/Private/Sections/MovieSceneAnimationSection.cpp
new file mode 100644
--- /dev/null
+++ b/Engine/Source/Runtime/MovieSceneTracks/Private/Sections/MovieSceneAnimationSection.cpp
@@ -0,0 +1,45 @@
+// Copyright 1998-2015 Epic Games, Inc. All Rights Reserved.
+
+#include "MovieSceneTracksPrivatePCH.h"
+#include "MovieSceneAnimationSection.h"
+
+UMovieSceneAnimationSection::UMovieSceneAnimationSection( const FObjectInitializer& ObjectInitializer )
+	: Super( ObjectInitializer )
+{
+	AnimSequence = NULL;
+	AnimationStartTime = 0.f;
+	AnimationDilationFactor = 1.f;
+}
+
+void UMovieSceneAnimationSection::MoveSection( float DeltaTime, TSet<FKeyHandle>& KeyHandles )
+{
+	Super::MoveSection(DeltaTime, KeyHandles);
+
+	AnimationStartTime += DeltaTime;
+}
+
+void UMovieSceneAnimationSection::DilateSection( float DilationFactor, float Origin, TSet<FKeyHandle>& KeyHandles )
+{
+	AnimationStartTime = (AnimationStartTime - Origin) * DilationFactor + Origin;
+	AnimationDilationFactor *= DilationFactor;
+
+	Super::DilateSection(DilationFactor, Origin, KeyHandles);
+}
+
+void UMovieSceneAnimationSection::GetKeyHandles(TSet<FKeyHandle>& KeyHandles) const
+{
+}
+
+void UMovieSceneAnimationSection::GetSnapTimes(TArray<float>& OutSnapTimes, bool bGetSectionBorders) const
+{
+	Super::GetSnapTimes(OutSnapTimes, bGetSectionBorders);
+	float CurrentTime = GetAnimationStartTime();
+	while (CurrentTime <= GetEndTime())
+	{
+		if (CurrentTime >= GetStartTime())
+		{
+			OutSnapTimes.Add(CurrentTime);
+		}
+		CurrentTime += GetAnimationDuration();
+	}
+}
diff --git a/Engine/Source/Runtime/MovieSceneTracks/Private/Sections/MovieSceneAudioSection.cpp b/Engine/Source/Runtime/MovieSceneTracks/Private/Sections/MovieSceneAudioSection.cpp
new file mode 100644
--- /dev/null
+++ b/Engine/Source/Runtime/MovieSceneTracks/Private/Sections/MovieSceneAudioSection.cpp
@@ -0,0 +1,61 @@
+// Copyright 1998-2015 Epic Games, Inc. All Rights Reserved.
+
+#include "MovieSceneTracksPrivatePCH.h"
+#include "MovieSceneAudioSection.h"
+#include "SoundDefinitions.h"
+
+UMovieSceneAudioSection::UMovieSceneAudioSection( const FObjectInitializer& ObjectInitializer )
+	: Super( ObjectInitializer )
+{
+
+	Sound = NULL;
+	AudioStartTime = 0.f;
+	AudioDilationFactor = 1.f;
+}
+
+TRange<float> UMovieSceneAudioSection::GetAudioRange() const
+{
+	return !Sound ? TRange<float>::Empty() :
+		TRange<float>(FMath::Max(GetStartTime(), AudioStartTime),
+		FMath::Min(AudioStartTime + Sound->GetDuration() * AudioDilationFactor, GetEndTime()));
+}
+
+TRange<float> UMovieSceneAudioSection::GetAudioTrueRange() const
+{
+	return !Sound ? TRange<float>::Empty() : TRange<float>(AudioStartTime, AudioStartTime + Sound->GetDuration() * AudioDilationFactor);
+}
+
+void UMovieSceneAudioSection::MoveSection( float DeltaTime, TSet<FKeyHandle>& KeyHandles )
+{
+	Super::MoveSection(DeltaTime, KeyHandles);
+
+	AudioStartTime += DeltaTime;
+}
+
+void UMovieSceneAudioSection::DilateSection( float DilationFactor, float Origin, TSet<FKeyHandle>& KeyHandles )
+{
+	AudioStartTime = (AudioStartTime - Origin) * DilationFactor + Origin;
+	AudioDilationFactor *= DilationFactor;
+
+	Super::DilateSection(DilationFactor, Origin, KeyHandles);
+}
+
+void UMovieSceneAudioSection::GetKeyHandles(TSet<FKeyHandle>& KeyHandles) const
+{
+}
+
+void UMovieSceneAudioSection::GetSnapTimes(TArray<float>& OutSnapTimes, bool bGetSectionBorders) const
+{
+	Super::GetSnapTimes(OutSnapTimes, bGetSectionBorders);
+	if (AudioStartTime > GetStartTime())
+	{
+		OutSnapTimes.Add(AudioStartTime);
+	}
+	float AudioEndTime = GetAudioTrueRange().GetUpperBoundValue();
+	if (AudioEndTime < GetEndTime())
+	{
+		OutSnapTimes.Add(AudioEndTime);
+	}
+	// @todo Sequencer handle snapping for time dilation
+	// @todo Don't add redundant times (can't use AddUnique due to floating point equality issues)
+}
diff --git a/Engine/Source/Runtime/MovieSceneTracks/Private/Sections/MovieSceneBoolSection.cpp b/Engine/Source/Runtime/MovieSceneTracks/Private/Sections/MovieSceneBoolSection.cpp
new file mode 100644
--- /dev/null
+++ b/Engine/Source/Runtime/MovieSceneTracks/Private/Sections/MovieSceneBoolSection.cpp
@@ -0,0 +1,55 @@
+// Copyright 1998-2015 Epic Games, Inc. All Rights Reserved.
+
+#include "MovieSceneTracksPrivatePCH.h"
+#include "MovieSceneBoolSection.h"
+
+
+UMovieSceneBoolSection::UMovieSceneBoolSection( const FObjectInitializer& ObjectInitializer )
+	: Super( ObjectInitializer )
+{
+}
+
+bool UMovieSceneBoolSection::Eval( float Position ) const
+{
+	return !!BoolCurve.Evaluate(Position);
+}
+
+void UMovieSceneBoolSection::MoveSection( float DeltaPosition, TSet<FKeyHandle>& KeyHandles )
+{
+	Super::MoveSection( DeltaPosition, KeyHandles );
+
+	BoolCurve.ShiftCurve(DeltaPosition, KeyHandles);
+}
+
+void UMovieSceneBoolSection::DilateSection( float DilationFactor, float Origin, TSet<FKeyHandle>& KeyHandles )
+{
+	float StartTime = GetStartTime();
+	float EndTime = GetEndTime();
+
+	Super::DilateSection(DilationFactor, Origin, KeyHandles);
+	
+	BoolCurve.ScaleCurve(Origin, DilationFactor, KeyHandles);
+}
+
+void UMovieSceneBoolSection::GetKeyHandles(TSet<FKeyHandle>& KeyHandles) const
+{
+	for (auto It(BoolCurve.GetKeyHandleIterator()); It; ++It)
+	{
+		float Time = BoolCurve.GetKeyTime(It.Key());
+		if (IsTimeWithinSection(Time))
+		{
+			KeyHandles.Add(It.Key());
+		}
+	}
+}
+
+void UMovieSceneBoolSection::AddKey( float Time, bool Value )
+{
+	Modify();
+	BoolCurve.UpdateOrAddKey(Time, Value ? 1 : 0);
+}
+
+bool UMovieSceneBoolSection::NewKeyIsNewData(float Time, bool Value) const
+{
+	return BoolCurve.GetNumKeys() == 0 || Eval(Time) != Value;
+}
diff --git a/Engine/Source/Runtime/MovieSceneTracks/Private/Sections/MovieSceneByteSection.cpp b/Engine/Source/Runtime/MovieSceneTracks/Private/Sections/MovieSceneByteSection.cpp
new file mode 100644
--- /dev/null
+++ b/Engine/Source/Runtime/MovieSceneTracks/Private/Sections/MovieSceneByteSection.cpp
@@ -0,0 +1,50 @@
+// Copyright 1998-2015 Epic Games, Inc. All Rights Reserved.
+
+#include "MovieSceneTracksPrivatePCH.h"
+#include "MovieSceneByteSection.h"
+
+
+UMovieSceneByteSection::UMovieSceneByteSection( const FObjectInitializer& ObjectInitializer )
+	: Super( ObjectInitializer )
+{
+}
+
+uint8 UMovieSceneByteSection::Eval( float Position ) const
+{
+	return !!ByteCurve.Evaluate(Position);
+}
+
+void UMovieSceneByteSection::MoveSection( float DeltaPosition, TSet<FKeyHandle>& KeyHandles )
+{
+	Super::MoveSection( DeltaPosition, KeyHandles );
+
+	ByteCurve.ShiftCurve(DeltaPosition, KeyHandles);
+}
+
+void UMovieSceneByteSection::DilateSection( float DilationFactor, float Origin, TSet<FKeyHandle>& KeyHandles )
+{
+	Super::DilateSection(DilationFactor, Origin, KeyHandles);
+	
+	ByteCurve.ScaleCurve(Origin, DilationFactor, KeyHandles);
+}
+
+void UMovieSceneByteSection::GetKeyHandles(TSet<FKeyHandle>& KeyHandles) const
+{
+	for (auto It(ByteCurve.GetKeyHandleIterator()); It; ++It)
+	{
+		float Time = ByteCurve.GetKeyTime(It.Key());
+		if (IsTimeWithinSection(Time))
+			KeyHandles.Add(It.Key());
+	}
+}
+
+void UMovieSceneByteSection::AddKey( float Time, uint8 Value )
+{
+	Modify();
+	ByteCurve.UpdateOrAddKey(Time, Value ? 1 : 0);
+}
+
+bool UMovieSceneByteSection::NewKeyIsNewData(float Time, uint8 Value) const
+{
+	return ByteCurve.GetNumKeys() == 0 || Eval(Time) != Value;
+}
diff --git a/Engine/Source/Runtime/MovieSceneTracks/Private/Sections/MovieSceneColorSection.cpp b/Engine/Source/Runtime/MovieSceneTracks/Private/Sections/MovieSceneColorSection.cpp
new file mode 100644
--- /dev/null
+++ b/Engine/Source/Runtime/MovieSceneTracks/Private/Sections/MovieSceneColorSection.cpp
@@ -0,0 +1,126 @@
+// Copyright 1998-2015 Epic Games, Inc. All Rights Reserved.
+
+#include "MovieSceneTracksPrivatePCH.h"
+#include "MovieSceneColorSection.h"
+#include "MovieSceneColorTrack.h"
+
+
+UMovieSceneColorSection::UMovieSceneColorSection( const FObjectInitializer& ObjectInitializer )
+	: Super( ObjectInitializer )
+{
+}
+
+void UMovieSceneColorSection::MoveSection( float DeltaTime, TSet<FKeyHandle>& KeyHandles )
+{
+	Super::MoveSection( DeltaTime, KeyHandles );
+
+	// Move all the curves in this section
+	RedCurve.ShiftCurve(DeltaTime, KeyHandles);
+	GreenCurve.ShiftCurve(DeltaTime, KeyHandles);
+	BlueCurve.ShiftCurve(DeltaTime, KeyHandles);
+	AlphaCurve.ShiftCurve(DeltaTime, KeyHandles);
+}
+
+void UMovieSceneColorSection::DilateSection( float DilationFactor, float Origin, TSet<FKeyHandle>& KeyHandles )
+{
+	Super::DilateSection(DilationFactor, Origin, KeyHandles);
+
+	RedCurve.ScaleCurve(Origin, DilationFactor, KeyHandles);
+	GreenCurve.ScaleCurve(Origin, DilationFactor, KeyHandles);
+	BlueCurve.ScaleCurve(Origin, DilationFactor, KeyHandles);
+	AlphaCurve.ScaleCurve(Origin, DilationFactor, KeyHandles);
+}
+
+void UMovieSceneColorSection::GetKeyHandles(TSet<FKeyHandle>& KeyHandles) const
+{
+	for (auto It(RedCurve.GetKeyHandleIterator()); It; ++It)
+	{
+		float Time = RedCurve.GetKeyTime(It.Key());
+		if (IsTimeWithinSection(Time))
+		{
+			KeyHandles.Add(It.Key());
+		}
+	}
+	for (auto It(GreenCurve.GetKeyHandleIterator()); It; ++It)
+	{
+		float Time = GreenCurve.GetKeyTime(It.Key());
+		if (IsTimeWithinSection(Time))
+		{
+			KeyHandles.Add(It.Key());
+		}
+	}
+	for (auto It(BlueCurve.GetKeyHandleIterator()); It; ++It)
+	{
+		float Time = BlueCurve.GetKeyTime(It.Key());
+		if (IsTimeWithinSection(Time))
+		{
+			KeyHandles.Add(It.Key());
+		}
+	}
+	for (auto It(AlphaCurve.GetKeyHandleIterator()); It; ++It)
+	{
+		float Time = AlphaCurve.GetKeyTime(It.Key());
+		if (IsTimeWithinSection(Time))
+			KeyHandles.Add(It.Key());
+	}
+}
+
+FLinearColor UMovieSceneColorSection::Eval( float Position, const FLinearColor& DefaultColor ) const
+{
+	return FLinearColor(RedCurve.Eval(Position, DefaultColor.R),
+						GreenCurve.Eval(Position, DefaultColor.G),
+						BlueCurve.Eval(Position, DefaultColor.B),
+						AlphaCurve.Eval(Position, DefaultColor.A));
+}
+
+void UMovieSceneColorSection::AddKey( float Time, const FColorKey& Key )
+{
+	Modify();
+
+	if( Key.CurveName == NAME_None )
+	{
+		AddKeyToCurve(RedCurve, Time, Key.Value.R);
+		AddKeyToCurve(GreenCurve, Time, Key.Value.G);
+		AddKeyToCurve(BlueCurve, Time, Key.Value.B);
+		AddKeyToCurve(AlphaCurve, Time, Key.Value.A);
+	}
+	else
+	{
+		AddKeyToNamedCurve( Time, Key );
+	}
+}
+
+bool UMovieSceneColorSection::NewKeyIsNewData(float Time, FLinearColor Value) const
+{
+	return RedCurve.GetNumKeys() == 0 ||
+		GreenCurve.GetNumKeys() == 0 ||
+		BlueCurve.GetNumKeys() == 0 ||
+		AlphaCurve.GetNumKeys() == 0 ||
+		!Eval(Time,Value).Equals(Value);
+}
+
+void UMovieSceneColorSection::AddKeyToNamedCurve(float Time, const FColorKey& Key)
+{
+	static FName R("R");
+	static FName G("G");
+	static FName B("B");
+	static FName A("A");
+
+	FName CurveName = Key.CurveName;
+	if (CurveName == R)
+	{
+		AddKeyToCurve(RedCurve, Time, Key.Value.R);
+	}
+	else if (CurveName == G)
+	{
+		AddKeyToCurve(GreenCurve, Time, Key.Value.G);
+	}
+	else if (CurveName == B)
+	{
+		AddKeyToCurve(BlueCurve, Time, Key.Value.B);
+	}
+	else if (CurveName == A)
+	{
+		AddKeyToCurve(AlphaCurve, Time, Key.Value.A);
+	}
+}
diff --git a/Engine/Source/Runtime/MovieSceneTracks/Private/Sections/MovieSceneFloatSection.cpp b/Engine/Source/Runtime/MovieSceneTracks/Private/Sections/MovieSceneFloatSection.cpp
new file mode 100644
--- /dev/null
+++ b/Engine/Source/Runtime/MovieSceneTracks/Private/Sections/MovieSceneFloatSection.cpp
@@ -0,0 +1,53 @@
+// Copyright 1998-2015 Epic Games, Inc. All Rights Reserved.
+
+#include "MovieSceneTracksPrivatePCH.h"
+#include "MovieSceneFloatSection.h"
+
+
+UMovieSceneFloatSection::UMovieSceneFloatSection( const FObjectInitializer& ObjectInitializer )
+	: Super( ObjectInitializer )
+{
+}
+
+float UMovieSceneFloatSection::Eval( float Position ) const
+{
+	return FloatCurve.Eval( Position );
+}
+
+void UMovieSceneFloatSection::MoveSection( float DeltaPosition, TSet<FKeyHandle>& KeyHandles )
+{
+	Super::MoveSection( DeltaPosition, KeyHandles );
+
+	// Move the curve
+	FloatCurve.ShiftCurve(DeltaPosition, KeyHandles);
+}
+
+void UMovieSceneFloatSection::DilateSection( float DilationFactor, float Origin, TSet<FKeyHandle>& KeyHandles )
+{	
+	Super::DilateSection(DilationFactor, Origin, KeyHandles);
+	
+	FloatCurve.ScaleCurve(Origin, DilationFactor, KeyHandles);
+}
+
+void UMovieSceneFloatSection::GetKeyHandles(TSet<FKeyHandle>& KeyHandles) const
+{
+	for (auto It(FloatCurve.GetKeyHandleIterator()); It; ++It)
+	{
+		float Time = FloatCurve.GetKeyTime(It.Key());
+		if (IsTimeWithinSection(Time))
+		{
+			KeyHandles.Add(It.Key());
+		}
+	}
+}
+
+void UMovieSceneFloatSection::AddKey( float Time, float Value )
+{
+	Modify();
+	FloatCurve.UpdateOrAddKey(Time, Value);
+}
+
+bool UMovieSceneFloatSection::NewKeyIsNewData(float Time, float Value) const
+{
+	return FloatCurve.GetNumKeys() == 0 || !FMath::IsNearlyEqual(FloatCurve.Eval(Time), Value);
+}
diff --git a/Engine/Source/Runtime/MovieSceneTracks/Private/Sections/MovieSceneParticleSection.cpp b/Engine/Source/Runtime/MovieSceneTracks/Private/Sections/MovieSceneParticleSection.cpp
new file mode 100644
--- /dev/null
+++ b/Engine/Source/Runtime/MovieSceneTracks/Private/Sections/MovieSceneParticleSection.cpp
@@ -0,0 +1,11 @@
+// Copyright 1998-2015 Epic Games, Inc. All Rights Reserved.
+
+#include "MovieSceneTracksPrivatePCH.h"
+#include "MovieSceneParticleSection.h"
+
+UMovieSceneParticleSection::UMovieSceneParticleSection( const FObjectInitializer& ObjectInitializer )
+	: Super( ObjectInitializer )
+{
+	KeyType = EParticleKey::Toggle;
+}
+
diff --git a/Engine/Source/Runtime/MovieSceneTracks/Private/Sections/MovieSceneShotSection.cpp b/Engine/Source/Runtime/MovieSceneTracks/Private/Sections/MovieSceneShotSection.cpp
new file mode 100644
--- /dev/null
+++ b/Engine/Source/Runtime/MovieSceneTracks/Private/Sections/MovieSceneShotSection.cpp
@@ -0,0 +1,29 @@
+// Copyright 1998-2015 Epic Games, Inc. All Rights Reserved.
+
+#include "MovieSceneTracksPrivatePCH.h"
+#include "MovieSceneShotSection.h"
+
+UMovieSceneShotSection::UMovieSceneShotSection( const FObjectInitializer& ObjectInitializer )
+	: Super( ObjectInitializer )
+{
+}
+
+void UMovieSceneShotSection::SetCameraGuid(const FGuid& InGuid)
+{
+	CameraGuid = InGuid;
+}
+
+FGuid UMovieSceneShotSection::GetCameraGuid() const
+{
+	return CameraGuid;
+}
+
+void UMovieSceneShotSection::SetTitle(const FText& InTitle)
+{
+	Title = InTitle;
+}
+
+FText UMovieSceneShotSection::GetTitle() const
+{
+	return Title;
+}
diff --git a/Engine/Source/Runtime/MovieSceneTracks/Private/Sections/MovieSceneVectorSection.cpp b/Engine/Source/Runtime/MovieSceneTracks/Private/Sections/MovieSceneVectorSection.cpp
new file mode 100644
--- /dev/null
+++ b/Engine/Source/Runtime/MovieSceneTracks/Private/Sections/MovieSceneVectorSection.cpp
@@ -0,0 +1,120 @@
+// Copyright 1998-2015 Epic Games, Inc. All Rights Reserved.
+
+#include "MovieSceneTracksPrivatePCH.h"
+#include "MovieSceneVectorSection.h"
+#include "MovieSceneVectorTrack.h"
+
+UMovieSceneVectorSection::UMovieSceneVectorSection( const FObjectInitializer& ObjectInitializer )
+	: Super( ObjectInitializer )
+{
+
+}
+
+FVector4 UMovieSceneVectorSection::Eval( float Position, const FVector4& DefaultVector ) const
+{
+	return FVector4(
+		Curves[0].Eval( Position, DefaultVector.X ),
+		Curves[1].Eval( Position, DefaultVector.Y ),
+		Curves[2].Eval( Position, DefaultVector.Z ),
+		Curves[3].Eval( Position, DefaultVector.W ) );
+}
+
+void UMovieSceneVectorSection::AddKey( float Time, FName CurveName, const FVector4& Value )
+{
+	Modify();
+
+	if( CurveName == NAME_None )
+	{
+		check(ChannelsUsed >= 2 && ChannelsUsed <= 4);
+
+		for (int32 i = 0; i < ChannelsUsed; ++i)
+		{
+			Curves[i].UpdateOrAddKey(Time, Value[i]);
+		}
+	}
+	else
+	{
+		AddKeyToNamedCurve( Time, CurveName, Value );
+	}
+}
+
+bool UMovieSceneVectorSection::NewKeyIsNewData(float Time, const FVector4& Value) const
+{
+	check(ChannelsUsed >= 2 && ChannelsUsed <= 4);
+
+	bool bNewData = false;
+	for (int32 i = 0; i < ChannelsUsed; ++i)
+	{
+		float OriginalData = Curves[i].Eval(Time);
+		// don't re-add keys if the data already matches
+		if (Curves[i].GetNumKeys() == 0 || !FMath::IsNearlyEqual(OriginalData, Value[i]))
+		{
+			bNewData = true;
+			break;
+		}
+	}
+	return bNewData;
+}
+
+void UMovieSceneVectorSection::AddKeyToNamedCurve(float Time, FName CurveName, const FVector4& Value)
+{
+	static FName X("X");
+	static FName Y("Y");
+	static FName Z("Z");
+	static FName W("W");
+
+	if (CurveName == X)
+	{
+		Curves[0].UpdateOrAddKey(Time, Value.X);
+	}
+	else if (CurveName == Y)
+	{
+		Curves[1].UpdateOrAddKey(Time, Value.Y);
+	}
+	else if (CurveName == Z)
+	{
+		Curves[2].UpdateOrAddKey(Time, Value.Z);
+	}
+	else if (CurveName == W)
+	{
+		Curves[3].UpdateOrAddKey(Time, Value.W);
+	}
+}
+
+void UMovieSceneVectorSection::MoveSection( float DeltaTime, TSet<FKeyHandle>& KeyHandles )
+{
+	check(ChannelsUsed >= 2 && ChannelsUsed <= 4);
+
+	Super::MoveSection( DeltaTime, KeyHandles );
+
+	for (int32 i = 0; i < ChannelsUsed; ++i)
+	{
+		Curves[i].ShiftCurve(DeltaTime, KeyHandles);
+	}
+}
+
+void UMovieSceneVectorSection::DilateSection( float DilationFactor, float Origin, TSet<FKeyHandle>& KeyHandles )
+{
+	check(ChannelsUsed >= 2 && ChannelsUsed <= 4);
+	Super::DilateSection(DilationFactor, Origin, KeyHandles);
+
+	for (int32 i = 0; i < ChannelsUsed; ++i)
+	{
+		Curves[i].ScaleCurve(Origin, DilationFactor, KeyHandles);
+	}
+}
+
+void UMovieSceneVectorSection::GetKeyHandles(TSet<FKeyHandle>& KeyHandles) const
+{
+	for (int32 i = 0; i < ChannelsUsed; ++i)
+	{
+		for (auto It(Curves[i].GetKeyHandleIterator()); It; ++It)
+		{
+			float Time = Curves[i].GetKeyTime(It.Key());
+			if (IsTimeWithinSection(Time))
+			{
+				KeyHandles.Add(It.Key());
+			}
+		}
+	}
+}
diff --git a/Engine/Source/Runtime/MovieSceneTracks/Private/Sections/SubMovieSceneSection.cpp b/Engine/Source/Runtime/MovieSceneTracks/Private/Sections/SubMovieSceneSection.cpp
new file mode 100644
--- /dev/null
+++ b/Engine/Source/Runtime/MovieSceneTracks/Private/Sections/SubMovieSceneSection.cpp
@@ -0,0 +1,10 @@
+// Copyright 1998-2015 Epic Games, Inc. All Rights Reserved.
+
+#include "MovieSceneTracksPrivatePCH.h"
+#include "SubMovieSceneSection.h"
+#include "IMovieScenePlayer.h"
+
+USubMovieSceneSection::USubMovieSceneSection( const FObjectInitializer& ObjectInitializer )
+	: Super( ObjectInitializer )
+{
+}
diff --git a/Engine/Source/Runtime/MovieSceneTracks/Private/TrackInstances/MovieScene3DTransformTrackInstance.cpp b/Engine/Source/Runtime/MovieSceneTracks/Private/TrackInstances/MovieScene3DTransformTrackInstance.cpp
new file mode 100644
--- /dev/null
+++ b/Engine/Source/Runtime/MovieSceneTracks/Private/TrackInstances/MovieScene3DTransformTrackInstance.cpp
@@ -0,0 +1,93 @@
+// Copyright 1998-2015 Epic Games, Inc. All Rights Reserved.
+
+#include "MovieSceneTracksPrivatePCH.h"
+#include "MovieScene3DTransformTrackInstance.h"
+
+
+FMovieScene3DTransformTrackInstance::FMovieScene3DTransformTrackInstance( UMovieScene3DTransformTrack& InTransformTrack )
+{
+	TransformTrack = &InTransformTrack;
+}
+
+namespace
+{
+
+USceneComponent* SceneComponentFromRuntimeObject(UObject* Object)
+{
+	AActor* Actor = Cast<AActor>(Object);
+
+	USceneComponent* SceneComponent = NULL;
+	if (Actor && Actor->GetRootComponent())
+	{
+		// If there is an actor, modify its root component
+		SceneComponent = Actor->GetRootComponent();
+	}
+	else
+	{
+		// No actor was found.  Attempt to get the object as a component in the case that we are editing them directly.
+		SceneComponent = Cast<USceneComponent>(Object);
+	}
+	return SceneComponent;
+}
+
+}
+
+void FMovieScene3DTransformTrackInstance::SaveState(const TArray<UObject*>& RuntimeObjects)
+{
+	for (int32 ObjIndex = 0; ObjIndex < RuntimeObjects.Num(); ++ObjIndex)
+	{
+		USceneComponent* SceneComponent = SceneComponentFromRuntimeObject(RuntimeObjects[ObjIndex]);
+		if (SceneComponent != NULL)
+		{
+			InitTransformMap.Add(RuntimeObjects[ObjIndex], SceneComponent->GetRelativeTransform());
+		}
+	}
+}
+
+void FMovieScene3DTransformTrackInstance::RestoreState(const TArray<UObject*>& RuntimeObjects)
+{
+	for (int32 ObjIndex = 0; ObjIndex < RuntimeObjects.Num(); ++ObjIndex)
+	{
+		USceneComponent* SceneComponent = SceneComponentFromRuntimeObject(RuntimeObjects[ObjIndex]);
+		if (SceneComponent != NULL)
+		{
+			FTransform *Transform = InitTransformMap.Find(RuntimeObjects[ObjIndex]);
+			if (Transform != NULL)
+			{
+				SceneComponent->SetRelativeTransform(*Transform);
+			}
+		}
+	}
+}
+
+void FMovieScene3DTransformTrackInstance::Update( float Position, float LastPosition, const TArray<UObject*>& RuntimeObjects, class IMovieScenePlayer& Player ) 
+{
+	FVector Translation;
+	FRotator Rotation;
+	FVector Scale;
+	bool bHasTranslationKeys = false;
+	bool bHasRotationKeys = false;
+	bool bHasScaleKeys = false;
+
+	if( TransformTrack->Eval( Position, LastPosition, Translation, Rotation, Scale, bHasTranslationKeys, bHasRotationKeys, bHasScaleKeys ) )
+	{
+		for( int32 ObjIndex = 0; ObjIndex < RuntimeObjects.Num(); ++ObjIndex )
+		{
+			USceneComponent* SceneComponent = SceneComponentFromRuntimeObject(RuntimeObjects[ObjIndex]);
+
+			if (SceneComponent != NULL)
+			{
+				// Set the relative translation and rotation.  Note they are set once instead of individually to avoid a redundant component transform update.
+				SceneComponent->SetRelativeLocationAndRotation(
+					bHasTranslationKeys ? Translation : SceneComponent->RelativeLocation,
+					bHasRotationKeys ? Rotation : SceneComponent->RelativeRotation );
+
+				if( bHasScaleKeys )
+				{
+					SceneComponent->SetRelativeScale3D( Scale );
+				}
+			}
+		}
+	}
+}
+
diff --git a/Engine/Source/Runtime/MovieSceneTracks/Private/TrackInstances/MovieSceneAnimationTrackInstance.cpp b/Engine/Source/Runtime/MovieSceneTracks/Private/TrackInstances/MovieSceneAnimationTrackInstance.cpp
new file mode 100644
--- /dev/null
+++ b/Engine/Source/Runtime/MovieSceneTracks/Private/TrackInstances/MovieSceneAnimationTrackInstance.cpp
@@ -0,0 +1,41 @@
+// Copyright 1998-2015 Epic Games, Inc. All Rights Reserved.
+
+#include "MovieSceneTracksPrivatePCH.h"
+#include "MovieSceneAnimationTrack.h"
+#include "MovieSceneAnimationTrackInstance.h"
+#include "IMovieScenePlayer.h"
+#include "Matinee/MatineeAnimInterface.h"
+
+FMovieSceneAnimationTrackInstance::FMovieSceneAnimationTrackInstance( UMovieSceneAnimationTrack& InAnimationTrack )
+{
+	AnimationTrack = &InAnimationTrack;
+}
+
+FMovieSceneAnimationTrackInstance::~FMovieSceneAnimationTrackInstance()
+{
+	// @todo Sequencer Need to find some way to call PreviewFinishAnimControl (needs the runtime objects)
+}
+
+void FMovieSceneAnimationTrackInstance::Update( float Position, float LastPosition, const TArray<UObject*>& RuntimeObjects, class IMovieScenePlayer& Player ) 
+{
+	// @todo Sequencer gameplay update has a different code path than editor update for animation
+
+	for (int32 i = 0; i < RuntimeObjects.Num(); ++i)
+	{
+		IMatineeAnimInterface* AnimInterface = Cast<IMatineeAnimInterface>(RuntimeObjects[i]);
+		if (AnimInterface)
+		{
+			UMovieSceneAnimationSection* AnimSection = Cast<UMovieSceneAnimationSection>(AnimationTrack->GetAnimSectionAtTime(Position));
+			if (AnimSection)
+			{
+				int32 ChannelIndex = 0;
+				FName SlotName = FName("AnimationSlot");
+				UAnimSequence* AnimSequence = AnimSection->GetAnimSequence();
+				float AnimPosition = FMath::Fmod(Position - AnimSection->GetAnimationStartTime(), AnimSection->GetAnimationDuration()) / AnimSection->GetAnimationDilationFactor();
+
+				AnimInterface->PreviewBeginAnimControl(NULL);
+				AnimInterface->PreviewSetAnimPosition(SlotName, ChannelIndex, AnimSequence, AnimPosition, true, false, 0.f);
+			}
+		}
+	}
+}
diff --git a/Engine/Source/Runtime/MovieSceneTracks/Private/TrackInstances/MovieSceneAudioTrackInstance.cpp b/Engine/Source/Runtime/MovieSceneTracks/Private/TrackInstances/MovieSceneAudioTrackInstance.cpp
new file mode 100644
--- /dev/null
+++ b/Engine/Source/Runtime/MovieSceneTracks/Private/TrackInstances/MovieSceneAudioTrackInstance.cpp
@@ -0,0 +1,192 @@
+// Copyright 1998-2015 Epic Games, Inc. All Rights Reserved.
+
+#include "MovieSceneTracksPrivatePCH.h"
+#include "MovieSceneAudioTrackInstance.h"
+#include "IMovieScenePlayer.h"
+#include "SoundDefinitions.h"
+#include "Runtime/Engine/Classes/Kismet/GameplayStatics.h"
+#include "Runtime/Engine/Public/AudioDecompress.h"
+#include "MovieSceneAudioTrack.h"
+
+FMovieSceneAudioTrackInstance::FMovieSceneAudioTrackInstance( UMovieSceneAudioTrack& InAudioTrack )
+{
+	AudioTrack = &InAudioTrack;
+}
+
+void FMovieSceneAudioTrackInstance::Update( float Position, float LastPosition, const TArray<UObject*>& RuntimeObjects, class IMovieScenePlayer& Player ) 
+{
+	const TArray<UMovieSceneSection*>& AudioSections = AudioTrack->GetAudioSections();
+
+	TArray<AActor*> Actors = GetRuntimeActors(RuntimeObjects);
+
+	if (Player.GetPlaybackStatus() == EMovieScenePlayerStatus::Playing)
+	{
+		if (Position > LastPosition)
+		{
+			TMap<int32, TArray<UMovieSceneAudioSection*> > AudioSectionsBySectionIndex;
+			for (int32 i = 0; i < AudioSections.Num(); ++i)
+			{
+				UMovieSceneAudioSection* AudioSection = Cast<UMovieSceneAudioSection>(AudioSections[i]);
+				int32 SectionIndex = AudioSection->GetRowIndex();
+				AudioSectionsBySectionIndex.FindOrAdd(SectionIndex).Add(AudioSection);
+			}
+
+			for (TMap<int32, TArray<UMovieSceneAudioSection*> >::TIterator It(AudioSectionsBySectionIndex); It; ++It)
+			{
+				int32 RowIndex = It.Key();
+				TArray<UMovieSceneAudioSection*>& MovieSceneAudioSections = It.Value();
+
+				for (int32 ActorIndex = 0; ActorIndex < Actors.Num(); ++ActorIndex)
+				{
+					UAudioComponent* Component = GetAudioComponent(Actors[ActorIndex], RowIndex);
+
+					bool bComponentIsPlaying = false;
+					for (int32 i = 0; i < MovieSceneAudioSections.Num(); ++i)
+					{
+						UMovieSceneAudioSection* AudioSection = MovieSceneAudioSections[i];
+						if (AudioSection->IsTimeWithinAudioRange(Position))
+						{
+							if (!AudioSection->IsTimeWithinAudioRange(LastPosition) || !Component->IsPlaying())
+							{
+								PlaySound(AudioSection, Component, Position);
+							}
+							bComponentIsPlaying = true;
+						}
+					}
+
+					if (!bComponentIsPlaying)
+					{
+						StopSound(RowIndex);
+					}
+				}
+			}
+		}
+		else
+		{
+			StopAllSounds();
+		}
+	}
+	else if (Player.GetPlaybackStatus() == EMovieScenePlayerStatus::Scrubbing)
+	{
+		// handle scrubbing
+		if (!FMath::IsNearlyEqual(Position, LastPosition))
+		{
+			for (int32 i = 0; i < AudioSections.Num(); ++i)
+			{
+				auto AudioSection = Cast<UMovieSceneAudioSection>(AudioSections[i]);
+				int32 RowIndex = AudioSection->GetRowIndex();
+				
+				for (int32 ActorIndex = 0; ActorIndex < Actors.Num(); ++ActorIndex)
+				{
+					UAudioComponent* Component = GetAudioComponent(Actors[ActorIndex], RowIndex);
+
+					if (AudioSection->IsTimeWithinAudioRange(Position) && !Component->IsPlaying())
+					{
+						PlaySound(AudioSection, Component, Position);
+						// Fade out the sound at the same volume in order to simply
+						// set a short duration on the sound, far from ideal soln
+						Component->FadeOut(AudioTrackConstants::ScrubDuration, 1.f);
+					}
+				}
+			}
+		}
+	}
+	else
+	{
+		// beginning scrubbing, stopped, recording
+		StopAllSounds();
+	}
+
+	// handle locality of non-master audio
+	if (!AudioTrack->IsAMasterTrack())
+	{
+		for (int32 RowIndex = 0; RowIndex < PlaybackAudioComponents.Num(); ++RowIndex)
+		{
+			for (int32 ActorIndex = 0; ActorIndex < Actors.Num(); ++ActorIndex)
+			{
+				UAudioComponent* Component = GetAudioComponent(Actors[ActorIndex], RowIndex);
+				if (Component->IsPlaying())
+				{
+					FAudioDevice* AudioDevice = Component->GetAudioDevice();
+					FActiveSound* ActiveSound = AudioDevice->FindActiveSound(Component);
+					ActiveSound->bLocationDefined = true;
+					ActiveSound->Transform = Actors[ActorIndex]->GetTransform();
+				}
+			}
+		}
+	}
+}
+
+void FMovieSceneAudioTrackInstance::PlaySound(UMovieSceneAudioSection* AudioSection, UAudioComponent* Component, float Time)
+{
+	float PitchMultiplier = 1.f / AudioSection->GetAudioDilationFactor();
+	Component->bAllowSpatialization = !AudioTrack->IsAMasterTrack();
+	Component->Stop();
+	Component->SetSound(AudioSection->GetSound());
+	Component->SetVolumeMultiplier(1.f);
+	Component->SetPitchMultiplier(PitchMultiplier);
+	Component->bIsUISound = true;
+	Component->Play(Time - AudioSection->GetAudioStartTime());
+}
+
+void FMovieSceneAudioTrackInstance::StopSound(int32 RowIndex)
+{
+	if (RowIndex >= PlaybackAudioComponents.Num()) {return;}
+
+	TMap<AActor*, UAudioComponent*>& AudioComponents = PlaybackAudioComponents[RowIndex];
+	for (TMap<AActor*, UAudioComponent*>::TIterator It(AudioComponents); It; ++It)
+	{
+		It.Value()->Stop();
+	}
+}
+
+void FMovieSceneAudioTrackInstance::StopAllSounds()
+{
+	for (int32 i = 0; i < PlaybackAudioComponents.Num(); ++i)
+	{
+		TMap<AActor*, UAudioComponent*>& AudioComponents = PlaybackAudioComponents[i];
+		for (TMap<AActor*, UAudioComponent*>::TIterator It(AudioComponents); It; ++It)
+		{
+			It.Value()->Stop();
+		}
+	}
+}
+
+UAudioComponent* FMovieSceneAudioTrackInstance::GetAudioComponent(AActor* Actor, int32 RowIndex)
+{
+	if (RowIndex + 1 > PlaybackAudioComponents.Num())
+	{
+		while (PlaybackAudioComponents.Num() < RowIndex + 1)
+		{
+			TMap<AActor*, UAudioComponent*> AudioComponentMap;
+			PlaybackAudioComponents.Add(AudioComponentMap);
+		}
+	}
+
+	UAudioComponent*& AudioComponent = PlaybackAudioComponents[RowIndex].FindOrAdd(Actor);
+	if (AudioComponent == NULL)
+	{
+		USoundCue* TempPlaybackAudioCue = NewObject<USoundCue>();
+		
+		AudioComponent = FAudioDevice::CreateComponent(TempPlaybackAudioCue, NULL, Actor, false, false);
+	}
+	return AudioComponent;
+}
+
+TArray<AActor*> FMovieSceneAudioTrackInstance::GetRuntimeActors(const TArray<UObject*>& RuntimeObjects) const
+{
+	TArray<AActor*> Actors;
+	for (int32 ObjectIndex = 0; ObjectIndex < RuntimeObjects.Num(); ++ObjectIndex)
+	{
+		if (RuntimeObjects[ObjectIndex]->IsA<AActor>())
+		{
+			Actors.Add(Cast<AActor>(RuntimeObjects[ObjectIndex]));
+		}
+	}
+	if (AudioTrack->IsAMasterTrack())
+	{
+		check(Actors.Num() == 0);
+		Actors.Add(NULL);
+	}
+	return Actors;
+}
diff --git a/Engine/Source/Runtime/MovieSceneTracks/Private/TrackInstances/MovieSceneBoolTrackInstance.cpp b/Engine/Source/Runtime/MovieSceneTracks/Private/TrackInstances/MovieSceneBoolTrackInstance.cpp
new file mode 100644
--- /dev/null
+++ b/Engine/Source/Runtime/MovieSceneTracks/Private/TrackInstances/MovieSceneBoolTrackInstance.cpp
@@ -0,0 +1,57 @@
+// Copyright 1998-2015 Epic Games, Inc. All Rights Reserved.
+
+#include "MovieSceneTracksPrivatePCH.h"
+#include "MovieSceneBoolTrack.h"
+#include "MovieSceneBoolTrackInstance.h"
+
+
+FMovieSceneBoolTrackInstance::FMovieSceneBoolTrackInstance( UMovieSceneBoolTrack& InBoolTrack )
+{
+	BoolTrack = &InBoolTrack;
+	
+	PropertyBindings = MakeShareable( new FTrackInstancePropertyBindings( BoolTrack->GetPropertyName(), BoolTrack->GetPropertyPath() ) );
+}
+
+
+void FMovieSceneBoolTrackInstance::SaveState(const TArray<UObject*>& RuntimeObjects)
+{
+	for( UObject* Object : RuntimeObjects )
+	{
+		bool BoolValue = PropertyBindings->GetCurrentValue<bool>(Object);
+		InitBoolMap.Add(Object, BoolValue);
+	}
+}
+
+
+void FMovieSceneBoolTrackInstance::RestoreState(const TArray<UObject*>& RuntimeObjects)
+{
+	for( UObject* Object : RuntimeObjects )
+	{
+		bool *BoolValue = InitBoolMap.Find(Object);
+		if (BoolValue != NULL)
+		{
+			PropertyBindings->CallFunction(Object, BoolValue);
+		}
+	}
+
+	PropertyBindings->UpdateBindings( RuntimeObjects );
+}
+
+
+void FMovieSceneBoolTrackInstance::Update( float Position, float LastPosition, const TArray<UObject*>& RuntimeObjects, class IMovieScenePlayer& Player ) 
+{
+	bool BoolValue = false;
+	if( BoolTrack->Eval( Position, LastPosition, BoolValue ) )
+	{
+		for( UObject* Object : RuntimeObjects )
+		{
+			PropertyBindings->CallFunction( Object, &BoolValue );
+		}
+	}
+}
+
+
+void FMovieSceneBoolTrackInstance::RefreshInstance( const TArray<UObject*>& RuntimeObjects, class IMovieScenePlayer& Player )
+{
+	PropertyBindings->UpdateBindings( RuntimeObjects );
+}
diff --git a/Engine/Source/Runtime/MovieSceneTracks/Private/TrackInstances/MovieSceneByteTrackInstance.cpp b/Engine/Source/Runtime/MovieSceneTracks/Private/TrackInstances/MovieSceneByteTrackInstance.cpp
new file mode 100644
--- /dev/null
+++ b/Engine/Source/Runtime/MovieSceneTracks/Private/TrackInstances/MovieSceneByteTrackInstance.cpp
@@ -0,0 +1,57 @@
+// Copyright 1998-2015 Epic Games, Inc. All Rights Reserved.
+
+#include "MovieSceneTracksPrivatePCH.h"
+#include "MovieSceneByteTrack.h"
+#include "MovieSceneByteTrackInstance.h"
+
+
+FMovieSceneByteTrackInstance::FMovieSceneByteTrackInstance( UMovieSceneByteTrack& InByteTrack )
+{
+	ByteTrack = &InByteTrack;
+	
+	PropertyBindings = MakeShareable( new FTrackInstancePropertyBindings( ByteTrack->GetPropertyName(), ByteTrack->GetPropertyPath() ) );
+}
+
+
+void FMovieSceneByteTrackInstance::SaveState(const TArray<UObject*>& RuntimeObjects)
+{
+	for( UObject* Object : RuntimeObjects )
+	{
+		uint8 ByteValue = PropertyBindings->GetCurrentValue<uint8>(Object);
+		InitByteMap.Add(Object, ByteValue);
+	}
+}
+
+
+void FMovieSceneByteTrackInstance::RestoreState(const TArray<UObject*>& RuntimeObjects)
+{
+	for( UObject* Object : RuntimeObjects )
+	{
+		uint8 *ByteValue = InitByteMap.Find(Object);
+		if (ByteValue != NULL)
+		{
+			PropertyBindings->CallFunction(Object, ByteValue);
+		}
+	}
+
+	PropertyBindings->UpdateBindings( RuntimeObjects );
+}
+
+
+void FMovieSceneByteTrackInstance::Update( float Position, float LastPosition, const TArray<UObject*>& RuntimeObjects, class IMovieScenePlayer& Player ) 
+{
+	uint8 ByteValue = 0;
+	if( ByteTrack->Eval( Position, LastPosition, ByteValue ) )
+	{
+		for( UObject* Object : RuntimeObjects )
+		{
+			PropertyBindings->CallFunction( Object, &ByteValue );
+		}
+	}
+}
+
+
+void FMovieSceneByteTrackInstance::RefreshInstance( const TArray<UObject*>& RuntimeObjects, class IMovieScenePlayer& Player )
+{
+	PropertyBindings->UpdateBindings( RuntimeObjects );
+}
diff --git a/Engine/Source/Runtime/MovieSceneTracks/Private/TrackInstances/MovieSceneColorTrackInstance.cpp b/Engine/Source/Runtime/MovieSceneTracks/Private/TrackInstances/MovieSceneColorTrackInstance.cpp
new file mode 100644
--- /dev/null
+++ b/Engine/Source/Runtime/MovieSceneTracks/Private/TrackInstances/MovieSceneColorTrackInstance.cpp
@@ -0,0 +1,86 @@
+// Copyright 1998-2015 Epic Games, Inc. All Rights Reserved.
+
+#include "MovieSceneTracksPrivatePCH.h"
+#include "MovieSceneColorTrackInstance.h"
+#include "MovieSceneCommonHelpers.h"
+#include "SlateCore.h"
+
+
+FMovieSceneColorTrackInstance::FMovieSceneColorTrackInstance( UMovieSceneColorTrack& InColorTrack )
+{
+	ColorTrack = &InColorTrack;
+
+	PropertyBindings = MakeShareable( new FTrackInstancePropertyBindings( ColorTrack->GetPropertyName(), ColorTrack->GetPropertyPath() ) );
+}
+
+void FMovieSceneColorTrackInstance::SaveState(const TArray<UObject*>& RuntimeObjects)
+{
+	for( UObject* Object : RuntimeObjects )
+	{
+		if( ColorTrack->IsSlateColor() )
+		{
+			FSlateColor ColorValue = PropertyBindings->GetCurrentValue<FSlateColor>(Object);
+			InitSlateColorMap.Add(Object, ColorValue);
+		}
+		else
+		{
+			FLinearColor ColorValue = PropertyBindings->GetCurrentValue<FLinearColor>(Object);
+			InitLinearColorMap.Add(Object, ColorValue);
+		}
+	}
+}
+
+void FMovieSceneColorTrackInstance::RestoreState(const TArray<UObject*>& RuntimeObjects)
+{
+	for( UObject* Object : RuntimeObjects )
+	{
+		if( ColorTrack->IsSlateColor() )
+		{
+			FSlateColor *ColorValue = InitSlateColorMap.Find(Object);
+			if (ColorValue != NULL)
+			{
+				PropertyBindings->CallFunction(Object, ColorValue);
+			}
+		}
+		else
+		{
+			FLinearColor *ColorValue = InitLinearColorMap.Find(Object);
+			if (ColorValue != NULL)
+			{
+				PropertyBindings->CallFunction(Object, ColorValue);
+			}
+		}
+	}
+
+	PropertyBindings->UpdateBindings( RuntimeObjects );
+}
+
+void FMovieSceneColorTrackInstance::Update( float Position, float LastPosition, const TArray<UObject*>& RuntimeObjects, class IMovieScenePlayer& Player ) 
+{
+	for(UObject* Object : RuntimeObjects)
+	{
+		if( ColorTrack->IsSlateColor() )
+		{
+			FSlateColor ColorValue = PropertyBindings->GetCurrentValue<FSlateColor>(Object);
+			FLinearColor LinearColor = ColorValue.GetSpecifiedColor();
+			if(ColorTrack->Eval(Position, LastPosition, LinearColor))
+			{
+				FSlateColor NewColor(LinearColor);
+				PropertyBindings->CallFunction(Object, &NewColor);
+			}
+		}
+		else
+		{
+			FLinearColor ColorValue = PropertyBindings->GetCurrentValue<FLinearColor>(Object);
+			if(ColorTrack->Eval(Position, LastPosition, ColorValue))
+			{
+				PropertyBindings->CallFunction(Object, &ColorValue);
+			}
+		}
+	}
+}
+
+void FMovieSceneColorTrackInstance::RefreshInstance( const TArray<UObject*>& RuntimeObjects, class IMovieScenePlayer& Player )
+{
+	PropertyBindings->UpdateBindings( RuntimeObjects );
+}
diff --git a/Engine/Source/Runtime/MovieSceneTracks/Private/TrackInstances/MovieSceneDirectorTrackInstance.cpp b/Engine/Source/Runtime/MovieSceneTracks/Private/TrackInstances/MovieSceneDirectorTrackInstance.cpp
new file mode 100644
--- /dev/null
+++ b/Engine/Source/Runtime/MovieSceneTracks/Private/TrackInstances/MovieSceneDirectorTrackInstance.cpp
@@ -0,0 +1,46 @@
+// Copyright 1998-2015 Epic Games, Inc. All Rights Reserved.
+
+#include "MovieSceneTracksPrivatePCH.h"
+#include "MovieSceneDirectorTrack.h"
+#include "MovieSceneDirectorTrackInstance.h"
+#include "IMovieScenePlayer.h"
+
+
+FMovieSceneDirectorTrackInstance::FMovieSceneDirectorTrackInstance( UMovieSceneDirectorTrack& InDirectorTrack )
+{
+	DirectorTrack = &InDirectorTrack;
+}
+
+
+void FMovieSceneDirectorTrackInstance::Update( float Position, float LastPosition, const TArray<UObject*>& RuntimeObjects, class IMovieScenePlayer& Player ) 
+{
+	// possess the 'first' shot section in order of track index
+	ACameraActor* PossessedCamera = NULL;
+	UMovieSceneShotSection* FirstShotSection = NULL;
+
+	const TArray<UMovieSceneSection*>& ShotSections = DirectorTrack->GetShotSections();
+
+	for (int32 ShotIndex = 0; ShotIndex < ShotSections.Num(); ++ShotIndex)
+	{
+		UMovieSceneShotSection* ShotSection = CastChecked<UMovieSceneShotSection>(ShotSections[ShotIndex]);
+
+		if (ShotSection->IsTimeWithinSection(Position) &&
+			(FirstShotSection == NULL || FirstShotSection->GetRowIndex() > ShotSection->GetRowIndex()))
+		{
+			TArray<UObject*> CameraObjects;
+			// @todo Sequencer - Sub-moviescenes: Get the cameras from the root movie scene instance.  We should support adding cameras for sub-moviescenes as shots
+			Player.GetRuntimeObjects( Player.GetRootMovieSceneInstance(), ShotSection->GetCameraGuid(), CameraObjects);
+
+			if (CameraObjects.Num() > 0)
+			{
+				check(CameraObjects.Num() == 1 && CameraObjects[0]->IsA(ACameraActor::StaticClass()));
+
+				PossessedCamera = Cast<ACameraActor>(CameraObjects[0]);
+
+				FirstShotSection = ShotSection;
+			}
+		}
+	}
+
+	Player.UpdatePreviewViewports(PossessedCamera);
+}
diff --git a/Engine/Source/Runtime/MovieSceneTracks/Private/TrackInstances/MovieSceneFloatTrackInstance.cpp b/Engine/Source/Runtime/MovieSceneTracks/Private/TrackInstances/MovieSceneFloatTrackInstance.cpp
new file mode 100644
--- /dev/null
+++ b/Engine/Source/Runtime/MovieSceneTracks/Private/TrackInstances/MovieSceneFloatTrackInstance.cpp
@@ -0,0 +1,57 @@
+// Copyright 1998-2015 Epic Games, Inc. All Rights Reserved.
+
+#include "MovieSceneTracksPrivatePCH.h"
+#include "MovieSceneFloatTrack.h"
+#include "MovieSceneFloatTrackInstance.h"
+
+
+FMovieSceneFloatTrackInstance::FMovieSceneFloatTrackInstance( UMovieSceneFloatTrack& InFloatTrack )
+{
+	FloatTrack = &InFloatTrack;
+
+	PropertyBindings = MakeShareable( new FTrackInstancePropertyBindings( FloatTrack->GetPropertyName(), FloatTrack->GetPropertyPath() ) );
+}
+
+
+void FMovieSceneFloatTrackInstance::SaveState(const TArray<UObject*>& RuntimeObjects)
+{
+	for( UObject* Object : RuntimeObjects )
+	{
+		float FloatValue = PropertyBindings->GetCurrentValue<float>(Object);
+		InitFloatMap.Add(Object, FloatValue);
+	}
+}
+
+
+void FMovieSceneFloatTrackInstance::RestoreState(const TArray<UObject*>& RuntimeObjects)
+{
+	for( UObject* Object : RuntimeObjects )
+	{
+		float *FloatValue = InitFloatMap.Find(Object);
+		if (FloatValue != NULL)
+		{
+			PropertyBindings->CallFunction(Object, FloatValue);
+		}
+	}
+
+	PropertyBindings->UpdateBindings( RuntimeObjects );
+}
+
+
+void FMovieSceneFloatTrackInstance::Update( float Position, float LastPosition, const TArray<UObject*>& RuntimeObjects, class IMovieScenePlayer& Player ) 
+{
+	float FloatValue = 0.0f;
+	if( FloatTrack->Eval( Position, LastPosition, FloatValue ) )
+	{
+		for(UObject* Object : RuntimeObjects)
+		{
+			PropertyBindings->CallFunction( Object, &FloatValue );
+		}
+	}
+}
+
+
+void FMovieSceneFloatTrackInstance::RefreshInstance( const TArray<UObject*>& RuntimeObjects, class IMovieScenePlayer& Player )
+{
+	PropertyBindings->UpdateBindings( RuntimeObjects );
+}
diff --git a/Engine/Source/Runtime/MovieSceneTracks/Private/TrackInstances/MovieSceneParticleTrackInstance.cpp b/Engine/Source/Runtime/MovieSceneTracks/Private/TrackInstances/MovieSceneParticleTrackInstance.cpp
new file mode 100644
--- /dev/null
+++ b/Engine/Source/Runtime/MovieSceneTracks/Private/TrackInstances/MovieSceneParticleTrackInstance.cpp
@@ -0,0 +1,80 @@
+// Copyright 1998-2015 Epic Games, Inc. All Rights Reserved.
+
+#include "MovieSceneTracksPrivatePCH.h"
+#include "MovieSceneParticleTrackInstance.h"
+#include "IMovieScenePlayer.h"
+#include "MovieSceneParticleTrack.h"
+#include "MovieSceneParticleSection.h"
+#include "Particles/Emitter.h"
+FMovieSceneParticleTrackInstance::~FMovieSceneParticleTrackInstance()
+{
+}
+
+void FMovieSceneParticleTrackInstance::Update( float Position, float LastPosition, const TArray<UObject*>& RuntimeObjects, class IMovieScenePlayer& Player ) 
+{
+	// @todo Sequencer We need something analagous to Matinee 1's particle replay tracks
+	// What we have here is simple toggling/triggering
+
+	if (Position > LastPosition && Player.GetPlaybackStatus() == EMovieScenePlayerStatus::Playing)
+	{
+		bool bTrigger = false, bOn = false, bOff = false;
+
+		const TArray<UMovieSceneSection*> Sections = ParticleTrack->GetAllParticleSections();
+		for (int32 i = 0; i < Sections.Num(); ++i)
+		{
+			UMovieSceneParticleSection* Section = Cast<UMovieSceneParticleSection>(Sections[i]);
+			if (Section->GetKeyType() == EParticleKey::Trigger)
+			{
+				if (Position > Section->GetStartTime() && LastPosition < Section->GetStartTime())
+				{
+					bTrigger = true;
+				}
+			}
+			else if (Section->GetKeyType() == EParticleKey::Toggle)
+			{
+				if (Position >= Section->GetStartTime() && Position <= Section->GetEndTime())
+				{
+					bOn = true;
+				}
+				else if (Position >= Section->GetEndTime() && LastPosition < Section->GetEndTime())
+				{
+					bOff = true;
+				}
+			}
+		}
+		
+		if (bTrigger || bOn || bOff)
+		{
+			for (int32 i = 0; i < RuntimeObjects.Num(); ++i)
+			{
+				AEmitter* Emitter = Cast<AEmitter>(RuntimeObjects[i]);
+				if (Emitter)
+				{
+					if (bTrigger)
+					{
+						Emitter->ToggleActive();
+					}
+					else if (bOn)
+					{
+						Emitter->Activate();
+					}
+					else if (bOff)
+					{
+						Emitter->Deactivate();
+					}
+				}
+			}
+		}
+	}
+	else
+	{
+		for (int32 i = 0; i < RuntimeObjects.Num(); ++i)
+		{
+			AEmitter* Emitter = Cast<AEmitter>(RuntimeObjects[i]);
+			if (Emitter)
+			{
+				Emitter->Deactivate();
+			}
+		}
+	}
+}
diff --git a/Engine/Source/Runtime/MovieSceneTracks/Private/TrackInstances/MovieSceneVectorTrackInstance.cpp b/Engine/Source/Runtime/MovieSceneTracks/Private/TrackInstances/MovieSceneVectorTrackInstance.cpp
new file mode 100644
--- /dev/null
+++ b/Engine/Source/Runtime/MovieSceneTracks/Private/TrackInstances/MovieSceneVectorTrackInstance.cpp
@@ -0,0 +1,128 @@
+// Copyright 1998-2015 Epic Games, Inc. All Rights Reserved.
+
+#include "MovieSceneTracksPrivatePCH.h"
+#include "MovieSceneVectorTrackInstance.h"
+
+FMovieSceneVectorTrackInstance::FMovieSceneVectorTrackInstance( UMovieSceneVectorTrack& InVectorTrack )
+{
+	VectorTrack = &InVectorTrack;
+
+	PropertyBindings = MakeShareable( new FTrackInstancePropertyBindings( VectorTrack->GetPropertyName(), VectorTrack->GetPropertyPath() ) );
+}
+
+void FMovieSceneVectorTrackInstance::SaveState(const TArray<UObject*>& RuntimeObjects)
+{
+	int32 NumChannelsUsed = VectorTrack->GetNumChannelsUsed();
+	for( UObject* Object : RuntimeObjects )
+	{
+		switch( NumChannelsUsed )
+		{
+			case 2:
+			{
+				FVector2D VectorValue = PropertyBindings->GetCurrentValue<FVector2D>(Object);
+				InitVector2DMap.Add(Object, VectorValue);
+				break;
+			}
+			case 3:
+			{
+				FVector VectorValue = PropertyBindings->GetCurrentValue<FVector>(Object);
+				InitVector3DMap.Add(Object, VectorValue);
+				break;
+			}
+			case 4:
+			{
+				FVector4 VectorValue = PropertyBindings->GetCurrentValue<FVector4>(Object);
+				InitVector4DMap.Add(Object, VectorValue);
+				break;
+			}
+		}
+	}
+}
+
+void FMovieSceneVectorTrackInstance::RestoreState(const TArray<UObject*>& RuntimeObjects)
+{
+	int32 NumChannelsUsed = VectorTrack->GetNumChannelsUsed();
+	for( UObject* Object : RuntimeObjects )
+	{
+		switch( NumChannelsUsed )
+		{
+			case 2:
+			{
+				FVector2D *VectorValue = InitVector2DMap.Find(Object);
+				if (VectorValue != NULL)
+				{
+					PropertyBindings->CallFunction(Object, VectorValue);
+				}
+				break;
+			}
+			case 3:
+			{
+				FVector *VectorValue = InitVector3DMap.Find(Object);
+				if (VectorValue != NULL)
+				{
+					PropertyBindings->CallFunction(Object, VectorValue);
+				}
+				break;
+			}
+			case 4:
+			{
+				FVector4 *VectorValue = InitVector4DMap.Find(Object);
+				if (VectorValue != NULL)
+				{
+					PropertyBindings->CallFunction(Object, VectorValue);
+				}
+				break;
+			}
+		}
+	}
+
+	PropertyBindings->UpdateBindings( RuntimeObjects );
+}
+
+void FMovieSceneVectorTrackInstance::Update( float Position, float LastPosition, const TArray<UObject*>& RuntimeObjects, class IMovieScenePlayer& Player ) 
+{
+	FVector4 Vector;
+	if( VectorTrack->Eval( Position, LastPosition, Vector ) )
+	{
+		int32 NumChannelsUsed = VectorTrack->GetNumChannelsUsed();
+		switch( NumChannelsUsed )
+		{
+			case 2:
+			{
+				FVector2D Value(Vector.X, Vector.Y);
+				for(UObject* Object : RuntimeObjects)
+				{
+					PropertyBindings->CallFunction(Object, &Value);
+				}
+				break;
+			}
+			case 3:
+			{
+				FVector Value(Vector.X, Vector.Y, Vector.Z);
+				for(UObject* Object : RuntimeObjects)
+				{
+					PropertyBindings->CallFunction(Object, &Value);
+				}
+				break;
+			}
+			case 4:
+			{
+				for(UObject* Object : RuntimeObjects)
+				{
+					PropertyBindings->CallFunction(Object, &Vector);
+				}
+				break;
+			}
+			default:
+				UE_LOG(LogSequencerRuntime, Warning, TEXT("Invalid number of channels(%d) for vector track"), NumChannelsUsed );
+				break;
+		}
+		
+	}
+}
+
+
+void FMovieSceneVectorTrackInstance::RefreshInstance(const TArray<UObject*>& RuntimeObjects, class IMovieScenePlayer& Player)
+{
+	PropertyBindings->UpdateBindings(RuntimeObjects);
+}
diff --git a/Engine/Source/Runtime/MovieSceneTracks/Private/TrackInstances/SubMovieSceneTrackInstance.cpp b/Engine/Source/Runtime/MovieSceneTracks/Private/TrackInstances/SubMovieSceneTrackInstance.cpp
new file mode 100644
--- /dev/null
+++ b/Engine/Source/Runtime/MovieSceneTracks/Private/TrackInstances/SubMovieSceneTrackInstance.cpp
@@ -0,0 +1,83 @@
+// Copyright 1998-2015 Epic Games, Inc. All Rights Reserved.
+
+#include "MovieSceneTracksPrivatePCH.h"
+#include "SubMovieSceneTrack.h"
+#include "SubMovieSceneTrackInstance.h"
+#include "MovieSceneInstance.h"
+#include "SubMovieSceneSection.h"
+#include "IMovieScenePlayer.h"
+
+
+FSubMovieSceneTrackInstance::FSubMovieSceneTrackInstance( USubMovieSceneTrack& InTrack )
+{
+	SubMovieSceneTrack = &InTrack;
+}
+
+
+void FSubMovieSceneTrackInstance::RefreshInstance( const TArray<UObject*>& RuntimeObjects, class IMovieScenePlayer& Player )
+{
+	TArray<UMovieSceneSection*> AllSections = SubMovieSceneTrack->GetAllSections();
+
+	// Sections we encountered
+	TSet<UMovieSceneSection*> FoundSections;
+
+	// Iterate through each section and create an instance for it if it doesn't exist
+	for( int32 SectionIndex = 0; SectionIndex < AllSections.Num(); ++SectionIndex )
+	{
+		USubMovieSceneSection* Section = CastChecked<USubMovieSceneSection>( AllSections[SectionIndex] );
+
+		FoundSections.Add( Section );
+
+		TSharedPtr<FMovieSceneInstance> Instance = SubMovieSceneInstances.FindRef(Section);
+		if( !Instance.IsValid() )
+		{
+			Instance = MakeShareable( new FMovieSceneInstance( *Section->GetMovieScene() ) );
+
+			Player.AddMovieSceneInstance( *Section, Instance.ToSharedRef() );
+
+			SubMovieSceneInstances.Add( Section, Instance.ToSharedRef() );
+		}
+
+		// Refresh the existing instance
+		Instance->RefreshInstance( Player );
+	}
+
+	TMap< TWeakObjectPtr<USubMovieSceneSection>, TSharedPtr<FMovieSceneInstance> >::TIterator It =  SubMovieSceneInstances.CreateIterator();
+	for( ; It; ++It )
+	{
+		// Remove any sections that no longer exist
+		if( !FoundSections.Contains( It.Key().Get() ) )
+		{
+			Player.RemoveMovieSceneInstance( *It.Key().Get(), It.Value().ToSharedRef() );
+			It.RemoveCurrent();
+		}
+	}
+}
+
+
+void FSubMovieSceneTrackInstance::Update( float Position, float LastPosition, const TArray<UObject*>& RuntimeObjects, class IMovieScenePlayer& Player ) 
+{
+	TArray<UMovieSceneSection*> AllSections = SubMovieSceneTrack->GetAllSections();
+
+	TArray<UMovieSceneSection*> TraversedSections = MovieSceneHelpers::GetTraversedSections( AllSections, Position, LastPosition );
+
+	for( int32 SectionIndex = 0; SectionIndex < TraversedSections.Num(); ++SectionIndex )
+	{
+		USubMovieSceneSection* Section = CastChecked<USubMovieSceneSection>( TraversedSections[SectionIndex] );
+
+		TSharedPtr<FMovieSceneInstance>& Instance = SubMovieSceneInstances.FindChecked( Section );
+
+		UMovieScene* MovieScene = Section->GetMovieScene();
+
+		// @todo Sequencer we should not do this every update
+		TRange<float> TimeRange = MovieScene->GetTimeRange();
+
+		// Position for the movie scene needs to be in local space
+		float LocalDelta = TimeRange.GetLowerBoundValue() - Section->GetStartTime();
+		float LocalPosition = Position + LocalDelta;
+
+
+		Instance->Update( LocalPosition, LastPosition + LocalDelta, Player );
+	}
+
+}
diff --git a/Engine/Source/Runtime/MovieSceneTracks/Private/Tracks/MovieScene3DTransformTrack.cpp b/Engine/Source/Runtime/MovieSceneTracks/Private/Tracks/MovieScene3DTransformTrack.cpp
new file mode 100644
--- /dev/null
+++ b/Engine/Source/Runtime/MovieSceneTracks/Private/Tracks/MovieScene3DTransformTrack.cpp
@@ -0,0 +1,110 @@
+// Copyright 1998-2015 Epic Games, Inc. All Rights Reserved.
+
+#include "MovieSceneTracksPrivatePCH.h"
+#include "MovieScene3DTransformSection.h"
+#include "MovieScene3DTransformTrack.h"
+#include "IMovieScenePlayer.h"
+#include "MovieScene3DTransformTrackInstance.h"
+
+bool FTransformKey::ShouldKeyTranslation( EAxis::Type Axis) const
+{
+	// If the previous transform is not valid we have nothing to compare against so assume we always key. Otherwise check for differences
+	return !PreviousTransform.IsValid() || GetVectorComponentIfDifferent( Axis, NewTransform.Translation, PreviousTransform.Translation );
+}
+
+bool FTransformKey::ShouldKeyRotation( EAxis::Type Axis ) const
+{
+	// If the previous transform is not valid we have nothing to compare against so assume we always key. Otherwise check for differences
+	return !PreviousTransform.IsValid() || GetVectorComponentIfDifferent( Axis, NewTransform.Rotation.Euler(), PreviousTransform.Rotation.Euler() );
+}
+
+bool FTransformKey::ShouldKeyScale( EAxis::Type Axis ) const
+{
+	// If the previous transform is not valid we have nothing to compare against so assume we always key. Otherwise check for differences
+	return !PreviousTransform.IsValid() || GetVectorComponentIfDifferent( Axis, NewTransform.Scale, PreviousTransform.Scale );
+}
+
+bool FTransformKey::ShouldKeyAny() const
+{
+	return ShouldKeyTranslation(EAxis::X) || ShouldKeyTranslation(EAxis::Y) || ShouldKeyTranslation(EAxis::Z) ||
+		ShouldKeyRotation(EAxis::X) || ShouldKeyRotation(EAxis::Y) || ShouldKeyRotation(EAxis::Z) ||
+		ShouldKeyScale(EAxis::X) || ShouldKeyScale(EAxis::Y) || ShouldKeyScale(EAxis::Z);
+}
+
+bool FTransformKey::GetVectorComponentIfDifferent( EAxis::Type Axis, const FVector& Current, const FVector& Previous ) const
+{
+	bool bShouldAddKey = false;
+
+	// Using lower precision to avoid small changes due to floating point error
+	// @todo Sequencer - Floating point error introduced by matrix math causing extra keys to be added
+	const float Precision = 1e-3;
+
+	// Add a key if there is no previous transform to compare against or the current and previous values are not equal
+
+	if( Axis == EAxis::X && !FMath::IsNearlyEqual( Current.X, Previous.X, Precision ) )
+	{
+		bShouldAddKey = true;
+	}
+	else if( Axis == EAxis::Y && !FMath::IsNearlyEqual( Current.Y, Previous.Y, Precision ) )
+	{
+		bShouldAddKey = true;
+	}
+	else if( Axis == EAxis::Z && !FMath::IsNearlyEqual( Current.Z, Previous.Z, Precision ) )
+	{
+		bShouldAddKey = true;
+	}
+
+	return bShouldAddKey;
+}
+
+UMovieScene3DTransformTrack::UMovieScene3DTransformTrack( const FObjectInitializer& ObjectInitializer )
+	: Super( ObjectInitializer )
+{
+}
+
+UMovieSceneSection* UMovieScene3DTransformTrack::CreateNewSection()
+{
+	return NewObject<UMovieSceneSection>(this, UMovieScene3DTransformSection::StaticClass(), NAME_None, RF_Transactional);
+}
+
+TSharedPtr<IMovieSceneTrackInstance> UMovieScene3DTransformTrack::CreateInstance()
+{
+	return MakeShareable( new FMovieScene3DTransformTrackInstance( *this ) );
+}
+
+bool UMovieScene3DTransformTrack::AddKeyToSection( const FGuid& ObjectHandle, const FTransformKey& InKey, const bool bUnwindRotation )
+{
+	const UMovieSceneSection* NearestSection = MovieSceneHelpers::FindSectionAtTime(Sections, InKey.GetKeyTime());
+	if (!NearestSection || CastChecked<UMovieScene3DTransformSection>(NearestSection)->NewKeyIsNewData(InKey))
+	{
+		Modify();
+
+		UMovieScene3DTransformSection* NewSection = Cast<UMovieScene3DTransformSection>( FindOrAddSection( InKey.GetKeyTime() ) );
+
+		// key each component of the transform
+		NewSection->AddTranslationKeys( InKey );
+		NewSection->AddRotationKeys( InKey, bUnwindRotation );
+		NewSection->AddScaleKeys( InKey );
+
+		return true;
+	}
+	return false;
+}
+
+
+bool UMovieScene3DTransformTrack::Eval( float Position, float LastPosition, FVector& OutTranslation, FRotator& OutRotation, FVector& OutScale, bool& OutHasTranslationKeys, bool& OutHasRotationKeys, bool& OutHasScaleKeys ) const
+{
+	const UMovieSceneSection* Section = MovieSceneHelpers::FindSectionAtTime( Sections, Position );
+
+	if( Section )
+	{
+		const UMovieScene3DTransformSection* TransformSection = CastChecked<UMovieScene3DTransformSection>( Section );
+
+		// Evalulate translation,rotation, and scale curves.  If no keys were found on one of these, that component of the transform will remain unchained
+		OutHasTranslationKeys = TransformSection->EvalTranslation( Position, OutTranslation );
+		OutHasRotationKeys = TransformSection->EvalRotation( Position, OutRotation );
+		OutHasScaleKeys = TransformSection->EvalScale( Position, OutScale );
+	}
+
+	return Section != NULL;
+}
\ No newline at end of file
diff --git a/Engine/Source/Runtime/MovieSceneTracks/Private/Tracks/MovieSceneAnimationTrack.cpp b/Engine/Source/Runtime/MovieSceneTracks/Private/Tracks/MovieSceneAnimationTrack.cpp
new file mode 100644
--- /dev/null
+++ b/Engine/Source/Runtime/MovieSceneTracks/Private/Tracks/MovieSceneAnimationTrack.cpp
@@ -0,0 +1,92 @@
+// Copyright 1998-2015 Epic Games, Inc. All Rights Reserved.
+
+#include "MovieSceneTracksPrivatePCH.h"
+#include "MovieSceneAnimationSection.h"
+#include "MovieSceneAnimationTrack.h"
+#include "IMovieScenePlayer.h"
+#include "MovieSceneAnimationTrackInstance.h"
+
+#define LOCTEXT_NAMESPACE "MovieSceneAnimationTrack"
+
+UMovieSceneAnimationTrack::UMovieSceneAnimationTrack( const FObjectInitializer& ObjectInitializer )
+	: Super( ObjectInitializer )
+{
+}
+
+FName UMovieSceneAnimationTrack::GetTrackName() const
+{
+	return FName("Animation");
+}
+
+
+TSharedPtr<IMovieSceneTrackInstance> UMovieSceneAnimationTrack::CreateInstance()
+{
+	return MakeShareable( new FMovieSceneAnimationTrackInstance( *this ) ); 
+}
+
+TArray<UMovieSceneSection*> UMovieSceneAnimationTrack::GetAllSections() const
+{
+	TArray<UMovieSceneSection*> OutSections;
+	for (int32 SectionIndex = 0; SectionIndex < AnimationSections.Num(); ++SectionIndex)
+	{
+		OutSections.Add( AnimationSections[SectionIndex] );
+	}
+	return OutSections;
+}
+
+void UMovieSceneAnimationTrack::RemoveAllAnimationData()
+{
+}
+
+bool UMovieSceneAnimationTrack::HasSection( UMovieSceneSection* Section ) const
+{
+	return AnimationSections.Find( Section ) != INDEX_NONE;
+}
+
+void UMovieSceneAnimationTrack::RemoveSection( UMovieSceneSection* Section )
+{
+	AnimationSections.Remove( Section );
+}
+
+bool UMovieSceneAnimationTrack::IsEmpty() const
+{
+	return AnimationSections.Num() == 0;
+}
+
+TRange<float> UMovieSceneAnimationTrack::GetSectionBoundaries() const
+{
+	TArray< TRange<float> > Bounds;
+	for (int32 i = 0; i < AnimationSections.Num(); ++i)
+	{
+		Bounds.Add(AnimationSections[i]->GetRange());
+	}
+	return TRange<float>::Hull(Bounds);
+}
+
+void UMovieSceneAnimationTrack::AddNewAnimation(float KeyTime, UAnimSequence* AnimSequence)
+{
+	// add the section
+	UMovieSceneAnimationSection* NewSection = NewObject<UMovieSceneAnimationSection>(this);
+	NewSection->InitialPlacement( AnimationSections, KeyTime, KeyTime + AnimSequence->SequenceLength, SupportsMultipleRows() );
+	NewSection->SetAnimationStartTime( KeyTime );
+	NewSection->SetAnimSequence(AnimSequence);
+
+	AnimationSections.Add(NewSection);
+}
+
+UMovieSceneSection* UMovieSceneAnimationTrack::GetAnimSectionAtTime(float Time)
+{
+	for (int32 i = 0; i < AnimationSections.Num(); ++i)
+	{
+		UMovieSceneSection* Section = AnimationSections[i];
+		if( Section->IsTimeWithinSection( Time ) )
+		{
+			return Section;
+		}
+	}
+
+	return NULL;
+}
+
+
+#undef LOCTEXT_NAMESPACE
diff --git a/Engine/Source/Runtime/MovieSceneTracks/Private/Tracks/MovieSceneAudioTrack.cpp b/Engine/Source/Runtime/MovieSceneTracks/Private/Tracks/MovieSceneAudioTrack.cpp
new file mode 100644
--- /dev/null
+++ b/Engine/Source/Runtime/MovieSceneTracks/Private/Tracks/MovieSceneAudioTrack.cpp
@@ -0,0 +1,99 @@
+// Copyright 1998-2015 Epic Games, Inc. All Rights Reserved.
+
+#include "MovieSceneTracksPrivatePCH.h"
+#include "MovieSceneAudioSection.h"
+#include "MovieSceneAudioTrack.h"
+#include "IMovieScenePlayer.h"
+#include "SoundDefinitions.h"
+#include "Runtime/Engine/Classes/Kismet/GameplayStatics.h"
+#include "Runtime/Engine/Public/AudioDecompress.h"
+#include "MovieSceneAudioTrackInstance.h"
+
+#define LOCTEXT_NAMESPACE "MovieSceneDirectorTrack"
+
+UMovieSceneAudioTrack::UMovieSceneAudioTrack( const FObjectInitializer& ObjectInitializer )
+	: Super( ObjectInitializer )
+{
+}
+
+FName UMovieSceneAudioTrack::GetTrackName() const
+{
+	return AudioTrackConstants::UniqueTrackName;
+}
+
+
+TSharedPtr<IMovieSceneTrackInstance> UMovieSceneAudioTrack::CreateInstance()
+{
+	return MakeShareable( new FMovieSceneAudioTrackInstance( *this ) ); 
+}
+
+TArray<UMovieSceneSection*> UMovieSceneAudioTrack::GetAllSections() const
+{
+	TArray<UMovieSceneSection*> OutSections;
+	for (int32 SectionIndex = 0; SectionIndex < AudioSections.Num(); ++SectionIndex)
+	{
+		OutSections.Add( AudioSections[SectionIndex] );
+	}
+	return OutSections;
+}
+
+void UMovieSceneAudioTrack::RemoveAllAnimationData()
+{
+}
+
+bool UMovieSceneAudioTrack::HasSection( UMovieSceneSection* Section ) const
+{
+	return AudioSections.Find( Section ) != INDEX_NONE;
+}
+
+void UMovieSceneAudioTrack::RemoveSection( UMovieSceneSection* Section )
+{
+	AudioSections.Remove( Section );
+}
+
+bool UMovieSceneAudioTrack::IsEmpty() const
+{
+	return AudioSections.Num() == 0;
+}
+
+TRange<float> UMovieSceneAudioTrack::GetSectionBoundaries() const
+{
+	TArray< TRange<float> > Bounds;
+	for (int32 i = 0; i < AudioSections.Num(); ++i)
+	{
+		Bounds.Add(AudioSections[i]->GetRange());
+	}
+	return TRange<float>::Hull(Bounds);
+}
+
+void UMovieSceneAudioTrack::AddNewSound(USoundBase* Sound, float Time)
+{
+	check(Sound);
+	
+	// determine initial duration
+	// @todo Once we have infinite sections, we can remove this
+	float DurationToUse = 1.f; // if all else fails, use 1 second duration
+
+	float SoundDuration = Sound->GetDuration();
+	if (SoundDuration != INDEFINITELY_LOOPING_DURATION)
+	{
+		DurationToUse = SoundDuration;
+	}
+
+	// add the section
+	UMovieSceneAudioSection* NewSection = NewObject<UMovieSceneAudioSection>(this);
+	NewSection->InitialPlacement( AudioSections, Time, Time + DurationToUse, SupportsMultipleRows() );
+	NewSection->SetAudioStartTime(Time);
+	NewSection->SetSound(Sound);
+
+	AudioSections.Add(NewSection);
+}
+
+bool UMovieSceneAudioTrack::IsAMasterTrack() const
+{
+	return Cast<UMovieScene>(GetOuter())->IsAMasterTrack(this);
+}
+
+
+
+#undef LOCTEXT_NAMESPACE
diff --git a/Engine/Source/Runtime/MovieSceneTracks/Private/Tracks/MovieSceneBoolTrack.cpp b/Engine/Source/Runtime/MovieSceneTracks/Private/Tracks/MovieSceneBoolTrack.cpp
new file mode 100644
--- /dev/null
+++ b/Engine/Source/Runtime/MovieSceneTracks/Private/Tracks/MovieSceneBoolTrack.cpp
@@ -0,0 +1,51 @@
+// Copyright 1998-2015 Epic Games, Inc. All Rights Reserved.
+
+#include "MovieSceneTracksPrivatePCH.h"
+#include "MovieSceneBoolSection.h"
+#include "MovieSceneBoolTrack.h"
+#include "IMovieScenePlayer.h"
+#include "MovieSceneBoolTrackInstance.h"
+
+
+UMovieSceneBoolTrack::UMovieSceneBoolTrack( const FObjectInitializer& ObjectInitializer )
+	: Super( ObjectInitializer )
+{
+}
+
+UMovieSceneSection* UMovieSceneBoolTrack::CreateNewSection()
+{
+	return NewObject<UMovieSceneSection>(this, UMovieSceneBoolSection::StaticClass(), NAME_None, RF_Transactional);
+}
+
+TSharedPtr<IMovieSceneTrackInstance> UMovieSceneBoolTrack::CreateInstance()
+{
+	return MakeShareable( new FMovieSceneBoolTrackInstance( *this ) );
+}
+
+bool UMovieSceneBoolTrack::AddKeyToSection( float Time, bool Value )
+{
+	const UMovieSceneSection* NearestSection = MovieSceneHelpers::FindNearestSectionAtTime( Sections, Time );
+	if (!NearestSection || CastChecked<UMovieSceneBoolSection>(NearestSection)->NewKeyIsNewData(Time, Value))
+	{
+		Modify();
+
+		UMovieSceneBoolSection* NewSection = CastChecked<UMovieSceneBoolSection>(FindOrAddSection(  Time ));
+	
+		NewSection->AddKey( Time, Value );
+
+		return true;
+	}
+	return false;
+}
+
+bool UMovieSceneBoolTrack::Eval( float Position, float LastPostion, bool& OutBool ) const
+{	
+	const UMovieSceneSection* Section = MovieSceneHelpers::FindSectionAtTime( Sections, Position );
+
+	if( Section )
+	{
+		OutBool = CastChecked<UMovieSceneBoolSection>( Section )->Eval( Position );
+	}
+
+	return Section != NULL;
+}
diff --git a/Engine/Source/Runtime/MovieSceneTracks/Private/Tracks/MovieSceneByteTrack.cpp b/Engine/Source/Runtime/MovieSceneTracks/Private/Tracks/MovieSceneByteTrack.cpp
new file mode 100644
--- /dev/null
+++ b/Engine/Source/Runtime/MovieSceneTracks/Private/Tracks/MovieSceneByteTrack.cpp
@@ -0,0 +1,61 @@
+// Copyright 1998-2015 Epic Games, Inc. All Rights Reserved.
+
+#include "MovieSceneTracksPrivatePCH.h"
+#include "MovieSceneByteSection.h"
+#include "MovieSceneByteTrack.h"
+#include "IMovieScenePlayer.h"
+#include "MovieSceneByteTrackInstance.h"
+
+
+UMovieSceneByteTrack::UMovieSceneByteTrack( const FObjectInitializer& ObjectInitializer )
+	: Super( ObjectInitializer )
+{
+}
+
+UMovieSceneSection* UMovieSceneByteTrack::CreateNewSection()
+{
+	return NewObject<UMovieSceneSection>(this, UMovieSceneByteSection::StaticClass(), NAME_None, RF_Transactional);
+}
+
+TSharedPtr<IMovieSceneTrackInstance> UMovieSceneByteTrack::CreateInstance()
+{
+	return MakeShareable( new FMovieSceneByteTrackInstance( *this ) );
+}
+
+bool UMovieSceneByteTrack::AddKeyToSection( float Time, uint8 Value )
+{
+	const UMovieSceneSection* NearestSection = MovieSceneHelpers::FindNearestSectionAtTime( Sections, Time );
+	if (!NearestSection || CastChecked<UMovieSceneByteSection>(NearestSection)->NewKeyIsNewData(Time, Value))
+	{
+		Modify();
+
+		UMovieSceneByteSection* NewSection = CastChecked<UMovieSceneByteSection>(FindOrAddSection(  Time ));
+	
+		NewSection->AddKey( Time, Value );
+
+		return true;
+	}
+	return false;
+}
+
+bool UMovieSceneByteTrack::Eval( float Position, float LastPostion, uint8& OutByte ) const
+{	
+	const UMovieSceneSection* Section = MovieSceneHelpers::FindSectionAtTime( Sections, Position );
+
+	if( Section )
+	{
+		OutByte = CastChecked<UMovieSceneByteSection>( Section )->Eval( Position );
+	}
+
+	return Section != NULL;
+}
+
+void UMovieSceneByteTrack::SetEnum(UEnum* InEnum)
+{
+	Enum = InEnum;
+}
+
+class UEnum* UMovieSceneByteTrack::GetEnum() const
+{
+	return Enum;
+}
diff --git a/Engine/Source/Runtime/MovieSceneTracks/Private/Tracks/MovieSceneColorTrack.cpp b/Engine/Source/Runtime/MovieSceneTracks/Private/Tracks/MovieSceneColorTrack.cpp
new file mode 100644
--- /dev/null
+++ b/Engine/Source/Runtime/MovieSceneTracks/Private/Tracks/MovieSceneColorTrack.cpp
@@ -0,0 +1,54 @@
+// Copyright 1998-2015 Epic Games, Inc. All Rights Reserved.
+
+#include "MovieSceneTracksPrivatePCH.h"
+#include "MovieSceneColorSection.h"
+#include "MovieSceneColorTrack.h"
+#include "IMovieScenePlayer.h"
+#include "MovieSceneColorTrackInstance.h"
+
+UMovieSceneColorTrack::UMovieSceneColorTrack( const FObjectInitializer& ObjectInitializer )
+	: Super( ObjectInitializer )
+	, bIsSlateColor( false )
+{
+}
+
+UMovieSceneSection* UMovieSceneColorTrack::CreateNewSection()
+{
+	return NewObject<UMovieSceneSection>(this, UMovieSceneColorSection::StaticClass(), NAME_None, RF_Transactional);
+}
+
+TSharedPtr<IMovieSceneTrackInstance> UMovieSceneColorTrack::CreateInstance()
+{
+	return MakeShareable( new FMovieSceneColorTrackInstance( *this ) ); 
+}
+
+
+bool UMovieSceneColorTrack::AddKeyToSection( float Time, const FColorKey& Key )
+{
+	bIsSlateColor = Key.bIsSlateColor;
+
+	const UMovieSceneSection* NearestSection = MovieSceneHelpers::FindNearestSectionAtTime( Sections, Time );
+	if (!NearestSection || Key.bAddKeyEvenIfUnchanged || CastChecked<UMovieSceneColorSection>(NearestSection)->NewKeyIsNewData(Time, Key.Value))
+	{
+		Modify();
+		UMovieSceneColorSection* NewSection = CastChecked<UMovieSceneColorSection>( FindOrAddSection( Time ) );
+
+		NewSection->AddKey( Time, Key );
+
+		return true;
+	}
+	return false;
+}
+
+
+bool UMovieSceneColorTrack::Eval( float Position, float LastPosition, FLinearColor& OutColor ) const
+{
+	const UMovieSceneSection* Section = MovieSceneHelpers::FindSectionAtTime( Sections, Position );
+
+	if( Section )
+	{
+		OutColor = CastChecked<UMovieSceneColorSection>( Section )->Eval( Position, OutColor );
+	}
+
+	return Section != NULL;
+}
\ No newline at end of file
diff --git a/Engine/Source/Runtime/MovieSceneTracks/Private/Tracks/MovieSceneDirectorTrack.cpp b/Engine/Source/Runtime/MovieSceneTracks/Private/Tracks/MovieSceneDirectorTrack.cpp
new file mode 100644
--- /dev/null
+++ b/Engine/Source/Runtime/MovieSceneTracks/Private/Tracks/MovieSceneDirectorTrack.cpp
@@ -0,0 +1,79 @@
+// Copyright 1998-2015 Epic Games, Inc. All Rights Reserved.
+
+#include "MovieSceneTracksPrivatePCH.h"
+#include "MovieSceneShotSection.h"
+#include "MovieSceneDirectorTrack.h"
+#include "IMovieScenePlayer.h"
+#include "MovieSceneDirectorTrackInstance.h"
+
+#define LOCTEXT_NAMESPACE "MovieSceneDirectorTrack"
+
+UMovieSceneDirectorTrack::UMovieSceneDirectorTrack( const FObjectInitializer& ObjectInitializer )
+	: Super( ObjectInitializer )
+{
+}
+
+
+FName UMovieSceneDirectorTrack::GetTrackName() const
+{
+	static const FName UniqueName("DirectorTrack");
+
+	return UniqueName;
+}
+
+TSharedPtr<IMovieSceneTrackInstance> UMovieSceneDirectorTrack::CreateInstance()
+{
+	return MakeShareable( new FMovieSceneDirectorTrackInstance( *this ) ); 
+}
+
+TArray<UMovieSceneSection*> UMovieSceneDirectorTrack::GetAllSections() const
+{
+	TArray<UMovieSceneSection*> OutSectionPackets;
+	for ( int32 SectionIndex = 0; SectionIndex < ShotSections.Num(); ++SectionIndex )
+	{
+		OutSectionPackets.Add( ShotSections[SectionIndex] );
+	}
+	return OutSectionPackets;
+}
+
+void UMovieSceneDirectorTrack::RemoveAllAnimationData()
+{
+	ShotSections.Empty();
+}
+
+bool UMovieSceneDirectorTrack::HasSection( UMovieSceneSection* Section ) const
+{
+	return ShotSections.Find(Section) != INDEX_NONE;
+}
+
+void UMovieSceneDirectorTrack::RemoveSection( UMovieSceneSection* Section )
+{
+	ShotSections.Remove( Section );
+}
+
+bool UMovieSceneDirectorTrack::IsEmpty() const
+{
+	return ShotSections.Num() == 0;
+}
+
+void UMovieSceneDirectorTrack::AddNewShot(FGuid CameraHandle, float Time)
+{
+	UMovieSceneShotSection* NewSection = NewObject<UMovieSceneShotSection>(this);
+	NewSection->InitialPlacement(ShotSections, Time, Time + 1.f, SupportsMultipleRows()); //@todo find a better default end time
+	NewSection->SetCameraGuid(CameraHandle);
+	NewSection->SetTitle(LOCTEXT("AddNewShot", "New Shot"));
+
+	ShotSections.Add(NewSection);
+}
+
+TRange<float> UMovieSceneDirectorTrack::GetSectionBoundaries() const
+{
+	TArray< TRange<float> > Bounds;
+	for (int32 i = 0; i < ShotSections.Num(); ++i)
+	{
+		Bounds.Add(ShotSections[i]->GetRange());
+	}
+	return TRange<float>::Hull(Bounds);
+}
+
+#undef LOCTEXT_NAMESPACE
\ No newline at end of file
diff --git a/Engine/Source/Runtime/MovieSceneTracks/Private/Tracks/MovieSceneFloatTrack.cpp b/Engine/Source/Runtime/MovieSceneTracks/Private/Tracks/MovieSceneFloatTrack.cpp
new file mode 100644
--- /dev/null
+++ b/Engine/Source/Runtime/MovieSceneTracks/Private/Tracks/MovieSceneFloatTrack.cpp
@@ -0,0 +1,52 @@
+// Copyright 1998-2015 Epic Games, Inc. All Rights Reserved.
+
+#include "MovieSceneTracksPrivatePCH.h"
+#include "MovieSceneFloatSection.h"
+#include "MovieSceneFloatTrack.h"
+#include "IMovieScenePlayer.h"
+#include "MovieSceneFloatTrackInstance.h"
+
+
+UMovieSceneFloatTrack::UMovieSceneFloatTrack( const FObjectInitializer& ObjectInitializer )
+	: Super( ObjectInitializer )
+{
+}
+
+UMovieSceneSection* UMovieSceneFloatTrack::CreateNewSection()
+{
+	return NewObject<UMovieSceneSection>(this, UMovieSceneFloatSection::StaticClass(), NAME_None, RF_Transactional);
+}
+
+TSharedPtr<IMovieSceneTrackInstance> UMovieSceneFloatTrack::CreateInstance()
+{
+	return MakeShareable( new FMovieSceneFloatTrackInstance( *this ) ); 
+}
+
+
+bool UMovieSceneFloatTrack::AddKeyToSection( float Time, float Value )
+{
+	const UMovieSceneSection* NearestSection = MovieSceneHelpers::FindNearestSectionAtTime( Sections, Time );
+	if (!NearestSection || CastChecked<UMovieSceneFloatSection>(NearestSection)->NewKeyIsNewData(Time, Value))
+	{
+		Modify();
+
+		UMovieSceneFloatSection* NewSection = Cast<UMovieSceneFloatSection>( FindOrAddSection( Time ) );
+
+		NewSection->AddKey( Time, Value );
+
+		return true;
+	}
+	return false;
+}
+
+bool UMovieSceneFloatTrack::Eval( float Position, float LastPosition, float& OutFloat ) const
+{
+	const UMovieSceneSection* Section = MovieSceneHelpers::FindSectionAtTime( Sections, Position );
+
+	if( Section )
+	{
+		OutFloat = CastChecked<UMovieSceneFloatSection>( Section )->Eval( Position );
+	}
+
+	return Section != NULL;
+}
diff --git a/Engine/Source/Runtime/MovieSceneTracks/Private/Tracks/MovieSceneParticleTrack.cpp b/Engine/Source/Runtime/MovieSceneTracks/Private/Tracks/MovieSceneParticleTrack.cpp
new file mode 100644
--- /dev/null
+++ b/Engine/Source/Runtime/MovieSceneTracks/Private/Tracks/MovieSceneParticleTrack.cpp
@@ -0,0 +1,79 @@
+// Copyright 1998-2015 Epic Games, Inc. All Rights Reserved.
+
+#include "MovieSceneTracksPrivatePCH.h"
+#include "MovieSceneParticleSection.h"
+#include "MovieSceneParticleTrack.h"
+#include "IMovieScenePlayer.h"
+#include "MovieSceneParticleTrackInstance.h"
+
+#define LOCTEXT_NAMESPACE "MovieSceneParticleTrack"
+
+UMovieSceneParticleTrack::UMovieSceneParticleTrack( const FObjectInitializer& ObjectInitializer )
+	: Super( ObjectInitializer )
+{
+}
+
+FName UMovieSceneParticleTrack::GetTrackName() const
+{
+	return FName("ParticleSystem");
+}
+
+
+TSharedPtr<IMovieSceneTrackInstance> UMovieSceneParticleTrack::CreateInstance()
+{
+	return MakeShareable( new FMovieSceneParticleTrackInstance( *this ) ); 
+}
+
+TArray<UMovieSceneSection*> UMovieSceneParticleTrack::GetAllSections() const
+{
+	TArray<UMovieSceneSection*> OutSections;
+	for (int32 SectionIndex = 0; SectionIndex < ParticleSections.Num(); ++SectionIndex)
+	{
+		OutSections.Add( ParticleSections[SectionIndex] );
+	}
+	return OutSections;
+}
+
+void UMovieSceneParticleTrack::RemoveAllAnimationData()
+{
+}
+
+bool UMovieSceneParticleTrack::HasSection( UMovieSceneSection* Section ) const
+{
+	return ParticleSections.Find( Section ) != INDEX_NONE;
+}
+
+void UMovieSceneParticleTrack::RemoveSection( UMovieSceneSection* Section )
+{
+	ParticleSections.Remove( Section );
+}
+
+bool UMovieSceneParticleTrack::IsEmpty() const
+{
+	return ParticleSections.Num() == 0;
+}
+
+TRange<float> UMovieSceneParticleTrack::GetSectionBoundaries() const
+{
+	TArray< TRange<float> > Bounds;
+	for (int32 i = 0; i < ParticleSections.Num(); ++i)
+	{
+		Bounds.Add(ParticleSections[i]->GetRange());
+	}
+	return TRange<float>::Hull(Bounds);
+}
+
+void UMovieSceneParticleTrack::AddNewParticleSystem(float KeyTime, bool bTrigger)
+{
+	EParticleKey::Type KeyType = bTrigger ? EParticleKey::Trigger : EParticleKey::Toggle;
+	// @todo Instead of a 0.1 second event, this should be 0 seconds, requires handling 0 size sections
+	float Duration = KeyType == EParticleKey::Trigger ? 0.1f : 1.f;
+
+	UMovieSceneParticleSection* NewSection = NewObject<UMovieSceneParticleSection>(this);
+	NewSection->InitialPlacement(ParticleSections, KeyTime, KeyTime + Duration, SupportsMultipleRows());
+	NewSection->SetKeyType(KeyType);
+
+	ParticleSections.Add(NewSection);
+}
+
+#undef LOCTEXT_NAMESPACE
diff --git a/Engine/Source/Runtime/MovieSceneTracks/Private/Tracks/MovieSceneVectorTrack.cpp b/Engine/Source/Runtime/MovieSceneTracks/Private/Tracks/MovieSceneVectorTrack.cpp
new file mode 100644
--- /dev/null
+++ b/Engine/Source/Runtime/MovieSceneTracks/Private/Tracks/MovieSceneVectorTrack.cpp
@@ -0,0 +1,72 @@
+// Copyright 1998-2015 Epic Games, Inc. All Rights Reserved.
+
+#include "MovieSceneTracksPrivatePCH.h"
+#include "MovieSceneVectorSection.h"
+#include "MovieSceneVectorTrack.h"
+#include "IMovieScenePlayer.h"
+#include "MovieSceneVectorTrackInstance.h"
+
+UMovieSceneVectorTrack::UMovieSceneVectorTrack( const FObjectInitializer& ObjectInitializer )
+	: Super( ObjectInitializer )
+{
+	NumChannelsUsed = 0;
+}
+
+UMovieSceneSection* UMovieSceneVectorTrack::CreateNewSection()
+{
+	return NewObject<UMovieSceneSection>(this, UMovieSceneVectorSection::StaticClass());
+}
+
+TSharedPtr<IMovieSceneTrackInstance> UMovieSceneVectorTrack::CreateInstance()
+{
+	return MakeShareable( new FMovieSceneVectorTrackInstance( *this ) );
+}
+
+bool UMovieSceneVectorTrack::AddKeyToSection( float Time, const FVector4& Value, int32 InChannelsUsed, FName CurveName, bool bAddKeyEvenIfUnchanged )
+{
+	const UMovieSceneSection* NearestSection = MovieSceneHelpers::FindNearestSectionAtTime( Sections, Time );
+	if (!NearestSection || bAddKeyEvenIfUnchanged || CastChecked<UMovieSceneVectorSection>(NearestSection)->NewKeyIsNewData(Time, Value))
+	{
+		Modify();
+
+		UMovieSceneVectorSection* NewSection = Cast<UMovieSceneVectorSection>( FindOrAddSection( Time ) );
+		// @todo Sequencer - I don't like setting the channels used here. It should only be checked here, and set on section creation
+		NewSection->SetChannelsUsed(InChannelsUsed);
+
+		NewSection->AddKey( Time, CurveName, Value );
+
+		// We dont support one track containing multiple sections of differing channel amounts
+		NumChannelsUsed = InChannelsUsed;
+
+		return true;
+	}
+	return false;
+}
+
+bool UMovieSceneVectorTrack::AddKeyToSection( float Time, const FVectorKey<FVector4>& Key )
+{
+	return AddKeyToSection(Time, Key.Value, 4, Key.CurveName, Key.bAddKeyEvenIfUnchanged );
+}
+
+bool UMovieSceneVectorTrack::AddKeyToSection( float Time, const FVectorKey<FVector>& Key )
+{
+	return AddKeyToSection(Time, FVector4(Key.Value.X, Key.Value.Y, Key.Value.Z, 0.f), 3, Key.CurveName, Key.bAddKeyEvenIfUnchanged );
+}
+
+bool UMovieSceneVectorTrack::AddKeyToSection( float Time, const FVectorKey<FVector2D>& Key )
+{
+	return AddKeyToSection(Time, FVector4(Key.Value.X, Key.Value.Y, 0.f, 0.f), 2, Key.CurveName, Key.bAddKeyEvenIfUnchanged );
+}
+
+bool UMovieSceneVectorTrack::Eval( float Position, float LastPosition, FVector4& InOutVector ) const
+{
+	const UMovieSceneSection* Section = MovieSceneHelpers::FindSectionAtTime( Sections, Position );
+
+	if( Section )
+	{
+		InOutVector = CastChecked<UMovieSceneVectorSection>( Section )->Eval( Position, InOutVector );
+	}
+
+	return Section != NULL;
+}
+
diff --git a/Engine/Source/Runtime/MovieSceneTracks/Private/Tracks/SubMovieSceneTrack.cpp b/Engine/Source/Runtime/MovieSceneTracks/Private/Tracks/SubMovieSceneTrack.cpp
new file mode 100644
--- /dev/null
+++ b/Engine/Source/Runtime/MovieSceneTracks/Private/Tracks/SubMovieSceneTrack.cpp
@@ -0,0 +1,79 @@
+// Copyright 1998-2015 Epic Games, Inc. All Rights Reserved.
+
+#include "MovieSceneTracksPrivatePCH.h"
+#include "SubMovieSceneTrack.h"
+#include "IMovieScenePlayer.h"
+#include "SubMovieSceneSection.h"
+#include "SubMovieSceneTrackInstance.h"
+
+USubMovieSceneTrack::USubMovieSceneTrack( const FObjectInitializer& ObjectInitializer )
+	: Super( ObjectInitializer )
+{
+}
+
+UMovieSceneSection* USubMovieSceneTrack::CreateNewSection()
+{
+	return NewObject<USubMovieSceneSection>(this, NAME_None, RF_Transactional);
+}
+
+FName USubMovieSceneTrack::GetTrackName() const
+{
+	static const FName UniqueName( "Scenes" );
+
+	return UniqueName;
+}
+
+TSharedPtr<IMovieSceneTrackInstance> USubMovieSceneTrack::CreateInstance()
+{
+	return MakeShareable( new FSubMovieSceneTrackInstance( *this ) ); 
+}
+
+TArray<UMovieSceneSection*> USubMovieSceneTrack::GetAllSections() const 
+{
+	TArray<UMovieSceneSection*> AllSections;
+	for( int32 Section = 0; Section < SubMovieSceneSections.Num(); ++Section )
+	{
+		AllSections.Add( SubMovieSceneSections[Section] );
+	}
+
+	return AllSections;
+}
+
+void USubMovieSceneTrack::RemoveSection( UMovieSceneSection* Section )
+{
+	SubMovieSceneSections.Remove( Section );
+}
+
+bool USubMovieSceneTrack::IsEmpty() const
+{
+	return SubMovieSceneSections.Num() == 0;
+}
+
+TRange<float> USubMovieSceneTrack::GetSectionBoundaries() const
+{
+	TArray< TRange<float> > Bounds;
+	for (int32 SectionIndex = 0; SectionIndex < SubMovieSceneSections.Num(); ++SectionIndex)
+	{
+		Bounds.Add(SubMovieSceneSections[SectionIndex]->GetRange());
+	}
+	return TRange<float>::Hull(Bounds);
+}
+
+bool USubMovieSceneTrack::HasSection( UMovieSceneSection* Section ) const
+{
+	return SubMovieSceneSections.Contains( Section );
+}
+
+void USubMovieSceneTrack::AddMovieSceneSection( UMovieScene* SubMovieScene, float Time )
+{
+	Modify();
+
+	USubMovieSceneSection* NewSection = CastChecked<USubMovieSceneSection>( CreateNewSection() );
+
+	NewSection->SetMovieScene( SubMovieScene );
+	NewSection->SetStartTime( Time );
+	NewSection->SetEndTime( Time + SubMovieScene->GetTimeRange().Size<float>() );
+
+	SubMovieSceneSections.Add( NewSection );
+
+}
diff --git a/Engine/Source/Runtime/MovieSceneTracks/Public/IMovieSceneTracksModule.h b/Engine/Source/Runtime/MovieSceneTracks/Public/IMovieSceneTracksModule.h
new file mode 100644
--- /dev/null
+++ b/Engine/Source/Runtime/MovieSceneTracks/Public/IMovieSceneTracksModule.h
@@ -0,0 +1,36 @@
+// Copyright 1998-2015 Epic Games, Inc. All Rights Reserved.
+
+#pragma once
+
+#include "ModuleInterface.h"
+#include "ModuleManager.h"		// For inline LoadModuleChecked()
+
+
+/**
+ * The public interface of the MovieSceneTracks module
+ */
+class IMovieSceneTracksModule
+	: public IModuleInterface
+{
+public:
+
+	/**
+	 * Singleton-like access to IMovieSceneTracksModule
+	 *
+	 * @return Returns MovieSceneTracksModule singleton instance, loading the module on demand if needed
+	 */
+	static inline IMovieSceneTracksModule& Get()
+	{
+		return FModuleManager::LoadModuleChecked<IMovieSceneTracksModule>("MovieSceneTracks");
+	}
+
+	/**
+	 * Checks to see if this module is loaded and ready.  It is only valid to call Get() if IsAvailable() returns true.
+	 *
+	 * @return True if the module is loaded and ready to use
+	 */
+	static inline bool IsAvailable()
+	{
+		return FModuleManager::Get().IsModuleLoaded("MovieSceneTracks");
+	}
+};
diff --git a/Engine/Source/Runtime/MovieSceneTracks/Public/Sections/MovieScene3DTransformSection.h b/Engine/Source/Runtime/MovieSceneTracks/Public/Sections/MovieScene3DTransformSection.h
new file mode 100644
--- /dev/null
+++ b/Engine/Source/Runtime/MovieSceneTracks/Public/Sections/MovieScene3DTransformSection.h
@@ -0,0 +1,117 @@
+// Copyright 1998-2015 Epic Games, Inc. All Rights Reserved.
+
+#pragma once
+
+#include "MovieSceneSection.h"
+#include "MovieScene3DTransformSection.generated.h"
+
+
+/**
+ * A 3D transform section
+ */
+UCLASS(MinimalAPI)
+class UMovieScene3DTransformSection
+	: public UMovieSceneSection
+{
+	GENERATED_UCLASS_BODY()
+
+public:
+
+	/** MovieSceneSection interface */
+	virtual void MoveSection( float DeltaPosition, TSet<FKeyHandle>& KeyHandles ) override;
+	virtual void DilateSection( float DilationFactor, float Origin, TSet<FKeyHandle>& KeyHandles ) override;
+	virtual void GetKeyHandles(TSet<FKeyHandle>& KeyHandles) const override;
+
+	/**
+	 * Evaluates the translation component of the transform
+	 *
+	 * @param Time				The position in time within the movie scene
+	 * @param OutTranslation	The evaluated translation.  Note: will remain unchanged if there were no keys to evaluate
+	 * @return true if there were any keys to evaluate.  
+	 */
+	bool EvalTranslation( float Time, FVector& OutTranslation ) const;
+
+	/**
+	 * Evaluates the rotation component of the transform
+	 *
+	 * @param Time				The position in time within the movie scene
+	 * @param OutRotation		The evaluated rotation.  Note: will remain unchanged if there were no keys to evaluate
+	 * @return true if there were any keys to evaluate.  
+	 */
+	bool EvalRotation( float Time, FRotator& OutRotation ) const;
+
+	/**
+	 * Evaluates the scale component of the transform
+	 *
+	 * @param Time				The position in time within the movie scene
+	 * @param OutScale			The evaluated scale.  Note: will remain unchanged if there were no keys to evaluate
+	 * @return true if there were any keys to evaluate.  
+	 */
+	bool EvalScale( float Time, FVector& OutScale ) const;
+
+	/** 
+	 * Returns the translation curve for a specific axis
+	 *
+	 * @param Axis	The axis of the translation curve to get
+	 * @return The curve on the axis
+	 */
+	MOVIESCENETRACKS_API FRichCurve& GetTranslationCurve( EAxis::Type Axis );
+
+	/** 
+	 * Returns the rotation curve for a specific axis
+	 *
+	 * @param Axis	The axis of the rotation curve to get
+	 * @return The curve on the axis
+	 */
+	MOVIESCENETRACKS_API FRichCurve& GetRotationCurve( EAxis::Type Axis );
+
+	/** 
+	 * Returns the scale curve for a specific axis
+	 *
+	 * @param Axis	The axis of the scale curve to get
+	 * @return The curve on the axis
+	 */
+	MOVIESCENETRACKS_API FRichCurve& GetScaleCurve( EAxis::Type Axis );
+
+	/**
+	 * Adds keys each component of translation. Note: Only adds keys if a value has changed
+	 *
+	 * @param TransformKey	Information about the transform that should be keyed
+	 */
+	void AddTranslationKeys( const class FTransformKey& TransformKey );
+
+	/**
+	 * Adds keys each component of rotation. Note: Only adds keys if a value has changed
+	 *
+	 * @param TransformKey	Information about the transform that should be keyed
+	 * @param bUnwindRotation	When true, the value will be treated like a rotation value in degrees, and will automatically be unwound to prevent flipping 360 degrees from the previous key 	 
+	 */
+	void AddRotationKeys( const class FTransformKey& TransformKey, const bool bUnwindRotation );
+
+	/**
+	 * Adds keys each component of scale. Note: Only adds keys if a value has changed
+	 *
+	 * @param TransformKey	Information about the transform that should be keyed
+	 */
+	void AddScaleKeys( const class FTransformKey& TransformKey );
+	
+	/** 
+	 * Determines if a new key would be new data, or just a duplicate of existing data
+	 *
+	 * @param TransformKey The new key that would be added
+	 * @return True if the new key would be new data, false if duplicate
+	 */
+	bool NewKeyIsNewData(const class FTransformKey& TransformKey) const;
+private:
+	/** Translation curves */
+	UPROPERTY()
+	FRichCurve Translation[3];
+	
+	/** Rotation curves */
+	UPROPERTY()
+	FRichCurve Rotation[3];
+
+	/** Scale curves */
+	UPROPERTY()
+	FRichCurve Scale[3];
+};
\ No newline at end of file
diff --git a/Engine/Source/Runtime/MovieSceneTracks/Public/Sections/MovieSceneAnimationSection.h b/Engine/Source/Runtime/MovieSceneTracks/Public/Sections/MovieSceneAnimationSection.h
new file mode 100644
--- /dev/null
+++ b/Engine/Source/Runtime/MovieSceneTracks/Public/Sections/MovieSceneAnimationSection.h
@@ -0,0 +1,54 @@
+// Copyright 1998-2015 Epic Games, Inc. All Rights Reserved.
+
+#pragma once
+
+#include "MovieSceneSection.h"
+#include "MovieSceneAnimationSection.generated.h"
+
+/**
+ * Audio section, for use in the master audio, or by attached audio objects
+ */
+UCLASS( MinimalAPI )
+class UMovieSceneAnimationSection : public UMovieSceneSection
+{
+	GENERATED_UCLASS_BODY()
+public:
+	/** Sets the animation sequence for this section */
+	void SetAnimSequence(class UAnimSequence* InAnimSequence) {AnimSequence = InAnimSequence;}
+	
+	/** Gets the animation sequence for this section */
+	class UAnimSequence* GetAnimSequence() {return AnimSequence;}
+	
+	/** Sets the time that the animation is supposed to be played at */
+	void SetAnimationStartTime(float InAnimationStartTime) {AnimationStartTime = InAnimationStartTime;}
+	
+	/** Gets the (absolute) time that the animation is supposed to be played at */
+	float GetAnimationStartTime() const {return AnimationStartTime;}
+	
+	/** Gets the animation duration, modified by dilation */
+	float GetAnimationDuration() const {return AnimSequence->SequenceLength * AnimationDilationFactor;}
+
+	/**
+	 * @return The dilation multiplier of the animation
+	 */
+	float GetAnimationDilationFactor() const {return AnimationDilationFactor;}
+
+	/** MovieSceneSection interface */
+	virtual void MoveSection( float DeltaPosition, TSet<FKeyHandle>& KeyHandles ) override;
+	virtual void DilateSection( float DilationFactor, float Origin, TSet<FKeyHandle>& KeyHandles  ) override;
+	virtual void GetKeyHandles(TSet<FKeyHandle>& KeyHandles) const override;
+	virtual void GetSnapTimes(TArray<float>& OutSnapTimes, bool bGetSectionBorders) const override;
+
+private:
+	/** The animation sequence this section has */
+	UPROPERTY()
+	class UAnimSequence* AnimSequence;
+
+	/** The absolute time that the animation starts playing at */
+	UPROPERTY()
+	float AnimationStartTime;
+	
+	/** The amount which this animation is time dilated by */
+	UPROPERTY()
+	float AnimationDilationFactor;
+};
diff --git a/Engine/Source/Runtime/MovieSceneTracks/Public/Sections/MovieSceneAudioSection.h b/Engine/Source/Runtime/MovieSceneTracks/Public/Sections/MovieSceneAudioSection.h
new file mode 100644
--- /dev/null
+++ b/Engine/Source/Runtime/MovieSceneTracks/Public/Sections/MovieSceneAudioSection.h
@@ -0,0 +1,77 @@
+// Copyright 1998-2015 Epic Games, Inc. All Rights Reserved.
+
+#pragma once
+
+#include "MovieSceneSection.h"
+#include "MovieSceneAudioSection.generated.h"
+
+
+/**
+ * Audio section, for use in the master audio, or by attached audio objects
+ */
+UCLASS( MinimalAPI )
+class UMovieSceneAudioSection
+	: public UMovieSceneSection
+{
+	GENERATED_UCLASS_BODY()
+
+public:
+
+	/** Sets this section's sound */
+	void SetSound(class USoundBase* InSound) {Sound = InSound;}
+
+	/** Gets the sound for this section */
+	class USoundBase* GetSound() {return Sound;}
+	
+	/** Sets the time that the sound is supposed to be played at */
+	void SetAudioStartTime(float InAudioStartTime) {AudioStartTime = InAudioStartTime;}
+	
+	/** Gets the (absolute) time that the sound is supposed to be played at */
+	float GetAudioStartTime() {return AudioStartTime;}
+	
+	/**
+	 * @return The range of times that the sound plays, truncated by the section limits
+	 */
+	TRange<float> MOVIESCENETRACKS_API GetAudioRange() const;
+	
+	/**
+	 * @return The range of times that the sound plays, not truncated by the section limits
+	 */
+	TRange<float> MOVIESCENETRACKS_API GetAudioTrueRange() const;
+	
+	/**
+	 * @return The dilation multiplier of the sound
+	 */
+	float GetAudioDilationFactor() const {return AudioDilationFactor;}
+
+	/**
+	 * Returns whether or not a provided position in time is within the timespan of the audio range
+	 *
+	 * @param Position	The position to check
+	 * @return true if the position is within the timespan, false otherwise
+	 */
+	bool IsTimeWithinAudioRange( float Position ) const 
+	{
+		TRange<float> AudioRange = GetAudioRange();
+		return Position >= AudioRange.GetLowerBoundValue() && Position <= AudioRange.GetUpperBoundValue();
+	}
+
+	/** MovieSceneSection interface */
+	virtual void MoveSection( float DeltaPosition, TSet<FKeyHandle>& KeyHandles ) override;
+	virtual void DilateSection( float DilationFactor, float Origin, TSet<FKeyHandle>& KeyHandles ) override;
+	virtual void GetKeyHandles(TSet<FKeyHandle>& KeyHandles) const override;
+	virtual void GetSnapTimes(TArray<float>& OutSnapTimes, bool bGetSectionBorders) const override;
+
+private:
+	/** The sound cue or wave that this section plays */
+	UPROPERTY()
+	USoundBase* Sound;
+
+	/** The absolute time that the sound starts playing at */
+	UPROPERTY()
+	float AudioStartTime;
+	
+	/** The amount which this audio is time dilated by */
+	UPROPERTY()
+	float AudioDilationFactor;
+};
diff --git a/Engine/Source/Runtime/MovieSceneTracks/Public/Sections/MovieSceneBoolSection.h b/Engine/Source/Runtime/MovieSceneTracks/Public/Sections/MovieSceneBoolSection.h
new file mode 100644
--- /dev/null
+++ b/Engine/Source/Runtime/MovieSceneTracks/Public/Sections/MovieSceneBoolSection.h
@@ -0,0 +1,55 @@
+// Copyright 1998-2015 Epic Games, Inc. All Rights Reserved.
+
+#pragma once
+
+#include "MovieSceneBoolSection.generated.h"
+
+/**
+ * A single bool section
+ */
+UCLASS( MinimalAPI )
+class UMovieSceneBoolSection : public UMovieSceneSection
+{
+	GENERATED_UCLASS_BODY()
+public:
+	/**
+	 * Updates this section
+	 *
+	 * @param Position	The position in time within the movie scene
+	 */
+	virtual bool Eval( float Position ) const;
+
+	/** 
+	 * Adds a key to the section
+	 *
+	 * @param Time	The location in time where the key should be added
+	 * @param Value	The value of the key
+	 */
+	void AddKey( float Time, bool Value );
+	
+	/** 
+	 * Determines if a new key would be new data, or just a duplicate of existing data
+	 *
+	 * @param Time	The location in time where the key would be added
+	 * @param Value	The value of the new key
+	 * @return True if the new key would be new data, false if duplicate
+	 */
+	bool NewKeyIsNewData(float Time, bool Value) const;
+
+	/**
+	 * UMovieSceneSection interface 
+	 */
+	virtual void MoveSection(float DeltaPosition, TSet<FKeyHandle>& KeyHandles) override;
+	virtual void DilateSection(float DilationFactor, float Origin, TSet<FKeyHandle>& KeyHandles) override;
+	virtual void GetKeyHandles(TSet<FKeyHandle>& KeyHandles) const override;
+
+	/** Gets all the keys of this boolean section */
+	FIntegralCurve& GetCurve() { return BoolCurve; }
+
+private:
+	/** Ordered curve data */
+	// @todo Sequencer This could be optimized by packing the bools separately
+	// but that may not be worth the effort
+	UPROPERTY()
+	FIntegralCurve BoolCurve;
+};
diff --git a/Engine/Source/Runtime/MovieSceneTracks/Public/Sections/MovieSceneByteSection.h b/Engine/Source/Runtime/MovieSceneTracks/Public/Sections/MovieSceneByteSection.h
new file mode 100644
--- /dev/null
+++ b/Engine/Source/Runtime/MovieSceneTracks/Public/Sections/MovieSceneByteSection.h
@@ -0,0 +1,55 @@
+// Copyright 1998-2015 Epic Games, Inc. All Rights Reserved.
+
+#pragma once
+
+#include "MovieSceneByteSection.generated.h"
+
+/**
+ * A single byte section
+ */
+UCLASS( MinimalAPI )
+class UMovieSceneByteSection : public UMovieSceneSection
+{
+	GENERATED_UCLASS_BODY()
+public:
+	/**
+	 * Updates this section
+	 *
+	 * @param Position	The position in time within the movie scene
+	 */
+	virtual uint8 Eval( float Position ) const;
+
+	/** 
+	 * Adds a key to the section
+	 *
+	 * @param Time	The location in time where the key should be added
+	 * @param Value	The value of the key
+	 */
+	void AddKey( float Time, uint8 Value );
+	
+	/** 
+	 * Determines if a new key would be new data, or just a duplicate of existing data
+	 *
+	 * @param Time	The location in time where the key would be added
+	 * @param Value	The value of the new key
+	 * @return True if the new key would be new data, false if duplicate
+	 */
+	bool NewKeyIsNewData(float Time, uint8 Value) const;
+
+	/**
+	 * UMovieSceneSection interface 
+	 */
+	virtual void MoveSection(float DeltaPosition, TSet<FKeyHandle>& KeyHandles) override;
+	virtual void DilateSection(float DilationFactor, float Origin, TSet<FKeyHandle>& KeyHandles) override;
+	virtual void GetKeyHandles(TSet<FKeyHandle>& KeyHandles) const override;
+
+	/** Gets all the keys of this byte section */
+	FIntegralCurve& GetCurve() { return ByteCurve; }
+
+private:
+	/** Ordered curve data */
+	// @todo Sequencer This could be optimized by packing the bytes separately
+	// but that may not be worth the effort
+	UPROPERTY()
+	FIntegralCurve ByteCurve;
+};
diff --git a/Engine/Source/Runtime/MovieSceneTracks/Public/Sections/MovieSceneColorSection.h b/Engine/Source/Runtime/MovieSceneTracks/Public/Sections/MovieSceneColorSection.h
new file mode 100644
--- /dev/null
+++ b/Engine/Source/Runtime/MovieSceneTracks/Public/Sections/MovieSceneColorSection.h
@@ -0,0 +1,96 @@
+// Copyright 1998-2015 Epic Games, Inc. All Rights Reserved.
+
+#pragma once
+
+#include "MovieSceneColorSection.generated.h"
+
+struct FColorKey;
+
+/**
+ * A single floating point section
+ */
+UCLASS( MinimalAPI )
+class UMovieSceneColorSection : public UMovieSceneSection
+{
+	GENERATED_UCLASS_BODY()
+public:
+	/** MovieSceneSection interface */
+	virtual void MoveSection(float DeltaPosition, TSet<FKeyHandle>& KeyHandles) override;
+	virtual void DilateSection(float DilationFactor, float Origin, TSet<FKeyHandle>& KeyHandles) override;
+	virtual void GetKeyHandles(TSet<FKeyHandle>& KeyHandles) const override;
+
+	/**
+	 * Updates this section
+	 *
+	 * @param Position	The position in time within the movie scene
+	 */
+	virtual FLinearColor Eval( float Position, const FLinearColor& DefaultColor ) const;
+
+	/** 
+	 * Adds a key to the section
+	 *
+	 * @param Time	The location in time where the key should be added
+	 * @param Value	The value of the key
+	 */
+	void AddKey( float Time, const FColorKey& Key );
+	
+	/** 
+	 * Determines if a new key would be new data, or just a duplicate of existing data
+	 *
+	 * @param Time	The location in time where the key would be added
+	 * @param Value	The value of the new key
+	 * @return True if the new key would be new data, false if duplicate
+	 */
+	bool NewKeyIsNewData(float Time, FLinearColor Value) const;
+
+	/**
+	 * Gets the red color curve
+	 *
+	 * @return The rich curve for this color channel
+	 */
+	FRichCurve& GetRedCurve() { return RedCurve; }
+	const FRichCurve& GetRedCurve() const { return RedCurve; }
+
+	/**
+	 * Gets the green color curve
+	 *
+	 * @return The rich curve for this color channel
+	 */
+	FRichCurve& GetGreenCurve() { return GreenCurve; }
+	const FRichCurve& GetGreenCurve() const { return GreenCurve; }
+	/**
+	 * Gets the blue color curve
+	 *
+	 * @return The rich curve for this color channel
+	 */
+	FRichCurve& GetBlueCurve() { return BlueCurve; }
+	const FRichCurve& GetBlueCurve() const { return BlueCurve; }
+	
+	/**
+	 * Gets the alpha color curve
+	 *
+	 * @return The rich curve for this color channel
+	 */
+	FRichCurve& GetAlphaCurve() { return AlphaCurve; }
+	const FRichCurve& GetAlphaCurve() const { return AlphaCurve; }
+
+private:
+	void AddKeyToNamedCurve(float Time, const FColorKey& Key);
+
+private:
+	/** Red curve data */
+	UPROPERTY()
+	FRichCurve RedCurve;
+
+	/** Green curve data */
+	UPROPERTY()
+	FRichCurve GreenCurve;
+
+	/** Blue curve data */
+	UPROPERTY()
+	FRichCurve BlueCurve;
+
+	/** Alpha curve data */
+	UPROPERTY()
+	FRichCurve AlphaCurve;
+};
diff --git a/Engine/Source/Runtime/MovieSceneTracks/Public/Sections/MovieSceneFloatSection.h b/Engine/Source/Runtime/MovieSceneTracks/Public/Sections/MovieSceneFloatSection.h
new file mode 100644
--- /dev/null
+++ b/Engine/Source/Runtime/MovieSceneTracks/Public/Sections/MovieSceneFloatSection.h
@@ -0,0 +1,54 @@
+// Copyright 1998-2015 Epic Games, Inc. All Rights Reserved.
+
+#pragma once
+
+#include "MovieSceneFloatSection.generated.h"
+
+/**
+ * A single floating point section
+ */
+UCLASS( MinimalAPI )
+class UMovieSceneFloatSection : public UMovieSceneSection
+{
+	GENERATED_UCLASS_BODY()
+public:
+	/**
+	 * Updates this section
+	 *
+	 * @param Position	The position in time within the movie scene
+	 */
+	float Eval( float Position ) const;
+
+	/** 
+	 * Adds a key to the section
+	 *
+	 * @param Time	The location in time where the key should be added
+	 * @param Value	The value of the key
+	 */
+	void AddKey( float Time, float Value );
+	
+	/** 
+	 * Determines if a new key would be new data, or just a duplicate of existing data
+	 *
+	 * @param Time	The location in time where the key would be added
+	 * @param Value	The value of the new key
+	 * @return True if the new key would be new data, false if duplicate
+	 */
+	bool NewKeyIsNewData(float Time, float Value) const;
+
+	/**
+	 * UMovieSceneSection interface 
+	 */
+	virtual void MoveSection(float DeltaPosition, TSet<FKeyHandle>& KeyHandles) override;
+	virtual void DilateSection(float DilationFactor, float Origin, TSet<FKeyHandle>& KeyHandles) override;
+	virtual void GetKeyHandles(TSet<FKeyHandle>& KeyHandles) const override;
+
+	/**
+	 * @return The float curve on this section
+	 */
+	FRichCurve& GetFloatCurve() { return FloatCurve; }
+private:
+	/** Curve data */
+	UPROPERTY()
+	FRichCurve FloatCurve;
+};
\ No newline at end of file
diff --git a/Engine/Source/Runtime/MovieSceneTracks/Public/Sections/MovieSceneParticleSection.h b/Engine/Source/Runtime/MovieSceneTracks/Public/Sections/MovieSceneParticleSection.h
new file mode 100644
--- /dev/null
+++ b/Engine/Source/Runtime/MovieSceneTracks/Public/Sections/MovieSceneParticleSection.h
@@ -0,0 +1,39 @@
+// Copyright 1998-2015 Epic Games, Inc. All Rights Reserved.
+
+#pragma once
+
+#include "MovieSceneSection.h"
+#include "MovieSceneParticleSection.generated.h"
+
+/**
+ * The kind of key this particle section is
+ */
+UENUM()
+namespace EParticleKey
+{
+	enum Type
+	{
+		Toggle,
+		Trigger
+	};
+}
+
+/**
+ * Particle section, for particle toggling and triggering
+ */
+UCLASS( MinimalAPI )
+class UMovieSceneParticleSection : public UMovieSceneSection
+{
+	GENERATED_UCLASS_BODY()
+public:
+	/** Sets this section's key type*/
+	void SetKeyType(EParticleKey::Type InKeyType) {KeyType = InKeyType;}
+
+	/** Gets the key type for this section */
+	EParticleKey::Type GetKeyType() {return KeyType;}
+	
+private:
+	/** The way this particle key will operate */
+	UPROPERTY()
+	TEnumAsByte<EParticleKey::Type> KeyType;
+};
diff --git a/Engine/Source/Runtime/MovieSceneTracks/Public/Sections/MovieSceneShotSection.h b/Engine/Source/Runtime/MovieSceneTracks/Public/Sections/MovieSceneShotSection.h
new file mode 100644
--- /dev/null
+++ b/Engine/Source/Runtime/MovieSceneTracks/Public/Sections/MovieSceneShotSection.h
@@ -0,0 +1,36 @@
+// Copyright 1998-2015 Epic Games, Inc. All Rights Reserved.
+
+#pragma once
+
+#include "MovieSceneSection.h"
+#include "MovieSceneShotSection.generated.h"
+
+/**
+ * Movie shots are sections on the director's track, that show what the viewer "sees"
+ */
+UCLASS( MinimalAPI )
+class UMovieSceneShotSection : public UMovieSceneSection
+{
+	GENERATED_UCLASS_BODY()
+public:
+	/** Sets the camera guid for this shot section */
+	void SetCameraGuid(const FGuid& InGuid);
+
+	/** Gets the camera guid for this shot section */
+	virtual FGuid GetCameraGuid() const;
+
+	/** Sets the shot title */
+	virtual void SetTitle(const FText& NewTitle);
+
+	/** Gets the shot title */
+	virtual FText GetTitle() const;
+
+private:
+	/** The camera possessable or spawnable that this movie shot uses */
+	UPROPERTY()
+	FGuid CameraGuid;
+
+	/** The shot's title */
+	UPROPERTY()
+	FText Title;
+};
diff --git a/Engine/Source/Runtime/MovieSceneTracks/Public/Sections/MovieSceneVectorSection.h b/Engine/Source/Runtime/MovieSceneTracks/Public/Sections/MovieSceneVectorSection.h
new file mode 100644
--- /dev/null
+++ b/Engine/Source/Runtime/MovieSceneTracks/Public/Sections/MovieSceneVectorSection.h
@@ -0,0 +1,67 @@
+// Copyright 1998-2015 Epic Games, Inc. All Rights Reserved.
+
+#pragma once
+
+#include "MovieSceneSection.h"
+
+#include "MovieSceneVectorSection.generated.h"
+
+/**
+ * A vector section
+ */
+UCLASS(MinimalAPI )
+class UMovieSceneVectorSection : public UMovieSceneSection
+{
+	GENERATED_UCLASS_BODY()
+public:
+	
+	/**
+	 * Updates this section
+	 *
+	 * @param Position	The position in time within the movie scene
+	 */
+	FVector4 Eval( float Position, const FVector4& DefaultVector ) const;
+
+	/** MovieSceneSection interface */
+	virtual void MoveSection(float DeltaPosition, TSet<FKeyHandle>& KeyHandles) override;
+	virtual void DilateSection(float DilationFactor, float Origin, TSet<FKeyHandle>& KeyHandles) override;
+	virtual void GetKeyHandles(TSet<FKeyHandle>& KeyHandles) const override;
+
+	/** 
+	 * Adds a key to the section
+	 *
+	 * @param Time	The location in time where the key should be added
+	 * @param Value	The value of the key
+	 */
+	void AddKey( float Time, FName CurveName, const FVector4& Value );
+	
+	/** 
+	 * Determines if a new key would be new data, or just a duplicate of existing data
+	 *
+	 * @param Time	The location in time where the key would be added
+	 * @param Value	The value of the new key
+	 * @return True if the new key would be new data, false if duplicate
+	 */
+	bool NewKeyIsNewData(float Time, const FVector4& Value) const;
+
+	/** Gets one of four curves in this section */
+	FRichCurve& GetCurve(const int32& Index) { return Curves[Index]; }
+
+	/** Sets how many channels are to be used */
+	void SetChannelsUsed(int32 InChannelsUsed) {ChannelsUsed = InChannelsUsed;}
+	
+	/** Gets the number of channels in use */
+	int32 GetChannelsUsed() const {return ChannelsUsed;}
+
+private:
+	void AddKeyToNamedCurve(float Time, FName CurveName, const FVector4& Value);
+
+private:
+	/** Vector t */
+	UPROPERTY()
+	FRichCurve Curves[4];
+
+	/** How many curves are actually used */
+	UPROPERTY()
+	int32 ChannelsUsed;
+};
diff --git a/Engine/Source/Runtime/MovieSceneTracks/Public/Sections/SubMovieSceneSection.h b/Engine/Source/Runtime/MovieSceneTracks/Public/Sections/SubMovieSceneSection.h
new file mode 100644
--- /dev/null
+++ b/Engine/Source/Runtime/MovieSceneTracks/Public/Sections/SubMovieSceneSection.h
@@ -0,0 +1,36 @@
+// Copyright 1998-2015 Epic Games, Inc. All Rights Reserved.
+
+#pragma once
+
+#include "SubMovieSceneSection.generated.h"
+
+/**
+ * A container for a movie scene within a movie scene
+ */
+UCLASS( MinimalAPI )
+class USubMovieSceneSection : public UMovieSceneSection
+{
+	GENERATED_UCLASS_BODY()
+public:
+	/**
+	 * Sets the movie scene played by this section
+	 *
+	 * @param InMovieScene	The movie scene to play
+	 */
+	void SetMovieScene( UMovieScene* InMovieScene ) { MovieScene = InMovieScene; }
+
+	/** @return The movie scene played by this section */
+	UMovieScene* GetMovieScene() const { return MovieScene; }
+private:
+	/** Movie scene being played by this section */
+	UPROPERTY()
+	UMovieScene* MovieScene;
+
+	/** The start time where the movie scene should begin playing */
+	UPROPERTY()
+	float MovieSceneStartTime;
+
+	/** The end time where the movie scene should begin playing */
+	UPROPERTY()
+	float MovieSceneTimeEndTime;
+};
diff --git a/Engine/Source/Runtime/MovieSceneTracks/Public/TrackInstances/MovieScene3DTransformTrackInstance.h b/Engine/Source/Runtime/MovieSceneTracks/Public/TrackInstances/MovieScene3DTransformTrackInstance.h
new file mode 100644
--- /dev/null
+++ b/Engine/Source/Runtime/MovieSceneTracks/Public/TrackInstances/MovieScene3DTransformTrackInstance.h
@@ -0,0 +1,34 @@
+// Copyright 1998-2015 Epic Games, Inc. All Rights Reserved.
+
+#pragma once
+
+#include "IMovieSceneTrackInstance.h"
+
+
+class UMovieScene3DTransformTrack;
+
+
+/**
+ * Instance of a UMovieSceneTransformTrack
+ */
+class FMovieScene3DTransformTrackInstance
+	: public IMovieSceneTrackInstance
+{
+public:
+
+	FMovieScene3DTransformTrackInstance( UMovieScene3DTransformTrack& InTransformTrack );
+
+	/** IMovieSceneTrackInstance interface */
+	virtual void SaveState (const TArray<UObject*>& RuntimeObjects) override;
+	virtual void RestoreState (const TArray<UObject*>& RuntimeObjects) override;
+	virtual void Update( float Position, float LastPosition, const TArray<UObject*>& RuntimeObjects, IMovieScenePlayer& Player ) override;
+	virtual void RefreshInstance( const TArray<UObject*>& RuntimeObjects, IMovieScenePlayer& Player ) override {}
+
+private:
+
+	/** Track that is being instanced */
+	UMovieScene3DTransformTrack* TransformTrack;
+
+	/** Map from object to initial state */
+	TMap< TWeakObjectPtr<UObject>, FTransform > InitTransformMap;
+};
diff --git a/Engine/Source/Runtime/MovieSceneTracks/Public/TrackInstances/MovieSceneAnimationTrackInstance.h b/Engine/Source/Runtime/MovieSceneTracks/Public/TrackInstances/MovieSceneAnimationTrackInstance.h
new file mode 100644
--- /dev/null
+++ b/Engine/Source/Runtime/MovieSceneTracks/Public/TrackInstances/MovieSceneAnimationTrackInstance.h
@@ -0,0 +1,30 @@
+// Copyright 1998-2015 Epic Games, Inc. All Rights Reserved.
+
+#pragma once
+
+#include "IMovieSceneTrackInstance.h"
+
+
+class UMovieSceneAnimationTrack;
+
+
+/**
+ * Instance of a UMovieSceneAnimationTrack
+ */
+class FMovieSceneAnimationTrackInstance
+	: public IMovieSceneTrackInstance
+{
+public:
+
+	FMovieSceneAnimationTrackInstance( UMovieSceneAnimationTrack& InAnimationTrack );
+	virtual ~FMovieSceneAnimationTrackInstance();
+
+	/** IMovieSceneTrackInstance interface */
+	virtual void Update( float Position, float LastPosition, const TArray<UObject*>& RuntimeObjects, IMovieScenePlayer& Player ) override;
+	virtual void RefreshInstance( const TArray<UObject*>& RuntimeObjects, IMovieScenePlayer& Player ) override {}
+
+private:
+
+	/** Track that is being instanced */
+	UMovieSceneAnimationTrack* AnimationTrack;
+};
diff --git a/Engine/Source/Runtime/MovieSceneTracks/Public/TrackInstances/MovieSceneAudioTrackInstance.h b/Engine/Source/Runtime/MovieSceneTracks/Public/TrackInstances/MovieSceneAudioTrackInstance.h
new file mode 100644
--- /dev/null
+++ b/Engine/Source/Runtime/MovieSceneTracks/Public/TrackInstances/MovieSceneAudioTrackInstance.h
@@ -0,0 +1,50 @@
+// Copyright 1998-2015 Epic Games, Inc. All Rights Reserved.
+
+#pragma once
+
+#include "IMovieSceneTrackInstance.h"
+
+
+class UAudioComponent;
+class UMovieSceneAudioTrack;
+
+
+/**
+ * Instance of a UMovieSceneAudioTrack
+ */
+class FMovieSceneAudioTrackInstance
+	: public IMovieSceneTrackInstance
+{
+public:
+
+	FMovieSceneAudioTrackInstance( UMovieSceneAudioTrack& InAudioTrack );
+
+	/** IMovieSceneTrackInstance interface */
+	virtual void Update( float Position, float LastPosition, const TArray<UObject*>& RuntimeObjects, IMovieScenePlayer& Player ) override;
+	virtual void RefreshInstance( const TArray<UObject*>& RuntimeObjects, IMovieScenePlayer& Player ) override {}
+
+private:
+
+	/** Plays the sound of the given section at the given time */
+	void PlaySound(class UMovieSceneAudioSection* AudioSection, UAudioComponent* Component, float Time);
+	
+	/** Stops all components for the given row index */
+	void StopSound(int32 RowIndex);
+
+	/** Stops all components from playing */
+	void StopAllSounds();
+
+	/** Gets the audio component component for the actor and row index, creating it if necessary */
+	UAudioComponent* GetAudioComponent(AActor* Actor, int32 RowIndex);
+
+	/** Utility function for getting actors from objects array */
+	TArray<AActor*> GetRuntimeActors(const TArray<UObject*>& RuntimeObjects) const;
+
+private:
+
+	/** Track that is being instanced */
+	UMovieSceneAudioTrack* AudioTrack;
+
+	/** Audio components to play our audio tracks with, one per row per actor */
+	TArray< TMap<AActor*, UAudioComponent*> > PlaybackAudioComponents;
+};
diff --git a/Engine/Source/Runtime/MovieSceneTracks/Public/TrackInstances/MovieSceneBoolTrackInstance.h b/Engine/Source/Runtime/MovieSceneTracks/Public/TrackInstances/MovieSceneBoolTrackInstance.h
new file mode 100644
--- /dev/null
+++ b/Engine/Source/Runtime/MovieSceneTracks/Public/TrackInstances/MovieSceneBoolTrackInstance.h
@@ -0,0 +1,38 @@
+// Copyright 1998-2015 Epic Games, Inc. All Rights Reserved.
+
+#pragma once
+
+#include "IMovieSceneTrackInstance.h"
+
+
+class FTrackInstancePropertyBindings;
+class UMovieSceneBoolTrack;
+
+
+/**
+ * Instance of a UMovieSceneBoolTrack
+ */
+class FMovieSceneBoolTrackInstance
+	: public IMovieSceneTrackInstance
+{
+public:
+
+	FMovieSceneBoolTrackInstance( UMovieSceneBoolTrack& InBoolTrack );
+
+	/** IMovieSceneTrackInstance interface */
+	virtual void SaveState (const TArray<UObject*>& RuntimeObjects) override;
+	virtual void RestoreState (const TArray<UObject*>& RuntimeObjects) override;
+	virtual void Update( float Position, float LastPosition, const TArray<UObject*>& RuntimeObjects, IMovieScenePlayer& Player ) override;
+	virtual void RefreshInstance( const TArray<UObject*>& RuntimeObjects, IMovieScenePlayer& Player ) override;
+
+private:
+
+	/** Bool track that is being instanced */
+	UMovieSceneBoolTrack* BoolTrack;
+
+	/** Runtime property bindings */
+	TSharedPtr<FTrackInstancePropertyBindings> PropertyBindings;
+
+	/** Map from object to initial state */
+	TMap<TWeakObjectPtr<UObject>, bool> InitBoolMap;
+};
diff --git a/Engine/Source/Runtime/MovieSceneTracks/Public/TrackInstances/MovieSceneByteTrackInstance.h b/Engine/Source/Runtime/MovieSceneTracks/Public/TrackInstances/MovieSceneByteTrackInstance.h
new file mode 100644
--- /dev/null
+++ b/Engine/Source/Runtime/MovieSceneTracks/Public/TrackInstances/MovieSceneByteTrackInstance.h
@@ -0,0 +1,38 @@
+// Copyright 1998-2015 Epic Games, Inc. All Rights Reserved.
+
+#pragma once
+
+#include "IMovieSceneTrackInstance.h"
+
+
+class FTrackInstancePropertyBindings;
+class UMovieSceneByteTrack;
+
+
+/**
+ * Instance of a UMovieSceneByteTrack
+ */
+class FMovieSceneByteTrackInstance
+	: public IMovieSceneTrackInstance
+{
+public:
+
+	FMovieSceneByteTrackInstance( UMovieSceneByteTrack& InByteTrack );
+
+	/** IMovieSceneTrackInstance interface */
+	virtual void SaveState (const TArray<UObject*>& RuntimeObjects) override;
+	virtual void RestoreState (const TArray<UObject*>& RuntimeObjects) override;
+	virtual void Update( float Position, float LastPosition, const TArray<UObject*>& RuntimeObjects, IMovieScenePlayer& Player ) override;
+	virtual void RefreshInstance( const TArray<UObject*>& RuntimeObjects, IMovieScenePlayer& Player ) override;
+
+private:
+
+	/** Byte track that is being instanced */
+	UMovieSceneByteTrack* ByteTrack;
+
+	/** Runtime property bindings */
+	TSharedPtr<FTrackInstancePropertyBindings> PropertyBindings;
+
+	/** Map from object to initial state */
+	TMap< TWeakObjectPtr<UObject>, uint8 > InitByteMap;
+};
diff --git a/Engine/Source/Runtime/MovieSceneTracks/Public/TrackInstances/MovieSceneColorTrackInstance.h b/Engine/Source/Runtime/MovieSceneTracks/Public/TrackInstances/MovieSceneColorTrackInstance.h
new file mode 100644
--- /dev/null
+++ b/Engine/Source/Runtime/MovieSceneTracks/Public/TrackInstances/MovieSceneColorTrackInstance.h
@@ -0,0 +1,39 @@
+// Copyright 1998-2015 Epic Games, Inc. All Rights Reserved.
+
+#pragma once
+
+#include "IMovieSceneTrackInstance.h"
+
+
+class FTrackInstancePropertyBindings;
+class UMovieSceneColorTrack;
+
+
+/**
+ * Instance of a UMovieSceneColorTrack
+ */
+class FMovieSceneColorTrackInstance
+	: public IMovieSceneTrackInstance
+{
+public:
+
+	FMovieSceneColorTrackInstance( UMovieSceneColorTrack& InColorTrack );
+
+	/** IMovieSceneTrackInstance interface */
+	virtual void SaveState (const TArray<UObject*>& RuntimeObjects) override;
+	virtual void RestoreState (const TArray<UObject*>& RuntimeObjects) override;
+	virtual void Update( float Position, float LastPosition, const TArray<UObject*>& RuntimeObjects, IMovieScenePlayer& Player ) override;
+	virtual void RefreshInstance( const TArray<UObject*>& RuntimeObjects, IMovieScenePlayer& Player ) override;
+
+private:
+
+	/** The track being instanced */
+	UMovieSceneColorTrack* ColorTrack;
+
+	/** Runtime property bindings */
+	TSharedPtr<FTrackInstancePropertyBindings> PropertyBindings;
+
+	/** Map from object to initial state */
+	TMap< TWeakObjectPtr<UObject>, FSlateColor > InitSlateColorMap;
+	TMap< TWeakObjectPtr<UObject>, FLinearColor > InitLinearColorMap;
+};
diff --git a/Engine/Source/Runtime/MovieSceneTracks/Public/TrackInstances/MovieSceneDirectorTrackInstance.h b/Engine/Source/Runtime/MovieSceneTracks/Public/TrackInstances/MovieSceneDirectorTrackInstance.h
new file mode 100644
--- /dev/null
+++ b/Engine/Source/Runtime/MovieSceneTracks/Public/TrackInstances/MovieSceneDirectorTrackInstance.h
@@ -0,0 +1,29 @@
+// Copyright 1998-2015 Epic Games, Inc. All Rights Reserved.
+
+#pragma once
+
+#include "IMovieSceneTrackInstance.h"
+
+
+class UMovieSceneDirectorTrack;
+
+
+/**
+ * Instance of a UMovieSceneDirectorTrack
+ */
+class FMovieSceneDirectorTrackInstance
+	: public IMovieSceneTrackInstance
+{
+public:
+
+	FMovieSceneDirectorTrackInstance( UMovieSceneDirectorTrack& InDirectorTrack );
+
+	/** IMovieSceneTrackInstance interface */
+	virtual void Update( float Position, float LastPosition, const TArray<UObject*>& RuntimeObjects, IMovieScenePlayer& Player ) override;
+	virtual void RefreshInstance( const TArray<UObject*>& RuntimeObjects, IMovieScenePlayer& Player  ) override {}
+
+private:
+
+	/** Track that is being instanced */
+	UMovieSceneDirectorTrack* DirectorTrack;
+};
diff --git a/Engine/Source/Runtime/MovieSceneTracks/Public/TrackInstances/MovieSceneFloatTrackInstance.h b/Engine/Source/Runtime/MovieSceneTracks/Public/TrackInstances/MovieSceneFloatTrackInstance.h
new file mode 100644
--- /dev/null
+++ b/Engine/Source/Runtime/MovieSceneTracks/Public/TrackInstances/MovieSceneFloatTrackInstance.h
@@ -0,0 +1,38 @@
+// Copyright 1998-2015 Epic Games, Inc. All Rights Reserved.
+
+#pragma once
+
+#include "IMovieSceneTrackInstance.h"
+
+
+class FTrackInstancePropertyBindings;
+class UMovieSceneFloatTrack;
+
+
+/**
+ * Instance of a UMovieSceneFloatTrack
+ */
+class FMovieSceneFloatTrackInstance
+	: public IMovieSceneTrackInstance
+{
+public:
+
+	FMovieSceneFloatTrackInstance( UMovieSceneFloatTrack& InFloatTrack );
+
+	/** IMovieSceneTrackInstance interface */
+	virtual void SaveState (const TArray<UObject*>& RuntimeObjects) override;
+	virtual void RestoreState (const TArray<UObject*>& RuntimeObjects) override;
+	virtual void Update( float Position, float LastPosition, const TArray<UObject*>& RuntimeObjects, IMovieScenePlayer& Player ) override;
+	virtual void RefreshInstance( const TArray<UObject*>& RuntimeObjects, IMovieScenePlayer& Player ) override;
+
+private:
+
+	/** Track that is being instanced */
+	UMovieSceneFloatTrack* FloatTrack;
+
+	/** Runtime property bindings */
+	TSharedPtr<FTrackInstancePropertyBindings> PropertyBindings;
+
+	/** Map from object to initial state */
+	TMap<TWeakObjectPtr<UObject>, float> InitFloatMap;
+};
diff --git a/Engine/Source/Runtime/MovieSceneTracks/Public/TrackInstances/MovieSceneParticleTrackInstance.h b/Engine/Source/Runtime/MovieSceneTracks/Public/TrackInstances/MovieSceneParticleTrackInstance.h
new file mode 100644
--- /dev/null
+++ b/Engine/Source/Runtime/MovieSceneTracks/Public/TrackInstances/MovieSceneParticleTrackInstance.h
@@ -0,0 +1,33 @@
+// Copyright 1998-2015 Epic Games, Inc. All Rights Reserved.
+
+#pragma once
+
+#include "IMovieSceneTrackInstance.h"
+
+
+class UMovieSceneParticleTrack;
+
+
+/**
+ * Instance of a UMovieSceneParticleTrack
+ */
+class FMovieSceneParticleTrackInstance
+	: public IMovieSceneTrackInstance
+{
+public:
+
+	FMovieSceneParticleTrackInstance( UMovieSceneParticleTrack& InTrack )
+		: ParticleTrack(&InTrack)
+	{ }
+
+	virtual ~FMovieSceneParticleTrackInstance();
+
+	/** IMovieSceneTrackInstance interface */
+	virtual void Update( float Position, float LastPosition, const TArray<UObject*>& RuntimeObjects, IMovieScenePlayer& Player ) override;
+	virtual void RefreshInstance( const TArray<UObject*>& RuntimeObjects, IMovieScenePlayer& Player ) override {}
+
+private:
+
+	/** Track that is being instanced */
+	UMovieSceneParticleTrack* ParticleTrack;
+};
diff --git a/Engine/Source/Runtime/MovieSceneTracks/Public/TrackInstances/MovieSceneVectorTrackInstance.h b/Engine/Source/Runtime/MovieSceneTracks/Public/TrackInstances/MovieSceneVectorTrackInstance.h
new file mode 100644
--- /dev/null
+++ b/Engine/Source/Runtime/MovieSceneTracks/Public/TrackInstances/MovieSceneVectorTrackInstance.h
@@ -0,0 +1,40 @@
+// Copyright 1998-2015 Epic Games, Inc. All Rights Reserved.
+
+#pragma once
+
+#include "IMovieSceneTrackInstance.h"
+
+
+class FTrackInstancePropertyBindings;
+class UMovieSceneVectorTrack;
+
+
+/**
+ * Instance of a UMovieSceneVectorTrack
+ */
+class FMovieSceneVectorTrackInstance
+	: public IMovieSceneTrackInstance
+{
+public:
+
+	FMovieSceneVectorTrackInstance( UMovieSceneVectorTrack& InVectorTrack );
+
+	/** IMovieSceneTrackInstance interface */
+	virtual void SaveState (const TArray<UObject*>& RuntimeObjects) override;
+	virtual void RestoreState (const TArray<UObject*>& RuntimeObjects) override;
+	virtual void Update( float Position, float LastPosition, const TArray<UObject*>& RuntimeObjects, IMovieScenePlayer& Player ) override;
+	virtual void RefreshInstance( const TArray<UObject*>& RuntimeObjects, IMovieScenePlayer& Player ) override;
+
+private:
+
+	/** Track that is being instanced */
+	UMovieSceneVectorTrack* VectorTrack;
+
+	/** Runtime property bindings */
+	TSharedPtr<FTrackInstancePropertyBindings> PropertyBindings;
+
+	/** Map from object to initial state */
+	TMap< TWeakObjectPtr<UObject>, FVector2D > InitVector2DMap;
+	TMap< TWeakObjectPtr<UObject>, FVector > InitVector3DMap;
+	TMap< TWeakObjectPtr<UObject>, FVector4 > InitVector4DMap;
+};
diff --git a/Engine/Source/Runtime/MovieSceneTracks/Public/TrackInstances/SubMovieSceneTrackInstance.h b/Engine/Source/Runtime/MovieSceneTracks/Public/TrackInstances/SubMovieSceneTrackInstance.h
new file mode 100644
--- /dev/null
+++ b/Engine/Source/Runtime/MovieSceneTracks/Public/TrackInstances/SubMovieSceneTrackInstance.h
@@ -0,0 +1,34 @@
+// Copyright 1998-2015 Epic Games, Inc. All Rights Reserved.
+
+#pragma once
+
+#include "IMovieSceneTrackInstance.h"
+
+
+class FMovieSceneInstance;
+class USubMovieSceneSection;
+class USubMovieSceneTrack;
+
+
+/**
+ * Instance of a USubMovieSceneTrack
+ */
+class FSubMovieSceneTrackInstance
+	: public IMovieSceneTrackInstance
+{
+public:
+
+	FSubMovieSceneTrackInstance( USubMovieSceneTrack& InTrack );
+
+	/** IMovieSceneTrackInstance interface */
+	virtual void Update( float Position, float LastPosition, const TArray<UObject*>& RuntimeObjects, IMovieScenePlayer& Player ) override;
+	virtual void RefreshInstance( const TArray<UObject*>& RuntimeObjects, IMovieScenePlayer& Player ) override;
+
+private:
+
+	/** Track that is being instanced */
+	TWeakObjectPtr<USubMovieSceneTrack> SubMovieSceneTrack;
+
+	/** Mapping of section lookups to instances.  Each section has a movie scene which must be instanced */
+	TMap< TWeakObjectPtr<USubMovieSceneSection>, TSharedPtr<FMovieSceneInstance> > SubMovieSceneInstances; 
+};
diff --git a/Engine/Source/Runtime/MovieSceneTracks/Public/Tracks/MovieScene3DTransformTrack.h b/Engine/Source/Runtime/MovieSceneTracks/Public/Tracks/MovieScene3DTransformTrack.h
new file mode 100644
--- /dev/null
+++ b/Engine/Source/Runtime/MovieSceneTracks/Public/Tracks/MovieScene3DTransformTrack.h
@@ -0,0 +1,162 @@
+// Copyright 1998-2015 Epic Games, Inc. All Rights Reserved.
+
+#pragma once
+
+#include "MovieScene.h"
+#include "MovieSceneTrack.h"
+#include "MovieScenePropertyTrack.h"
+#include "MovieScene3DTransformTrack.generated.h"
+
+/**
+ * Stores information about a transform for the purpose of adding keys to a transform section
+ */
+struct FTransformData
+{
+	/** Translation component */
+	FVector Translation;
+	/** Rotation component */
+	FRotator Rotation;
+	/** Scale component */
+	FVector Scale;
+	/** Whether or not the data is valid (any values set) */
+	bool bValid;
+
+	bool IsValid() const { return bValid; }
+
+	/**
+	 * Constructor.  Builds the data from a scene component
+	 * Uses relative transform only
+	 *
+	 * @param InComponent	The component to build from
+	 */
+	FTransformData( const USceneComponent* InComponent )
+		: Translation( InComponent->RelativeLocation )
+		, Rotation( InComponent->RelativeRotation )
+		, Scale( InComponent->RelativeScale3D )
+		, bValid( true )
+	{}
+
+	FTransformData()
+		: bValid( false )
+	{}
+};
+
+/**
+ * Represents a transform that should be keyed at a specific  time
+ * 
+ * Only keys components of the transform that have changed
+ */
+class FTransformKey
+{
+public:
+	/**
+	 * Constructor
+	 *
+	 * @param InKeyTime			The time where keys should be added
+	 * @param InNewTransform	The transform values that should be keyed
+	 * @param PreviousTransform	The previous transform values (for comparing against new values).  If this is set only changed values will be keyed
+	 * 
+	 */
+	FTransformKey( float InKeyTime, const FTransformData& InNewTransform, const FTransformData& InPreviousTransform = FTransformData() )
+		: NewTransform( InNewTransform )
+		, PreviousTransform( InPreviousTransform )
+		, KeyTime( InKeyTime )
+	{
+	}	
+
+	/** @return The time where transform keys should be added */
+	float GetKeyTime() const { return KeyTime; }
+
+	/** @return the translation value of the key */
+	const FVector& GetTranslationValue() const { return NewTransform.Translation; }
+
+	/** @return the rotation value of the key */
+	const FRotator& GetRotationValue() const { return NewTransform.Rotation; }
+
+	/** @return the scale value of the key */
+	const FVector& GetScaleValue() const { return NewTransform.Scale; }
+
+	/**
+	 * Determines if a key should be added on the given axis of translation
+	 *
+	 * @param Axis	The axis to check
+	 * @return Whether or not we should key translation on the given axis
+	 */
+	bool ShouldKeyTranslation( EAxis::Type Axis ) const;
+
+	/**
+	 * Determines if a key should be added on the given axis of rotation
+	 *
+	 * @param Axis	The axis to check
+	 * @return Whether or not we should key rotation on the given axis
+	 */
+	bool ShouldKeyRotation( EAxis::Type Axis ) const;
+
+	/**
+	 * Determines if a key should be added on the given axis of scale
+	 *
+	 * @param Axis	The axis to check
+	 * @return Whether or not we should key scale on the given axis
+	 */
+	bool ShouldKeyScale( EAxis::Type Axis ) const;
+	
+	/**
+	 * Determines if a key should be added to any axis for any type of transform
+	 *
+	 * @return Whether or not we should key on the given axis
+	 */
+	bool ShouldKeyAny() const;
+private:
+	/**
+	 * Determines whether or not a vector component is different on the given axis
+	 *
+	 * @param Axis	The axis to check
+	 * @param Current	The current vector value
+	 * @param Previous	The previous vector value
+	 * @return true if a vector component is different on the given axis
+	 */
+	bool GetVectorComponentIfDifferent( EAxis::Type Axis, const FVector& Current, const FVector& Previous ) const;
+private:
+	/** New transform data */
+	FTransformData NewTransform;
+	FTransformData PreviousTransform;
+	float KeyTime;
+};
+
+/**
+ * Handles manipulation of component transforms in a movie scene
+ */
+UCLASS( MinimalAPI )
+class UMovieScene3DTransformTrack : public UMovieScenePropertyTrack
+{
+	GENERATED_UCLASS_BODY()
+public:
+	/** UMovieSceneTrack interface */
+	virtual UMovieSceneSection* CreateNewSection() override;
+	virtual TSharedPtr<IMovieSceneTrackInstance> CreateInstance() override;
+	
+	/**
+	 * Adds a key to a section.  Will create the section if it doesn't exist
+	 *
+	 * @param ObjectHandle		Handle to the object(s) being changed
+	 * @param InKey				The transform key to add
+	 * @param bUnwindRotation	When true, the value will be treated like a rotation value in degrees, and will automatically be unwound to prevent flipping 360 degrees from the previous key 	 
+	 * @return True if the key was successfully added.
+	 */
+	virtual bool AddKeyToSection( const FGuid& ObjectHandle, const FTransformKey& InKey, const bool bUnwindRotation );
+
+	/**
+	 * Evaluates the track at the playback position
+	 *
+	 * @param Position		The current playback position
+	 * @param LastPosition		The last plackback position
+	 * @param OutTranslation	The evalulated translation component of the transform
+	 * @param OutRotation		The evalulated rotation component of the transform
+	 * @param OutScale		The evalulated scale component of the transform
+	 * @param OutHasTranslationKeys true if a translation key was evaluated 
+	 * @param OutHasRotationKeys 	true if a rotation key was evaluated
+	 * @param OutHasScaleKeys 	true if a scale key was evaluated
+	 * @return true if anything was evaluated. Note: if false is returned OutBool remains unchanged
+	 */
+	virtual bool Eval( float Position, float LastPosition, FVector& OutTranslation, FRotator& OutRotation, FVector& OutScale, bool& OutHasTranslationKeys, bool& OutHasRotationKeys, bool& OutHasScaleKeys ) const;
+};
diff --git a/Engine/Source/Runtime/MovieSceneTracks/Public/Tracks/MovieSceneAnimationTrack.h b/Engine/Source/Runtime/MovieSceneTracks/Public/Tracks/MovieSceneAnimationTrack.h
new file mode 100644
--- /dev/null
+++ b/Engine/Source/Runtime/MovieSceneTracks/Public/Tracks/MovieSceneAnimationTrack.h
@@ -0,0 +1,40 @@
+// Copyright 1998-2015 Epic Games, Inc. All Rights Reserved.
+
+#pragma once
+
+#include "MovieScene.h"
+#include "MovieSceneTrack.h"
+#include "MovieSceneAnimationTrack.generated.h"
+
+
+/**
+ * Handles animation of skeletal mesh actors
+ */
+UCLASS( MinimalAPI )
+class UMovieSceneAnimationTrack : public UMovieSceneTrack
+{
+	GENERATED_UCLASS_BODY()
+
+public:
+	/** UMovieSceneTrack interface */
+	virtual FName GetTrackName() const override;
+	virtual TSharedPtr<IMovieSceneTrackInstance> CreateInstance() override;
+	virtual void RemoveAllAnimationData() override;
+	virtual bool HasSection( UMovieSceneSection* Section ) const override;
+	virtual void RemoveSection( UMovieSceneSection* Section ) override;
+	virtual bool IsEmpty() const override;
+	virtual TRange<float> GetSectionBoundaries() const override;
+	virtual bool SupportsMultipleRows() const override { return true; }
+	virtual TArray<UMovieSceneSection*> GetAllSections() const override;
+
+	/** Adds a new animation to this track */
+	virtual void AddNewAnimation(float KeyTime, class UAnimSequence* AnimSequence);
+
+	/** Gets the animation section at a certain time, or NULL if there is none */
+	class UMovieSceneSection* GetAnimSectionAtTime(float Time);
+
+private:
+	/** List of all animation sections */
+	UPROPERTY()
+	TArray<UMovieSceneSection*> AnimationSections;
+};
diff --git a/Engine/Source/Runtime/MovieSceneTracks/Public/Tracks/MovieSceneAudioTrack.h b/Engine/Source/Runtime/MovieSceneTracks/Public/Tracks/MovieSceneAudioTrack.h
new file mode 100644
--- /dev/null
+++ b/Engine/Source/Runtime/MovieSceneTracks/Public/Tracks/MovieSceneAudioTrack.h
@@ -0,0 +1,47 @@
+// Copyright 1998-2015 Epic Games, Inc. All Rights Reserved.
+
+#pragma once
+
+#include "MovieScene.h"
+#include "MovieSceneTrack.h"
+#include "MovieSceneAudioTrack.generated.h"
+
+namespace AudioTrackConstants
+{
+	const float ScrubDuration = 0.050f;
+	const FName UniqueTrackName("MasterAudio");
+}
+
+/**
+ * Handles manipulation of audio
+ */
+UCLASS( MinimalAPI )
+class UMovieSceneAudioTrack : public UMovieSceneTrack
+{
+	GENERATED_UCLASS_BODY()
+
+public:
+	/** UMovieSceneTrack interface */
+	virtual FName GetTrackName() const override;
+	virtual TSharedPtr<IMovieSceneTrackInstance> CreateInstance() override;
+	virtual void RemoveAllAnimationData() override;
+	virtual bool HasSection( UMovieSceneSection* Section ) const override;
+	virtual void RemoveSection( UMovieSceneSection* Section ) override;
+	virtual bool IsEmpty() const override;
+	virtual TRange<float> GetSectionBoundaries() const override;
+	virtual bool SupportsMultipleRows() const override { return true; }
+	virtual TArray<UMovieSceneSection*> GetAllSections() const override;
+
+	/** Adds a new sound cue to the audio */
+	virtual void AddNewSound(class USoundBase* Sound, float Time);
+
+	/** @return The audio sections on this track */
+	const TArray<UMovieSceneSection*>& GetAudioSections() const { return AudioSections; }
+
+	/** @return true if this is a master audio track */
+	bool IsAMasterTrack() const;
+private:
+	/** List of all master audio sections */
+	UPROPERTY()
+	TArray<UMovieSceneSection*> AudioSections;
+};
diff --git a/Engine/Source/Runtime/MovieSceneTracks/Public/Tracks/MovieSceneBoolTrack.h b/Engine/Source/Runtime/MovieSceneTracks/Public/Tracks/MovieSceneBoolTrack.h
new file mode 100644
--- /dev/null
+++ b/Engine/Source/Runtime/MovieSceneTracks/Public/Tracks/MovieSceneBoolTrack.h
@@ -0,0 +1,41 @@
+// Copyright 1998-2015 Epic Games, Inc. All Rights Reserved.
+
+#pragma once
+
+#include "MovieScene.h"
+#include "MovieSceneTrack.h"
+#include "MovieScenePropertyTrack.h"
+#include "MovieSceneBoolTrack.generated.h"
+
+/**
+ * Handles manipulation of float properties in a movie scene
+ */
+UCLASS( MinimalAPI )
+class UMovieSceneBoolTrack : public UMovieScenePropertyTrack
+{
+	GENERATED_UCLASS_BODY()
+
+public:
+	/** UMovieSceneTrack interface */
+	virtual UMovieSceneSection* CreateNewSection() override;
+	virtual TSharedPtr<IMovieSceneTrackInstance> CreateInstance() override;
+
+	/**
+	 * Adds a key to a section.  Will create the section if it doesn't exist
+	 *
+	 * @param Time				The time relative to the owning movie scene where the section should be
+	 * @param Value				The value of the key
+	 * @return True if the key was successfully added.
+	 */
+	virtual bool AddKeyToSection( float Time, bool Value );
+
+	/**
+	 * Evaluates the track at the playback position
+	 *
+	 * @param Position	The current playback position
+	 * @param LastPosition	The last plackback position
+	 * @param OutBool 	The value at the playback position
+	 * @return true if anything was evaluated. Note: if false is returned OutBool remains unchanged
+	 */
+	virtual bool Eval( float Position, float LastPostion, bool& OutBool ) const;
+};
diff --git a/Engine/Source/Runtime/MovieSceneTracks/Public/Tracks/MovieSceneByteTrack.h b/Engine/Source/Runtime/MovieSceneTracks/Public/Tracks/MovieSceneByteTrack.h
new file mode 100644
--- /dev/null
+++ b/Engine/Source/Runtime/MovieSceneTracks/Public/Tracks/MovieSceneByteTrack.h
@@ -0,0 +1,49 @@
+// Copyright 1998-2015 Epic Games, Inc. All Rights Reserved.
+
+#pragma once
+
+#include "MovieScene.h"
+#include "MovieSceneTrack.h"
+#include "MovieScenePropertyTrack.h"
+#include "MovieSceneByteTrack.generated.h"
+
+/**
+ * Handles manipulation of byte properties in a movie scene
+ */
+UCLASS()
+class MOVIESCENETRACKS_API UMovieSceneByteTrack : public UMovieScenePropertyTrack
+{
+	GENERATED_UCLASS_BODY()
+
+public:
+	/** UMovieSceneTrack interface */
+	virtual UMovieSceneSection* CreateNewSection() override;
+	virtual TSharedPtr<IMovieSceneTrackInstance> CreateInstance() override;
+
+	/**
+	 * Adds a key to a section.  Will create the section if it doesn't exist
+	 *
+	 * @param Time				The time relative to the owning movie scene where the section should be
+	 * @param Value				The value of the key
+	 * @return True if the key was successfully added.
+	 */
+	virtual bool AddKeyToSection( float Time, uint8 Value );
+
+	/**
+	 * Evaluates the track at the playback position
+	 *
+	 * @param Position	The current playback position
+	 * @param LastPosition	The last plackback position
+	 * @param OutByte 	The value at the playback position
+	 * @return true if anything was evaluated. Note: if false is returned OutByte remains unchanged
+	 */
+	virtual bool Eval( float Position, float LastPostion, uint8& OutByte ) const;
+
+	void SetEnum(UEnum* Enum);
+
+	class UEnum* GetEnum() const;
+
+protected:
+	UPROPERTY()
+	UEnum* Enum;
+};
diff --git a/Engine/Source/Runtime/MovieSceneTracks/Public/Tracks/MovieSceneColorTrack.h b/Engine/Source/Runtime/MovieSceneTracks/Public/Tracks/MovieSceneColorTrack.h
new file mode 100644
--- /dev/null
+++ b/Engine/Source/Runtime/MovieSceneTracks/Public/Tracks/MovieSceneColorTrack.h
@@ -0,0 +1,55 @@
+// Copyright 1998-2015 Epic Games, Inc. All Rights Reserved.
+
+#pragma once
+
+#include "MovieScene.h"
+#include "MovieSceneTrack.h"
+#include "MovieScenePropertyTrack.h"
+#include "MovieSceneColorTrack.generated.h"
+
+struct FColorKey
+{
+	FLinearColor Value;
+	FName CurveName;
+	bool bAddKeyEvenIfUnchanged;
+	bool bIsSlateColor;
+};
+
+/**
+ * Handles manipulation of float properties in a movie scene
+ */
+UCLASS( MinimalAPI )
+class UMovieSceneColorTrack : public UMovieScenePropertyTrack
+{
+	GENERATED_UCLASS_BODY()
+
+public:
+	/** UMovieSceneTrack interface */
+	virtual UMovieSceneSection* CreateNewSection() override;
+	virtual TSharedPtr<IMovieSceneTrackInstance> CreateInstance() override;
+
+	/**
+	 * Adds a key to a section.  Will create the section if it doesn't exist
+	 *
+	 * @param Time				The time relative to the owning movie scene where the section should be
+	 * @param Value				The value of the key
+	 * @return True if the key was successfully added.
+	 */
+	virtual bool AddKeyToSection( float Time, const FColorKey& Key );
+	
+	/**
+	 * Evaluates the track at the playback position
+	 *
+	 * @param Position	The current playback position
+	 * @param LastPosition	The last plackback position
+	 * @param The color at the playback position
+	 * @return true if anything was evaluated. Note: if false is returned OutColor remains unchanged
+	 */
+	virtual bool Eval( float Position, float LastPostion, FLinearColor& OutColor ) const;
+
+	bool IsSlateColor() const { return bIsSlateColor; }
+
+private:
+	UPROPERTY()
+	bool bIsSlateColor;
+};
diff --git a/Engine/Source/Runtime/MovieSceneTracks/Public/Tracks/MovieSceneDirectorTrack.h b/Engine/Source/Runtime/MovieSceneTracks/Public/Tracks/MovieSceneDirectorTrack.h
new file mode 100644
--- /dev/null
+++ b/Engine/Source/Runtime/MovieSceneTracks/Public/Tracks/MovieSceneDirectorTrack.h
@@ -0,0 +1,39 @@
+// Copyright 1998-2015 Epic Games, Inc. All Rights Reserved.
+
+#pragma once
+
+#include "MovieScene.h"
+#include "MovieSceneTrack.h"
+#include "MovieSceneDirectorTrack.generated.h"
+
+/**
+ * Handles manipulation of shot properties in a movie scene
+ */
+UCLASS( MinimalAPI )
+class UMovieSceneDirectorTrack : public UMovieSceneTrack
+{
+	GENERATED_UCLASS_BODY()
+
+public:
+	/** UMovieSceneTrack interface */
+	virtual FName GetTrackName() const override;
+	virtual TSharedPtr<IMovieSceneTrackInstance> CreateInstance() override;
+	virtual void RemoveAllAnimationData() override;
+	virtual bool HasSection( UMovieSceneSection* Section ) const override;
+	virtual void RemoveSection( UMovieSceneSection* Section ) override;
+	virtual bool IsEmpty() const override;
+	virtual TRange<float> GetSectionBoundaries() const override;
+	virtual bool SupportsMultipleRows() const override { return true; }
+	virtual TArray<UMovieSceneSection*> GetAllSections() const override;
+
+	/** Adds a new shot at the specified time */
+	virtual void AddNewShot(FGuid CameraHandle, float Time);
+
+	/** @return The shot sections on this track */
+	const TArray<UMovieSceneSection*>& GetShotSections() const { return ShotSections; }
+
+private:
+	/** List of all shot sections */
+	UPROPERTY()
+	TArray<UMovieSceneSection*> ShotSections;
+};
diff --git a/Engine/Source/Runtime/MovieSceneTracks/Public/Tracks/MovieSceneFloatTrack.h b/Engine/Source/Runtime/MovieSceneTracks/Public/Tracks/MovieSceneFloatTrack.h
new file mode 100644
--- /dev/null
+++ b/Engine/Source/Runtime/MovieSceneTracks/Public/Tracks/MovieSceneFloatTrack.h
@@ -0,0 +1,42 @@
+// Copyright 1998-2015 Epic Games, Inc. All Rights Reserved.
+
+#pragma once
+
+#include "MovieScene.h"
+#include "MovieSceneTrack.h"
+#include "MovieScenePropertyTrack.h"
+#include "MovieSceneFloatTrack.generated.h"
+
+/**
+ * Handles manipulation of float properties in a movie scene
+ */
+UCLASS( MinimalAPI )
+class UMovieSceneFloatTrack : public UMovieScenePropertyTrack
+{
+	GENERATED_UCLASS_BODY()
+
+public:
+	/** UMovieSceneTrack interface */
+	virtual UMovieSceneSection* CreateNewSection() override;
+	virtual TSharedPtr<IMovieSceneTrackInstance> CreateInstance() override;
+	
+	/**
+	 * Adds a key to a section.  Will create the section if it doesn't exist
+	 *
+	 * @param Time				The time relative to the owning movie scene where the section should be
+	 * @param Value				The value of the key
+	 * @return True if the key was successfully added.
+	 */
+	virtual bool AddKeyToSection( float Time, float Value );
+
+	/**
+	 * Evaluates the track at the playback position
+	 *
+	 * @param Position	The current playback position
+	 * @param LastPosition	The last plackback position
+	 * @param OutFloat 	The value at the playback position
+	 * @return true if anything was evaluated. Note: if false is returned OutFloat remains unchanged
+	 */
+	virtual bool Eval( float Position, float LastPosition, float& OutFloat ) const;
+
+};
diff --git a/Engine/Source/Runtime/MovieSceneTracks/Public/Tracks/MovieSceneParticleTrack.h b/Engine/Source/Runtime/MovieSceneTracks/Public/Tracks/MovieSceneParticleTrack.h
new file mode 100644
--- /dev/null
+++ b/Engine/Source/Runtime/MovieSceneTracks/Public/Tracks/MovieSceneParticleTrack.h
@@ -0,0 +1,38 @@
+// Copyright 1998-2015 Epic Games, Inc. All Rights Reserved.
+
+#pragma once
+
+#include "MovieScene.h"
+#include "MovieSceneTrack.h"
+#include "MovieSceneParticleTrack.generated.h"
+
+
+/**
+ * Handles triggering of particle emitters
+ */
+UCLASS( MinimalAPI )
+class UMovieSceneParticleTrack : public UMovieSceneTrack
+{
+	GENERATED_UCLASS_BODY()
+
+public:
+	/** UMovieSceneTrack interface */
+	virtual FName GetTrackName() const override;
+	virtual TSharedPtr<IMovieSceneTrackInstance> CreateInstance() override;
+	virtual void RemoveAllAnimationData() override;
+	virtual bool HasSection( UMovieSceneSection* Section ) const override;
+	virtual void RemoveSection( UMovieSceneSection* Section ) override;
+	virtual bool IsEmpty() const override;
+	virtual TRange<float> GetSectionBoundaries() const override;
+	virtual bool SupportsMultipleRows() const override { return false; }
+	virtual TArray<UMovieSceneSection*> GetAllSections() const override;
+	
+	virtual void AddNewParticleSystem(float KeyTime, bool bTrigger);
+
+	virtual TArray<UMovieSceneSection*> GetAllParticleSections() const {return ParticleSections;}
+
+private:
+	/** List of all particle sections */
+	UPROPERTY()
+	TArray<UMovieSceneSection*> ParticleSections;
+};
diff --git a/Engine/Source/Runtime/MovieSceneTracks/Public/Tracks/MovieSceneVectorTrack.h b/Engine/Source/Runtime/MovieSceneTracks/Public/Tracks/MovieSceneVectorTrack.h
new file mode 100644
--- /dev/null
+++ b/Engine/Source/Runtime/MovieSceneTracks/Public/Tracks/MovieSceneVectorTrack.h
@@ -0,0 +1,66 @@
+// Copyright 1998-2015 Epic Games, Inc. All Rights Reserved.
+
+#pragma once
+
+#include "MovieScene.h"
+#include "MovieSceneTrack.h"
+#include "MovieScenePropertyTrack.h"
+#include "MovieScenePropertyTrack.h"
+
+#include "MovieSceneVectorTrack.generated.h"
+
+template<typename VecType>
+struct FVectorKey
+{
+	VecType Value;
+	FName CurveName;
+	bool bAddKeyEvenIfUnchanged;
+};
+
+/**
+ * Handles manipulation of component transforms in a movie scene
+ */
+UCLASS(MinimalAPI)
+class UMovieSceneVectorTrack : public UMovieScenePropertyTrack
+{
+	GENERATED_UCLASS_BODY()
+public:
+	/** UMovieSceneTrack interface */
+	virtual UMovieSceneSection* CreateNewSection() override;
+	virtual TSharedPtr<IMovieSceneTrackInstance> CreateInstance() override;
+
+	/**
+	 * Adds a key to a section.  Will create the section if it doesn't exist
+	 *
+	 * @param ObjectHandle		Handle to the object(s) being changed
+	 * @param InPropertyName	The name of the property being manipulated.  @todo Sequencer - Could be a UFunction name
+	 * @param Time				The time relative to the owning movie scene where the section should be
+	 * @param InKey				The vector key to add
+	 * @param InChannelsUsed	The number of channels used, 2, 3, or 4, determining which type of vector this is
+	 * @return True if the key was successfully added.
+	 */
+
+	virtual bool AddKeyToSection( float Time, const FVectorKey<FVector4>& Key );
+	virtual bool AddKeyToSection( float Time, const FVectorKey<FVector>& Key );
+	virtual bool AddKeyToSection( float Time, const FVectorKey<FVector2D>& Key );
+	
+	/**
+	 * Evaluates the track at the playback position
+	 *
+	 * @param Position	The current playback position
+	 * @param LastPosition	The last plackback position
+	 * @param OutVector 	The value at the playback position
+	 * @return true if anything was evaluated. Note: if false is returned OutVector remains unchanged
+	 */
+	virtual bool Eval( float Position, float LastPostion, FVector4& InOutVector ) const;
+
+	/** @return Get the number of channels used by the vector */
+	int32 GetNumChannelsUsed() const { return NumChannelsUsed; }
+
+private:
+	virtual bool AddKeyToSection( float Time, const FVector4& InKey, int32 InChannelsUsed, FName CurveName, bool bAddKeyEvenIfUnchanged );
+private:
+	/** The number of channels used by the vector (2,3, or 4) */
+	UPROPERTY()
+	int32 NumChannelsUsed;
+};
diff --git a/Engine/Source/Runtime/MovieSceneTracks/Public/Tracks/SubMovieSceneTrack.h b/Engine/Source/Runtime/MovieSceneTracks/Public/Tracks/SubMovieSceneTrack.h
new file mode 100644
--- /dev/null
+++ b/Engine/Source/Runtime/MovieSceneTracks/Public/Tracks/SubMovieSceneTrack.h
@@ -0,0 +1,37 @@
+// Copyright 1998-2015 Epic Games, Inc. All Rights Reserved.
+
+#pragma once
+
+#include "MovieSceneTrack.h"
+#include "SubMovieSceneTrack.generated.h"
+
+/**
+ * Handles manipulation of float properties in a movie scene
+ */
+UCLASS( MinimalAPI )
+class USubMovieSceneTrack : public UMovieSceneTrack
+{
+	GENERATED_UCLASS_BODY()
+public:
+	/** UMovieSceneTrack interface */
+	virtual UMovieSceneSection* CreateNewSection() override;
+	virtual FName GetTrackName() const override;
+	virtual TSharedPtr<IMovieSceneTrackInstance> CreateInstance() override;
+	virtual TArray<UMovieSceneSection*> GetAllSections() const override;
+	virtual void RemoveSection( UMovieSceneSection* Section ) override;
+	virtual bool IsEmpty() const override;
+	virtual TRange<float> GetSectionBoundaries() const override;
+	virtual bool HasSection( UMovieSceneSection* Section ) const override;
+
+	/**
+	 * Adds a movie scene section
+	 *
+	 * @param SubMovieScene	The movie scene to add
+	 * @param Time	The time to add the section at
+	 */
+	virtual void AddMovieSceneSection( UMovieScene* SubMovieScene, float Time );
+private:
+	/** All movie scene sections.  */
+	UPROPERTY()
+	TArray<class UMovieSceneSection*> SubMovieSceneSections;
+};
