Fixed non-unity build errors

diff --git a/Engine/Source/Runtime/MovieSceneTracks/Private/TrackInstances/MovieScene3DTransformTrackInstance.cpp b/Engine/Source/Runtime/MovieSceneTracks/Private/TrackInstances/MovieScene3DTransformTrackInstance.cpp
--- a/Engine/Source/Runtime/MovieSceneTracks/Private/TrackInstances/MovieScene3DTransformTrackInstance.cpp
+++ b/Engine/Source/Runtime/MovieSceneTracks/Private/TrackInstances/MovieScene3DTransformTrackInstance.cpp
@@ -2,6 +2,7 @@
 
 #include "MovieSceneTracksPrivatePCH.h"
 #include "MovieScene3DTransformTrackInstance.h"
+#include "MovieScene3DTransformTrack.h"
 
 
 FMovieScene3DTransformTrackInstance::FMovieScene3DTransformTrackInstance( UMovieScene3DTransformTrack& InTransformTrack )
diff --git a/Engine/Source/Runtime/MovieSceneTracks/Private/TrackInstances/MovieSceneAnimationTrackInstance.cpp b/Engine/Source/Runtime/MovieSceneTracks/Private/TrackInstances/MovieSceneAnimationTrackInstance.cpp
--- a/Engine/Source/Runtime/MovieSceneTracks/Private/TrackInstances/MovieSceneAnimationTrackInstance.cpp
+++ b/Engine/Source/Runtime/MovieSceneTracks/Private/TrackInstances/MovieSceneAnimationTrackInstance.cpp
@@ -3,6 +3,7 @@
 #include "MovieSceneTracksPrivatePCH.h"
 #include "MovieSceneAnimationTrack.h"
 #include "MovieSceneAnimationTrackInstance.h"
+#include "MovieSceneAnimationSection.h"
 #include "IMovieScenePlayer.h"
 #include "Matinee/MatineeAnimInterface.h"
 
diff --git a/Engine/Source/Runtime/MovieSceneTracks/Private/TrackInstances/MovieSceneAudioTrackInstance.cpp b/Engine/Source/Runtime/MovieSceneTracks/Private/TrackInstances/MovieSceneAudioTrackInstance.cpp
--- a/Engine/Source/Runtime/MovieSceneTracks/Private/TrackInstances/MovieSceneAudioTrackInstance.cpp
+++ b/Engine/Source/Runtime/MovieSceneTracks/Private/TrackInstances/MovieSceneAudioTrackInstance.cpp
@@ -7,6 +7,7 @@
 #include "Runtime/Engine/Classes/Kismet/GameplayStatics.h"
 #include "Runtime/Engine/Public/AudioDecompress.h"
 #include "MovieSceneAudioTrack.h"
+#include "MovieSceneAudioSection.h"
 
 FMovieSceneAudioTrackInstance::FMovieSceneAudioTrackInstance( UMovieSceneAudioTrack& InAudioTrack )
 {
diff --git a/Engine/Source/Runtime/MovieSceneTracks/Private/TrackInstances/MovieSceneColorTrackInstance.cpp b/Engine/Source/Runtime/MovieSceneTracks/Private/TrackInstances/MovieSceneColorTrackInstance.cpp
--- a/Engine/Source/Runtime/MovieSceneTracks/Private/TrackInstances/MovieSceneColorTrackInstance.cpp
+++ b/Engine/Source/Runtime/MovieSceneTracks/Private/TrackInstances/MovieSceneColorTrackInstance.cpp
@@ -3,6 +3,7 @@
 #include "MovieSceneTracksPrivatePCH.h"
 #include "MovieSceneColorTrackInstance.h"
 #include "MovieSceneCommonHelpers.h"
+#include "MovieSceneColorTrack.h"
 #include "SlateCore.h"
 
 
diff --git a/Engine/Source/Runtime/MovieSceneTracks/Private/TrackInstances/MovieSceneDirectorTrackInstance.cpp b/Engine/Source/Runtime/MovieSceneTracks/Private/TrackInstances/MovieSceneDirectorTrackInstance.cpp
--- a/Engine/Source/Runtime/MovieSceneTracks/Private/TrackInstances/MovieSceneDirectorTrackInstance.cpp
+++ b/Engine/Source/Runtime/MovieSceneTracks/Private/TrackInstances/MovieSceneDirectorTrackInstance.cpp
@@ -3,6 +3,7 @@
 #include "MovieSceneTracksPrivatePCH.h"
 #include "MovieSceneDirectorTrack.h"
 #include "MovieSceneDirectorTrackInstance.h"
+#include "MovieSceneShotSection.h"
 #include "IMovieScenePlayer.h"
 
 
diff --git a/Engine/Source/Runtime/MovieSceneTracks/Private/TrackInstances/MovieSceneVectorTrackInstance.cpp b/Engine/Source/Runtime/MovieSceneTracks/Private/TrackInstances/MovieSceneVectorTrackInstance.cpp
--- a/Engine/Source/Runtime/MovieSceneTracks/Private/TrackInstances/MovieSceneVectorTrackInstance.cpp
+++ b/Engine/Source/Runtime/MovieSceneTracks/Private/TrackInstances/MovieSceneVectorTrackInstance.cpp
@@ -2,6 +2,7 @@
 
 #include "MovieSceneTracksPrivatePCH.h"
 #include "MovieSceneVectorTrackInstance.h"
+#include "MovieSceneVectorTrack.h"
 
 FMovieSceneVectorTrackInstance::FMovieSceneVectorTrackInstance( UMovieSceneVectorTrack& InVectorTrack )
 {
