Sequencer: Added a custom track for the bHidden property:

MovieSceneVisibilityTrack (creates the custom section and instance), MovieSceneVisibilitySection (inverts the bHiddenInGame property value for the curve keyframes), MovieSceneVisibilityInstance (handles the visibility in the instance).

diff --git a/Engine/Source/Runtime/MovieSceneTracks/Private/Sections/MovieSceneVisibilitySection.cpp b/Engine/Source/Runtime/MovieSceneTracks/Private/Sections/MovieSceneVisibilitySection.cpp
new file mode 100644
--- /dev/null
+++ b/Engine/Source/Runtime/MovieSceneTracks/Private/Sections/MovieSceneVisibilitySection.cpp
@@ -0,0 +1,22 @@
+// Copyright 1998-2015 Epic Games, Inc. All Rights Reserved.
+
+#include "MovieSceneTracksPrivatePCH.h"
+#include "MovieSceneVisibilitySection.h"
+
+
+UMovieSceneVisibilitySection::UMovieSceneVisibilitySection( const FObjectInitializer& ObjectInitializer )
+	: Super( ObjectInitializer )
+{
+}
+
+void UMovieSceneVisibilitySection::AddKey( float Time, bool bHiddenInGame )
+{
+	// The property that's being changed is bHiddenInGame. But we want to store the inverse of that as visibility, so invert the incoming value.
+	Super::AddKey(Time, !bHiddenInGame);
+}
+
+bool UMovieSceneVisibilitySection::NewKeyIsNewData( float Time, bool bHiddenInGame ) const
+{
+	// The property that's being changed is bHiddenInGame. But we want to store the inverse of that as visibility, so invert the incoming value.
+	return Super::NewKeyIsNewData(Time, !bHiddenInGame);
+}
diff --git a/Engine/Source/Runtime/MovieSceneTracks/Private/TrackInstances/MovieSceneVisibilityTrackInstance.cpp b/Engine/Source/Runtime/MovieSceneTracks/Private/TrackInstances/MovieSceneVisibilityTrackInstance.cpp
new file mode 100644
--- /dev/null
+++ b/Engine/Source/Runtime/MovieSceneTracks/Private/TrackInstances/MovieSceneVisibilityTrackInstance.cpp
@@ -0,0 +1,101 @@
+// Copyright 1998-2015 Epic Games, Inc. All Rights Reserved.
+
+#include "MovieSceneTracksPrivatePCH.h"
+#include "MovieSceneVisibilityTrackInstance.h"
+#include "MovieSceneVisibilityTrack.h"
+
+FMovieSceneVisibilityTrackInstance::FMovieSceneVisibilityTrackInstance(UMovieSceneVisibilityTrack& InVisibilityTrack)
+{
+	VisibilityTrack = &InVisibilityTrack;
+
+	PropertyBindings = MakeShareable( new FTrackInstancePropertyBindings( VisibilityTrack->GetPropertyName(), VisibilityTrack->GetPropertyPath() ) );
+}
+
+void FMovieSceneVisibilityTrackInstance::SaveState(const TArray<UObject*>& RuntimeObjects)
+{
+ 	for( UObject* Object : RuntimeObjects )
+ 	{
+		AActor* Actor = Cast<AActor>(Object);
+
+		if (Actor != NULL)
+		{
+#if WITH_EDITOR
+			if (GIsEditor && !Actor->GetWorld()->IsPlayInEditor())
+			{
+				InitHiddenInEditorMap.Add(Actor, Actor->IsTemporarilyHiddenInEditor());
+			}
+#endif // WITH_EDITOR
+
+			InitHiddenInGameMap.Add(Object, Actor->bHidden);
+		}
+ 	}
+}
+
+void FMovieSceneVisibilityTrackInstance::RestoreState(const TArray<UObject*>& RuntimeObjects)
+{
+ 	for( UObject* Object : RuntimeObjects )
+ 	{
+		AActor* Actor = Cast<AActor>(Object);
+
+		if (Actor != NULL)
+		{
+			bool *HiddenInGameValue = InitHiddenInGameMap.Find(Object);
+			if (HiddenInGameValue != NULL)
+			{
+				Actor->SetActorHiddenInGame(*HiddenInGameValue);
+			}
+
+#if WITH_EDITOR
+			if (GIsEditor && !Actor->GetWorld()->IsPlayInEditor())
+			{
+				bool *HiddenInEditorValue = InitHiddenInEditorMap.Find(Object);
+				if (HiddenInEditorValue != NULL)
+				{
+					if (Actor->IsTemporarilyHiddenInEditor() != *HiddenInEditorValue)
+					{
+						Actor->SetIsTemporarilyHiddenInEditor(*HiddenInEditorValue);
+					}
+				}
+			}
+		}
+#endif // WITH_EDITOR
+ 	}
+}
+
+void FMovieSceneVisibilityTrackInstance::Update( float Position, float LastPosition, const TArray<UObject*>& RuntimeObjects, class IMovieScenePlayer& Player ) 
+{
+ 	bool Visible = false;
+ 	if( VisibilityTrack->Eval( Position, LastPosition, Visible ) )
+ 	{
+		// Invert this evaluation since the property is "bHiddenInGame" and we want the visualization to be the inverse of that. Green means visible.
+		Visible = !Visible;
+
+ 		for( UObject* Object : RuntimeObjects )
+ 		{
+			AActor* Actor = Cast<AActor>(Object);
+
+			if (Actor != NULL)
+			{
+				Actor->SetActorHiddenInGame(Visible);
+
+#if WITH_EDITOR
+				if (GIsEditor && !Actor->GetWorld()->IsPlayInEditor())
+				{
+					// In editor HiddenGame flag is not respected so set bHiddenEdTemporary too.
+					// It will be restored just like HiddenGame flag when Matinee is closed.
+					if (Actor->IsTemporarilyHiddenInEditor() != Visible)
+					{
+						Actor->SetIsTemporarilyHiddenInEditor(Visible);
+					}
+				}
+#endif // WITH_EDITOR
+			}
+ 		}
+	}
+}
+
+void FMovieSceneVisibilityTrackInstance::RefreshInstance( const TArray<UObject*>& RuntimeObjects, class IMovieScenePlayer& Player )
+{
+	PropertyBindings->UpdateBindings( RuntimeObjects );
+}
+
diff --git a/Engine/Source/Runtime/MovieSceneTracks/Private/Tracks/MovieSceneVisibilityTrack.cpp b/Engine/Source/Runtime/MovieSceneTracks/Private/Tracks/MovieSceneVisibilityTrack.cpp
new file mode 100644
--- /dev/null
+++ b/Engine/Source/Runtime/MovieSceneTracks/Private/Tracks/MovieSceneVisibilityTrack.cpp
@@ -0,0 +1,38 @@
+// Copyright 1998-2015 Epic Games, Inc. All Rights Reserved.
+
+#include "MovieSceneTracksPrivatePCH.h"
+#include "MovieSceneVisibilityTrack.h"
+#include "IMovieScenePlayer.h"
+#include "MovieSceneVisibilitySection.h"
+#include "MovieSceneVisibilityTrackInstance.h"
+
+UMovieSceneVisibilityTrack::UMovieSceneVisibilityTrack( const FObjectInitializer& ObjectInitializer )
+	: Super( ObjectInitializer )
+{
+}
+
+UMovieSceneSection* UMovieSceneVisibilityTrack::CreateNewSection()
+{
+	return NewObject<UMovieSceneSection>(this, UMovieSceneVisibilitySection::StaticClass(), NAME_None, RF_Transactional);
+}
+
+TSharedPtr<IMovieSceneTrackInstance> UMovieSceneVisibilityTrack::CreateInstance()
+{
+	return MakeShareable( new FMovieSceneVisibilityTrackInstance( *this ) );
+}
+
+bool UMovieSceneVisibilityTrack::AddKeyToSection( float Time, bool Value )
+{
+	const UMovieSceneSection* NearestSection = MovieSceneHelpers::FindNearestSectionAtTime( Sections, Time );
+	if (!NearestSection || CastChecked<UMovieSceneVisibilitySection>(NearestSection)->NewKeyIsNewData(Time, Value))
+	{
+		Modify();
+
+		UMovieSceneVisibilitySection* NewSection = CastChecked<UMovieSceneVisibilitySection>(FindOrAddSection(  Time ));
+	
+		NewSection->AddKey( Time, Value );
+
+		return true;
+	}
+	return false;
+}
diff --git a/Engine/Source/Runtime/MovieSceneTracks/Public/Sections/MovieSceneVisibilitySection.h b/Engine/Source/Runtime/MovieSceneTracks/Public/Sections/MovieSceneVisibilitySection.h
new file mode 100644
--- /dev/null
+++ b/Engine/Source/Runtime/MovieSceneTracks/Public/Sections/MovieSceneVisibilitySection.h
@@ -0,0 +1,37 @@
+// Copyright 1998-2015 Epic Games, Inc. All Rights Reserved.
+
+#pragma once
+
+#include "MovieSceneBoolSection.h"
+#include "MovieSceneVisibilitySection.generated.h"
+
+/**
+ * A single visibility section.
+ *
+ * The property that's being tracked by this section is bHiddenInGame. 
+ * This custom bool track stores the inverse keys to display visibility (A green section bar means visible).
+ *
+ */
+UCLASS( MinimalAPI )
+class UMovieSceneVisibilitySection : public UMovieSceneBoolSection
+{
+	GENERATED_UCLASS_BODY()
+public:
+
+	/** 
+	 * Adds a key to the section
+	 *
+	 * @param Time	The location in time where the key should be added
+	 * @param Value	The value of the key
+	 */
+	void AddKey( float Time, bool Value );
+
+	/** 
+	 * Determines if a new key would be new data, or just a duplicate of existing data
+	 *
+	 * @param Time	The location in time where the key would be added
+	 * @param Value	The value of the new key
+	 * @return True if the new key would be new data, false if duplicate
+	 */
+	bool NewKeyIsNewData(float Time, bool Value) const;
+};
diff --git a/Engine/Source/Runtime/MovieSceneTracks/Public/TrackInstances/MovieSceneVisibilityTrackInstance.h b/Engine/Source/Runtime/MovieSceneTracks/Public/TrackInstances/MovieSceneVisibilityTrackInstance.h
new file mode 100644
--- /dev/null
+++ b/Engine/Source/Runtime/MovieSceneTracks/Public/TrackInstances/MovieSceneVisibilityTrackInstance.h
@@ -0,0 +1,34 @@
+// Copyright 1998-2015 Epic Games, Inc. All Rights Reserved.
+
+#pragma once
+
+#include "IMovieSceneTrackInstance.h"
+
+class FTrackInstancePropertyBindings;
+class UMovieSceneVisibilityTrack;
+
+/**
+ * Instance of a UMovieSceneVisibilityTrack
+ *
+ * Visibility track instance is a specialized bool track instance that handles visibility in game and in editor
+ *
+ */
+class FMovieSceneVisibilityTrackInstance : public IMovieSceneTrackInstance
+{
+public:
+	FMovieSceneVisibilityTrackInstance( UMovieSceneVisibilityTrack& InVisibilityTrack );
+
+	/** IMovieSceneTrackInstance interface */
+	virtual void SaveState (const TArray<UObject*>& RuntimeObjects) override;
+	virtual void RestoreState (const TArray<UObject*>& RuntimeObjects) override;
+	virtual void Update( float Position, float LastPosition, const TArray<UObject*>& RuntimeObjects, class IMovieScenePlayer& Player ) override;
+	virtual void RefreshInstance( const TArray<UObject*>& RuntimeObjects, class IMovieScenePlayer& Player ) override;
+private:
+	/** Visibility track that is being instanced */
+	UMovieSceneVisibilityTrack* VisibilityTrack;
+	/** Runtime property bindings */
+	TSharedPtr<class FTrackInstancePropertyBindings> PropertyBindings;
+	/** Map from object to initial state */
+	TMap< TWeakObjectPtr<UObject>, bool > InitHiddenInGameMap;
+	TMap< TWeakObjectPtr<UObject>, bool > InitHiddenInEditorMap;
+};
diff --git a/Engine/Source/Runtime/MovieSceneTracks/Public/Tracks/MovieSceneVisibilityTrack.h b/Engine/Source/Runtime/MovieSceneTracks/Public/Tracks/MovieSceneVisibilityTrack.h
new file mode 100644
--- /dev/null
+++ b/Engine/Source/Runtime/MovieSceneTracks/Public/Tracks/MovieSceneVisibilityTrack.h
@@ -0,0 +1,31 @@
+// Copyright 1998-2015 Epic Games, Inc. All Rights Reserved.
+
+#pragma once
+
+#include "MovieScene.h"
+#include "MovieSceneTrack.h"
+#include "MovieScenePropertyTrack.h"
+#include "MovieSceneVisibilityTrack.generated.h"
+
+/**
+ * Handles manipulation of visibility properties in a movie scene
+ */
+UCLASS( MinimalAPI )
+class UMovieSceneVisibilityTrack : public UMovieSceneBoolTrack
+{
+	GENERATED_UCLASS_BODY()
+
+public:
+	/** UMovieSceneTrack interface */
+	virtual UMovieSceneSection* CreateNewSection() override;
+	virtual TSharedPtr<IMovieSceneTrackInstance> CreateInstance() override;
+
+	/**
+	 * Adds a key to a section.  Will create the section if it doesn't exist
+	 *
+	 * @param Time				The time relative to the owning movie scene where the section should be
+	 * @param Value				The value of the key
+	 * @return True if the key was successfully added.
+	 */
+	virtual bool AddKeyToSection( float Time, bool Value );
+};
