Sequencer: Fix visibility track compile (include booltrack)

diff --git a/Engine/Source/Runtime/MovieSceneTracks/Public/Tracks/MovieSceneVisibilityTrack.h b/Engine/Source/Runtime/MovieSceneTracks/Public/Tracks/MovieSceneVisibilityTrack.h
--- a/Engine/Source/Runtime/MovieSceneTracks/Public/Tracks/MovieSceneVisibilityTrack.h
+++ b/Engine/Source/Runtime/MovieSceneTracks/Public/Tracks/MovieSceneVisibilityTrack.h
@@ -5,6 +5,7 @@
 #include "MovieScene.h"
 #include "MovieSceneTrack.h"
 #include "MovieScenePropertyTrack.h"
+#include "MovieSceneBoolTrack.h"
 #include "MovieSceneVisibilityTrack.generated.h"
 
 /**
