Sequencer: Fix crash on close when runtime objects (spawnables) are being destroyed.

diff --git a/Engine/Source/Runtime/MovieSceneTracks/Private/TrackInstances/MovieScene3DTransformTrackInstance.cpp b/Engine/Source/Runtime/MovieSceneTracks/Private/TrackInstances/MovieScene3DTransformTrackInstance.cpp
--- a/Engine/Source/Runtime/MovieSceneTracks/Private/TrackInstances/MovieScene3DTransformTrackInstance.cpp
+++ b/Engine/Source/Runtime/MovieSceneTracks/Private/TrackInstances/MovieScene3DTransformTrackInstance.cpp
@@ -49,6 +49,11 @@ void FMovieScene3DTransformTrackInstance::RestoreState(const TArray<UObject*>& R
 {
 	for (int32 ObjIndex = 0; ObjIndex < RuntimeObjects.Num(); ++ObjIndex)
 	{
+		if (!IsValid(RuntimeObjects[ObjIndex]))
+		{
+			continue;
+		}
+
 		USceneComponent* SceneComponent = SceneComponentFromRuntimeObject(RuntimeObjects[ObjIndex]);
 		if (SceneComponent != NULL)
 		{
diff --git a/Engine/Source/Runtime/MovieSceneTracks/Private/TrackInstances/MovieSceneBoolTrackInstance.cpp b/Engine/Source/Runtime/MovieSceneTracks/Private/TrackInstances/MovieSceneBoolTrackInstance.cpp
--- a/Engine/Source/Runtime/MovieSceneTracks/Private/TrackInstances/MovieSceneBoolTrackInstance.cpp
+++ b/Engine/Source/Runtime/MovieSceneTracks/Private/TrackInstances/MovieSceneBoolTrackInstance.cpp
@@ -27,6 +27,11 @@ void FMovieSceneBoolTrackInstance::RestoreState(const TArray<UObject*>& RuntimeO
 {
 	for( UObject* Object : RuntimeObjects )
 	{
+		if (!IsValid(Object))
+		{
+			continue;
+		}
+
 		bool *BoolValue = InitBoolMap.Find(Object);
 		if (BoolValue != NULL)
 		{
diff --git a/Engine/Source/Runtime/MovieSceneTracks/Private/TrackInstances/MovieSceneByteTrackInstance.cpp b/Engine/Source/Runtime/MovieSceneTracks/Private/TrackInstances/MovieSceneByteTrackInstance.cpp
--- a/Engine/Source/Runtime/MovieSceneTracks/Private/TrackInstances/MovieSceneByteTrackInstance.cpp
+++ b/Engine/Source/Runtime/MovieSceneTracks/Private/TrackInstances/MovieSceneByteTrackInstance.cpp
@@ -27,6 +27,11 @@ void FMovieSceneByteTrackInstance::RestoreState(const TArray<UObject*>& RuntimeO
 {
 	for( UObject* Object : RuntimeObjects )
 	{
+		if (!IsValid(Object))
+		{
+			continue;
+		}
+
 		uint8 *ByteValue = InitByteMap.Find(Object);
 		if (ByteValue != NULL)
 		{
diff --git a/Engine/Source/Runtime/MovieSceneTracks/Private/TrackInstances/MovieSceneColorTrackInstance.cpp b/Engine/Source/Runtime/MovieSceneTracks/Private/TrackInstances/MovieSceneColorTrackInstance.cpp
--- a/Engine/Source/Runtime/MovieSceneTracks/Private/TrackInstances/MovieSceneColorTrackInstance.cpp
+++ b/Engine/Source/Runtime/MovieSceneTracks/Private/TrackInstances/MovieSceneColorTrackInstance.cpp
@@ -35,6 +35,11 @@ void FMovieSceneColorTrackInstance::RestoreState(const TArray<UObject*>& Runtime
 {
 	for( UObject* Object : RuntimeObjects )
 	{
+		if (!IsValid(Object))
+		{
+			continue;
+		}
+
 		if( ColorTrack->IsSlateColor() )
 		{
 			FSlateColor *ColorValue = InitSlateColorMap.Find(Object);
diff --git a/Engine/Source/Runtime/MovieSceneTracks/Private/TrackInstances/MovieSceneFloatTrackInstance.cpp b/Engine/Source/Runtime/MovieSceneTracks/Private/TrackInstances/MovieSceneFloatTrackInstance.cpp
--- a/Engine/Source/Runtime/MovieSceneTracks/Private/TrackInstances/MovieSceneFloatTrackInstance.cpp
+++ b/Engine/Source/Runtime/MovieSceneTracks/Private/TrackInstances/MovieSceneFloatTrackInstance.cpp
@@ -27,6 +27,11 @@ void FMovieSceneFloatTrackInstance::RestoreState(const TArray<UObject*>& Runtime
 {
 	for( UObject* Object : RuntimeObjects )
 	{
+		if (!IsValid(Object))
+		{
+			continue;
+		}
+
 		float *FloatValue = InitFloatMap.Find(Object);
 		if (FloatValue != NULL)
 		{
diff --git a/Engine/Source/Runtime/MovieSceneTracks/Private/TrackInstances/MovieSceneVectorTrackInstance.cpp b/Engine/Source/Runtime/MovieSceneTracks/Private/TrackInstances/MovieSceneVectorTrackInstance.cpp
--- a/Engine/Source/Runtime/MovieSceneTracks/Private/TrackInstances/MovieSceneVectorTrackInstance.cpp
+++ b/Engine/Source/Runtime/MovieSceneTracks/Private/TrackInstances/MovieSceneVectorTrackInstance.cpp
@@ -45,6 +45,11 @@ void FMovieSceneVectorTrackInstance::RestoreState(const TArray<UObject*>& Runtim
 	int32 NumChannelsUsed = VectorTrack->GetNumChannelsUsed();
 	for( UObject* Object : RuntimeObjects )
 	{
+		if (!IsValid(Object))
+		{
+			continue;
+		}
+
 		switch( NumChannelsUsed )
 		{
 			case 2:
diff --git a/Engine/Source/Runtime/MovieSceneTracks/Private/TrackInstances/MovieSceneVisibilityTrackInstance.cpp b/Engine/Source/Runtime/MovieSceneTracks/Private/TrackInstances/MovieSceneVisibilityTrackInstance.cpp
--- a/Engine/Source/Runtime/MovieSceneTracks/Private/TrackInstances/MovieSceneVisibilityTrackInstance.cpp
+++ b/Engine/Source/Runtime/MovieSceneTracks/Private/TrackInstances/MovieSceneVisibilityTrackInstance.cpp
@@ -35,6 +35,11 @@ void FMovieSceneVisibilityTrackInstance::RestoreState(const TArray<UObject*>& Ru
 {
  	for( UObject* Object : RuntimeObjects )
  	{
+		if (!IsValid(Object))
+		{
+			continue;
+		}
+
 		AActor* Actor = Cast<AActor>(Object);
 
 		if (Actor != NULL)
