Fixed Sub-MovieScenes not adding their own track. Previously they tried to add all movie scenes as sections to one track
Fixed the display name of sub-moviescene tracks.  It now uses the name of the movie scene instead of a generic "Scenes" name

diff --git a/Engine/Source/Editor/MovieSceneTools/Private/TrackEditors/SubMovieSceneTrackEditor.cpp b/Engine/Source/Editor/MovieSceneTools/Private/TrackEditors/SubMovieSceneTrackEditor.cpp
--- a/Engine/Source/Editor/MovieSceneTools/Private/TrackEditors/SubMovieSceneTrackEditor.cpp
+++ b/Engine/Source/Editor/MovieSceneTools/Private/TrackEditors/SubMovieSceneTrackEditor.cpp
@@ -24,7 +24,7 @@ public:
 
 	virtual FText GetDisplayName() const override
 	{
-		return DisplayName;
+		return GetSectionTitle();
 	}
 
 	virtual FText GetSectionTitle() const override
diff --git a/Engine/Source/Editor/Sequencer/Private/Sequencer.cpp b/Engine/Source/Editor/Sequencer/Private/Sequencer.cpp
--- a/Engine/Source/Editor/Sequencer/Private/Sequencer.cpp
+++ b/Engine/Source/Editor/Sequencer/Private/Sequencer.cpp
@@ -1127,13 +1127,9 @@ void FSequencer::AddSubMovieScene( UMovieScene* SubMovieScene )
 
 	OwnerMovieScene->Modify();
 
-	UMovieSceneTrack* Type = OwnerMovieScene->FindMasterTrack( USubMovieSceneTrack::StaticClass() ) ;
-	if( !Type )
-	{
-		Type = OwnerMovieScene->AddMasterTrack( USubMovieSceneTrack::StaticClass() );
-	}
+	UMovieSceneTrack* NewTrack = OwnerMovieScene->AddMasterTrack( USubMovieSceneTrack::StaticClass() );
 
-	USubMovieSceneTrack* SubMovieSceneType = CastChecked<USubMovieSceneTrack>( Type );
+	USubMovieSceneTrack* SubMovieSceneType = CastChecked<USubMovieSceneTrack>( NewTrack );
 
 	SubMovieSceneType->AddMovieSceneSection( SubMovieScene, ScrubPosition );
 }
diff --git a/Engine/Source/Runtime/MovieSceneTracks/Private/Tracks/SubMovieSceneTrack.cpp b/Engine/Source/Runtime/MovieSceneTracks/Private/Tracks/SubMovieSceneTrack.cpp
--- a/Engine/Source/Runtime/MovieSceneTracks/Private/Tracks/SubMovieSceneTrack.cpp
+++ b/Engine/Source/Runtime/MovieSceneTracks/Private/Tracks/SubMovieSceneTrack.cpp
@@ -18,9 +18,15 @@ UMovieSceneSection* USubMovieSceneTrack::CreateNewSection()
 
 FName USubMovieSceneTrack::GetTrackName() const
 {
-	static const FName UniqueName( "Scenes" );
+	static const FName DefaultName( "MovieScene" );
 
-	return UniqueName;
+	FName Name = DefaultName;
+	if( SubMovieSceneSections.Num() > 0 )
+	{
+		Name = SubMovieSceneSections[0]->GetMovieScene()->GetFName();
+	}
+
+	return Name;
 }
 
 TSharedPtr<IMovieSceneTrackInstance> USubMovieSceneTrack::CreateInstance()
@@ -41,7 +47,11 @@ TArray<UMovieSceneSection*> USubMovieSceneTrack::GetAllSections() const
 
 void USubMovieSceneTrack::RemoveSection( UMovieSceneSection* Section )
 {
-	SubMovieSceneSections.Remove( Section );
+	USubMovieSceneSection* SubMovieSceneSection = Cast<USubMovieSceneSection>( Section );
+	if(SubMovieSceneSection)
+	{
+		SubMovieSceneSections.Remove( SubMovieSceneSection );
+	}
 }
 
 bool USubMovieSceneTrack::IsEmpty() const
diff --git a/Engine/Source/Runtime/MovieSceneTracks/Public/Tracks/SubMovieSceneTrack.h b/Engine/Source/Runtime/MovieSceneTracks/Public/Tracks/SubMovieSceneTrack.h
--- a/Engine/Source/Runtime/MovieSceneTracks/Public/Tracks/SubMovieSceneTrack.h
+++ b/Engine/Source/Runtime/MovieSceneTracks/Public/Tracks/SubMovieSceneTrack.h
@@ -33,5 +33,5 @@ public:
 private:
 	/** All movie scene sections.  */
 	UPROPERTY()
-	TArray<class UMovieSceneSection*> SubMovieSceneSections;
+	TArray<class USubMovieSceneSection*> SubMovieSceneSections;
 };
