Added a Save and RestoreState methods to SubMovieSceneTrackInstance to ensure any objects in sub-moviescenes are saved and restored properly

diff --git a/Engine/Source/Runtime/MovieSceneTracks/Private/TrackInstances/SubMovieSceneTrackInstance.cpp b/Engine/Source/Runtime/MovieSceneTracks/Private/TrackInstances/SubMovieSceneTrackInstance.cpp
--- a/Engine/Source/Runtime/MovieSceneTracks/Private/TrackInstances/SubMovieSceneTrackInstance.cpp
+++ b/Engine/Source/Runtime/MovieSceneTracks/Private/TrackInstances/SubMovieSceneTrackInstance.cpp
@@ -81,3 +81,32 @@ void FSubMovieSceneTrackInstance::Update( float Position, float LastPosition, co
 	}
 
 }
+
+
+void FSubMovieSceneTrackInstance::SaveState(const TArray<UObject*>& RuntimeObjects)
+{
+	TArray<UMovieSceneSection*> AllSections = SubMovieSceneTrack->GetAllSections();
+
+	for( int32 SectionIndex = 0; SectionIndex < AllSections.Num(); ++SectionIndex )
+	{
+		USubMovieSceneSection* Section = CastChecked<USubMovieSceneSection>( AllSections[SectionIndex] );
+
+		TSharedPtr<FMovieSceneInstance>& Instance = SubMovieSceneInstances.FindChecked( Section );
+
+		Instance->SaveState();
+	}
+}
+
+void FSubMovieSceneTrackInstance::RestoreState(const TArray<UObject*>& RuntimeObjects)
+{
+	TArray<UMovieSceneSection*> AllSections = SubMovieSceneTrack->GetAllSections();
+
+	for( int32 SectionIndex = 0; SectionIndex < AllSections.Num(); ++SectionIndex )
+	{
+		USubMovieSceneSection* Section = CastChecked<USubMovieSceneSection>( AllSections[SectionIndex] );
+
+		TSharedPtr<FMovieSceneInstance>& Instance = SubMovieSceneInstances.FindChecked( Section );
+
+		Instance->RestoreState();
+	}
+}
\ No newline at end of file
diff --git a/Engine/Source/Runtime/MovieSceneTracks/Public/TrackInstances/SubMovieSceneTrackInstance.h b/Engine/Source/Runtime/MovieSceneTracks/Public/TrackInstances/SubMovieSceneTrackInstance.h
--- a/Engine/Source/Runtime/MovieSceneTracks/Public/TrackInstances/SubMovieSceneTrackInstance.h
+++ b/Engine/Source/Runtime/MovieSceneTracks/Public/TrackInstances/SubMovieSceneTrackInstance.h
@@ -23,7 +23,8 @@ public:
 	/** IMovieSceneTrackInstance interface */
 	virtual void Update( float Position, float LastPosition, const TArray<UObject*>& RuntimeObjects, IMovieScenePlayer& Player ) override;
 	virtual void RefreshInstance( const TArray<UObject*>& RuntimeObjects, IMovieScenePlayer& Player ) override;
-
+	virtual void SaveState (const TArray<UObject*>& RuntimeObjects) override;
+	virtual void RestoreState (const TArray<UObject*>& RuntimeObjects) override;
 private:
 
 	/** Track that is being instanced */
