Sequencer: Refactor SceneComponentFromRuntimeObject to Helpers.

diff --git a/Engine/Source/Runtime/MovieScene/Private/MovieSceneCommonHelpers.cpp b/Engine/Source/Runtime/MovieScene/Private/MovieSceneCommonHelpers.cpp
--- a/Engine/Source/Runtime/MovieScene/Private/MovieSceneCommonHelpers.cpp
+++ b/Engine/Source/Runtime/MovieScene/Private/MovieSceneCommonHelpers.cpp
@@ -86,6 +86,24 @@ UMovieSceneSection* MovieSceneHelpers::FindNearestSectionAtTime( const TArray<UM
 	return ClosestSection ? ClosestSection : EarliestSection;
 }
 
+USceneComponent* MovieSceneHelpers::SceneComponentFromRuntimeObject(UObject* Object)
+{
+	AActor* Actor = Cast<AActor>(Object);
+
+	USceneComponent* SceneComponent = NULL;
+	if (Actor && Actor->GetRootComponent())
+	{
+		// If there is an actor, modify its root component
+		SceneComponent = Actor->GetRootComponent();
+	}
+	else
+	{
+		// No actor was found.  Attempt to get the object as a component in the case that we are editing them directly.
+		SceneComponent = Cast<USceneComponent>(Object);
+	}
+	return SceneComponent;
+}
+
 FTrackInstancePropertyBindings::FTrackInstancePropertyBindings( FName InPropertyName, const FString& InPropertyPath, const FName& InFunctionName )
     : PropertyPath( InPropertyPath )
 	, PropertyName( InPropertyName )
diff --git a/Engine/Source/Runtime/MovieScene/Public/MovieSceneCommonHelpers.h b/Engine/Source/Runtime/MovieScene/Public/MovieSceneCommonHelpers.h
--- a/Engine/Source/Runtime/MovieScene/Public/MovieSceneCommonHelpers.h
+++ b/Engine/Source/Runtime/MovieScene/Public/MovieSceneCommonHelpers.h
@@ -26,6 +26,14 @@ public:
 	* @return The found section or null
 	*/
 	static UMovieSceneSection* FindNearestSectionAtTime( const TArray<UMovieSceneSection*>& Sections, float Time );
+
+	/**
+	* Get the scene component from the runtime object
+	*
+	* @param Object The object to get the scene component for
+	* @return The found scene component
+	*/	
+	static USceneComponent* SceneComponentFromRuntimeObject(UObject* Object);
 };
 
 /**
diff --git a/Engine/Source/Runtime/MovieSceneTracks/Private/TrackInstances/MovieScene3DTransformTrackInstance.cpp b/Engine/Source/Runtime/MovieSceneTracks/Private/TrackInstances/MovieScene3DTransformTrackInstance.cpp
--- a/Engine/Source/Runtime/MovieSceneTracks/Private/TrackInstances/MovieScene3DTransformTrackInstance.cpp
+++ b/Engine/Source/Runtime/MovieSceneTracks/Private/TrackInstances/MovieScene3DTransformTrackInstance.cpp
@@ -1,6 +1,7 @@
 // Copyright 1998-2015 Epic Games, Inc. All Rights Reserved.
 
 #include "MovieSceneTracksPrivatePCH.h"
+#include "MovieSceneCommonHelpers.h"
 #include "MovieScene3DTransformTrackInstance.h"
 #include "MovieScene3DTransformTrack.h"
 
@@ -10,34 +11,11 @@ FMovieScene3DTransformTrackInstance::FMovieScene3DTransformTrackInstance( UMovie
 	TransformTrack = &InTransformTrack;
 }
 
-namespace
-{
-
-USceneComponent* SceneComponentFromRuntimeObject(UObject* Object)
-{
-	AActor* Actor = Cast<AActor>(Object);
-
-	USceneComponent* SceneComponent = NULL;
-	if (Actor && Actor->GetRootComponent())
-	{
-		// If there is an actor, modify its root component
-		SceneComponent = Actor->GetRootComponent();
-	}
-	else
-	{
-		// No actor was found.  Attempt to get the object as a component in the case that we are editing them directly.
-		SceneComponent = Cast<USceneComponent>(Object);
-	}
-	return SceneComponent;
-}
-
-}
-
 void FMovieScene3DTransformTrackInstance::SaveState(const TArray<UObject*>& RuntimeObjects)
 {
 	for (int32 ObjIndex = 0; ObjIndex < RuntimeObjects.Num(); ++ObjIndex)
 	{
-		USceneComponent* SceneComponent = SceneComponentFromRuntimeObject(RuntimeObjects[ObjIndex]);
+		USceneComponent* SceneComponent = MovieSceneHelpers::SceneComponentFromRuntimeObject(RuntimeObjects[ObjIndex]);
 		if (SceneComponent != NULL)
 		{
 			InitTransformMap.Add(RuntimeObjects[ObjIndex], SceneComponent->GetRelativeTransform());
@@ -54,7 +32,7 @@ void FMovieScene3DTransformTrackInstance::RestoreState(const TArray<UObject*>& R
 			continue;
 		}
 
-		USceneComponent* SceneComponent = SceneComponentFromRuntimeObject(RuntimeObjects[ObjIndex]);
+		USceneComponent* SceneComponent = MovieSceneHelpers::SceneComponentFromRuntimeObject(RuntimeObjects[ObjIndex]);
 		if (SceneComponent != NULL)
 		{
 			FTransform *Transform = InitTransformMap.Find(RuntimeObjects[ObjIndex]);
@@ -79,7 +57,7 @@ void FMovieScene3DTransformTrackInstance::Update( float Position, float LastPosi
 	{
 		for( int32 ObjIndex = 0; ObjIndex < RuntimeObjects.Num(); ++ObjIndex )
 		{
-			USceneComponent* SceneComponent = SceneComponentFromRuntimeObject(RuntimeObjects[ObjIndex]);
+			USceneComponent* SceneComponent = MovieSceneHelpers::SceneComponentFromRuntimeObject(RuntimeObjects[ObjIndex]);
 
 			if (SceneComponent != NULL)
 			{
