MovieScene: Moved property track into MovieSceneTracks

diff --git a/Engine/Source/Runtime/MovieScene/Private/MovieScenePropertyTrack.cpp b/Engine/Source/Runtime/MovieScene/Private/MovieScenePropertyTrack.cpp
deleted file mode 100644
diff --git a/Engine/Source/Runtime/MovieScene/Public/MovieScenePropertyTrack.h b/Engine/Source/Runtime/MovieScene/Public/MovieScenePropertyTrack.h
deleted file mode 100644
diff --git a/Engine/Source/Runtime/MovieSceneTracks/Private/Tracks/MovieScenePropertyTrack.cpp b/Engine/Source/Runtime/MovieSceneTracks/Private/Tracks/MovieScenePropertyTrack.cpp
new file mode 100644
--- /dev/null
+++ b/Engine/Source/Runtime/MovieSceneTracks/Private/Tracks/MovieScenePropertyTrack.cpp
@@ -0,0 +1,128 @@
+// Copyright 1998-2015 Epic Games, Inc. All Rights Reserved.
+
+#include "MovieSceneTracksPrivatePCH.h"
+#include "MovieScenePropertyTrack.h"
+
+
+UMovieScenePropertyTrack::UMovieScenePropertyTrack( const FObjectInitializer& ObjectInitializer )
+	: Super( ObjectInitializer )
+{
+	bSectionsAreShowable = false;
+}
+
+
+void UMovieScenePropertyTrack::SetPropertyNameAndPath( FName InPropertyName, const FString& InPropertyPath )
+{
+	check( InPropertyName != NAME_None && !InPropertyPath.IsEmpty() );
+
+	PropertyName = InPropertyName;
+	PropertyPath = InPropertyPath;
+}
+
+
+const TArray<UMovieSceneSection*>& UMovieScenePropertyTrack::GetAllSections() const
+{
+	return Sections;
+}
+
+
+void UMovieScenePropertyTrack::RemoveAllAnimationData()
+{
+	Sections.Empty();
+}
+
+
+bool UMovieScenePropertyTrack::HasSection( UMovieSceneSection* Section ) const 
+{
+	return Sections.Contains( Section );
+}
+
+
+void UMovieScenePropertyTrack::RemoveSection( UMovieSceneSection* Section ) 
+{
+	Sections.Remove( Section );
+}
+
+
+bool UMovieScenePropertyTrack::IsEmpty() const
+{
+	return Sections.Num() == 0;
+}
+
+
+UMovieSceneSection* UMovieScenePropertyTrack::FindOrAddSection( float Time )
+{
+	// Find a spot for the section so that they are sorted by start time
+	for( int32 SectionIndex = 0; SectionIndex < Sections.Num(); ++SectionIndex )
+	{
+		UMovieSceneSection* Section = Sections[SectionIndex];
+
+		if( Section->IsTimeWithinSection( Time ) )
+		{
+			return Section;
+		}
+
+		// Check if there are no more sections that would overlap the time 
+		if( !Sections.IsValidIndex( SectionIndex+1 ) || Sections[SectionIndex+1]->GetStartTime() > Time )
+		{
+			// No sections overlap the time
+		
+			if( SectionIndex > 0 )
+			{
+				// Append and grow the previous section
+				UMovieSceneSection* PreviousSection = Sections[ SectionIndex ? SectionIndex-1 : 0 ];
+		
+				PreviousSection->SetEndTime( Time );
+				return PreviousSection;
+			}
+			else if( Sections.IsValidIndex( SectionIndex+1 ) )
+			{
+				// Prepend and grow the next section because there are no sections before this one
+				UMovieSceneSection* NextSection = Sections[SectionIndex+1];
+				NextSection->SetStartTime( Time );
+				return NextSection;
+			}	
+			else
+			{
+				// SectionIndex == 0 
+				UMovieSceneSection* PreviousSection = Sections[0];
+				if( PreviousSection->GetEndTime() < Time )
+				{
+					// Append and grow the section
+					PreviousSection->SetEndTime( Time );
+				}
+				else
+				{
+					// Prepend and grow the section
+					PreviousSection->SetStartTime( Time );
+				}
+				return PreviousSection;
+			}
+		}
+
+	}
+
+	check( Sections.Num() == 0 );
+
+	// Add a new section that starts and ends at the same time
+	UMovieSceneSection* NewSection = CreateNewSection();
+	NewSection->SetStartTime( Time );
+	NewSection->SetEndTime( Time );
+
+	Sections.Add( NewSection );
+
+	return NewSection;
+}
+
+
+TRange<float> UMovieScenePropertyTrack::GetSectionBoundaries() const
+{
+	TArray< TRange<float> > Bounds;
+
+	for (int32 SectionIndex = 0; SectionIndex < Sections.Num(); ++SectionIndex)
+	{
+		Bounds.Add(Sections[SectionIndex]->GetRange());
+	}
+
+	return TRange<float>::Hull(Bounds);
+}
diff --git a/Engine/Source/Runtime/MovieSceneTracks/Public/Tracks/MovieScenePropertyTrack.h b/Engine/Source/Runtime/MovieSceneTracks/Public/Tracks/MovieScenePropertyTrack.h
new file mode 100644
--- /dev/null
+++ b/Engine/Source/Runtime/MovieSceneTracks/Public/Tracks/MovieScenePropertyTrack.h
@@ -0,0 +1,69 @@
+// Copyright 1998-2015 Epic Games, Inc. All Rights Reserved.
+
+#pragma once
+
+#include "MovieSceneTrack.h"
+#include "MovieScenePropertyTrack.generated.h"
+
+/**
+ * Base class for tracks that animate an object property
+ */
+UCLASS( abstract )
+class MOVIESCENE_API UMovieScenePropertyTrack : public UMovieSceneTrack
+{
+	GENERATED_UCLASS_BODY()
+
+public:
+	/** UMovieSceneTrack interface */
+	virtual FName GetTrackName() const override { return PropertyName; }
+	virtual void RemoveAllAnimationData() override;
+	virtual bool HasSection( UMovieSceneSection* Section ) const override;
+	virtual void RemoveSection( UMovieSceneSection* Section ) override;
+	virtual bool IsEmpty() const override;
+	virtual TRange<float> GetSectionBoundaries() const override;
+	virtual const TArray<UMovieSceneSection*>& GetAllSections() const override;
+	virtual bool HasShowableData() const override {return bSectionsAreShowable;}
+	
+	/**
+	 * Sets the property name for this animatable property
+	 *
+	 * @param InPropertyName	The property being animated
+	 */
+	virtual void SetPropertyNameAndPath( FName InPropertyName, const FString& InPropertyPath );
+
+	/** @return the name of the property being animataed by this track */
+	FName GetPropertyName() const { return PropertyName; }
+	
+	/** @return The property path for this track */
+	const FString& GetPropertyPath() const { return PropertyPath; }
+
+	/** Sets this track as showable permanently */
+	void SetAsShowable() {bSectionsAreShowable = true;}
+
+protected:
+	/**
+	 * Finds a section at the current time.
+	 *
+	 * @param Time	The time relative to the owning movie scene where the section should be
+	 *
+	 * @return The found section, or NULL if it can't be found
+	 */
+	class UMovieSceneSection* FindOrAddSection(  float Time );
+protected:
+
+	/** Name of the property being changed */
+	UPROPERTY()
+	FName PropertyName;
+
+	/** Path to the property from the source object being changed */
+	UPROPERTY()
+	FString PropertyPath;
+
+	/** All the sections in this list */
+	UPROPERTY()
+	TArray<UMovieSceneSection*> Sections;
+	
+	/** True if this should generate a display node in Sequencer */
+	UPROPERTY()
+	bool bSectionsAreShowable;
+};
