MovieScene: Moved property track into MovieSceneTracks

diff --git a/Engine/Source/Runtime/MovieScene/Private/MovieScenePrivatePCH.h b/Engine/Source/Runtime/MovieScene/Private/MovieScenePrivatePCH.h
--- a/Engine/Source/Runtime/MovieScene/Private/MovieScenePrivatePCH.h
+++ b/Engine/Source/Runtime/MovieScene/Private/MovieScenePrivatePCH.h
@@ -11,7 +11,6 @@
 #include "MovieScene.h"
 #include "MovieSceneSection.h"
 #include "MovieSceneTrack.h"
-#include "MovieScenePropertyTrack.h"
 #include "MovieSceneInstance.h"
 #include "IMovieSceneTrackInstance.h"
 
diff --git a/Engine/Source/Runtime/MovieSceneTracks/Public/Tracks/MovieScenePropertyTrack.h b/Engine/Source/Runtime/MovieSceneTracks/Public/Tracks/MovieScenePropertyTrack.h
--- a/Engine/Source/Runtime/MovieSceneTracks/Public/Tracks/MovieScenePropertyTrack.h
+++ b/Engine/Source/Runtime/MovieSceneTracks/Public/Tracks/MovieScenePropertyTrack.h
@@ -9,7 +9,7 @@
  * Base class for tracks that animate an object property
  */
 UCLASS( abstract )
-class MOVIESCENE_API UMovieScenePropertyTrack : public UMovieSceneTrack
+class MOVIESCENETRACKS_API UMovieScenePropertyTrack : public UMovieSceneTrack
 {
 	GENERATED_UCLASS_BODY()
 
diff --git a/Engine/Source/Runtime/UMG/UMG.Build.cs b/Engine/Source/Runtime/UMG/UMG.Build.cs
--- a/Engine/Source/Runtime/UMG/UMG.Build.cs
+++ b/Engine/Source/Runtime/UMG/UMG.Build.cs
@@ -30,6 +30,7 @@ public class UMG : ModuleRules
         PublicDependencyModuleNames.AddRange(
             new string[] {
 				"HTTP",
+                "MovieSceneTracks",
 			}
         );
 
