Sequencer: Fix transform evaluation so that only channels that have keys are evaluated.

diff --git a/Engine/Source/Editor/MovieSceneTools/Private/TrackEditors/TransformTrackEditor.cpp b/Engine/Source/Editor/MovieSceneTools/Private/TrackEditors/TransformTrackEditor.cpp
--- a/Engine/Source/Editor/MovieSceneTools/Private/TrackEditors/TransformTrackEditor.cpp
+++ b/Engine/Source/Editor/MovieSceneTools/Private/TrackEditors/TransformTrackEditor.cpp
@@ -426,7 +426,9 @@ void F3DTransformTrackEditor::OnTransformChangedInternals(float KeyTime, UObject
 	
 		if (!TransformPair.LastTransformData.IsValid())
 		{
-			bool bHasTranslationKeys = false, bHasRotationKeys = false, bHasScaleKeys = false;
+			TArray<bool> bHasTranslationKeys;
+			TArray<bool> bHasRotationKeys;
+			TArray<bool> bHasScaleKeys;			
 			TransformPair.LastTransformData.bValid = TransformTrack->Eval(KeyTime, KeyTime, TransformPair.LastTransformData.Translation, TransformPair.LastTransformData.Rotation, TransformPair.LastTransformData.Scale, bHasTranslationKeys, bHasRotationKeys, bHasScaleKeys);
 		}
 
diff --git a/Engine/Source/Runtime/MovieSceneTracks/Private/Sections/MovieScene3DTransformSection.cpp b/Engine/Source/Runtime/MovieSceneTracks/Private/Sections/MovieScene3DTransformSection.cpp
--- a/Engine/Source/Runtime/MovieSceneTracks/Private/Sections/MovieScene3DTransformSection.cpp
+++ b/Engine/Source/Runtime/MovieSceneTracks/Private/Sections/MovieScene3DTransformSection.cpp
@@ -66,51 +66,70 @@ void UMovieScene3DTransformSection::GetKeyHandles(TSet<FKeyHandle>& KeyHandles)
 	}
 }	
 
-bool UMovieScene3DTransformSection::EvalTranslation( float Time, FVector& OutTranslation ) const
+void UMovieScene3DTransformSection::EvalTranslation( float Time, FVector& OutTranslation, TArray<bool>& OutHasTranslationKeys ) const
 {
-	bool bHasTranslationKeys = Translation[0].GetNumKeys() > 0 || Translation[1].GetNumKeys() > 0 || Translation[2].GetNumKeys() > 0;
+	OutHasTranslationKeys.Add(Translation[0].GetNumKeys() > 0);
+	OutHasTranslationKeys.Add(Translation[1].GetNumKeys() > 0);
+	OutHasTranslationKeys.Add(Translation[2].GetNumKeys() > 0);
 
-	if( bHasTranslationKeys )
+	// Evaluate each component of translation if there are any keys
+
+	if (OutHasTranslationKeys[0])
 	{
-		// Evaluate each component of translation if there are any keys
 		OutTranslation.X = Translation[0].Eval( Time );
+	}
+	if (OutHasTranslationKeys[1])
+	{
 		OutTranslation.Y = Translation[1].Eval( Time );
+	}
+	if (OutHasTranslationKeys[2])
+	{
 		OutTranslation.Z = Translation[2].Eval( Time );
 	}
-
-	return bHasTranslationKeys;
 }
 
-bool UMovieScene3DTransformSection::EvalRotation( float Time, FRotator& OutRotation ) const
+void UMovieScene3DTransformSection::EvalRotation( float Time, FRotator& OutRotation, TArray<bool>& OutHasRotationKeys ) const
 {
-	bool bHasRotationKeys = Rotation[0].GetNumKeys() > 0 || Rotation[1].GetNumKeys() > 0 || Rotation[2].GetNumKeys() > 0;
+	OutHasRotationKeys.Add(Rotation[0].GetNumKeys() > 0);
+	OutHasRotationKeys.Add(Rotation[1].GetNumKeys() > 0);
+	OutHasRotationKeys.Add(Rotation[2].GetNumKeys() > 0);
 
-	if( bHasRotationKeys )
+	// Evaluate each component of rotation if there are any keys
+
+	if (OutHasRotationKeys[0])
 	{
-		// Evaluate each component of rotation if there are any keys
 		OutRotation.Roll = Rotation[0].Eval( Time );
+	}
+	if (OutHasRotationKeys[1])
+	{
 		OutRotation.Pitch = Rotation[1].Eval( Time );
+	}
+	if (OutHasRotationKeys[2])
+	{
 		OutRotation.Yaw = Rotation[2].Eval( Time );
 	}
-
-	return bHasRotationKeys;
-
 }
 
-bool UMovieScene3DTransformSection::EvalScale( float Time, FVector& OutScale ) const
+void UMovieScene3DTransformSection::EvalScale( float Time, FVector& OutScale, TArray<bool>& OutHasScaleKeys ) const
 {
-	bool bHasScaleKeys = Scale[0].GetNumKeys() > 0 || Scale[1].GetNumKeys() > 0 || Scale[2].GetNumKeys() > 0;
-	
-	if( bHasScaleKeys )
+	OutHasScaleKeys.Add(Scale[0].GetNumKeys() > 0);
+	OutHasScaleKeys.Add(Scale[1].GetNumKeys() > 0);
+	OutHasScaleKeys.Add(Scale[2].GetNumKeys() > 0);
+
+	// Evaluate each component of scale if there are any keys
+
+	if (OutHasScaleKeys[0])
 	{
-		// Evaluate each component of scale if there are any keys
 		OutScale.X = Scale[0].Eval( Time );
+	}
+	if (OutHasScaleKeys[1])
+	{
 		OutScale.Y = Scale[1].Eval( Time );
+	}
+	if (OutHasScaleKeys[2])
+	{
 		OutScale.Z = Scale[2].Eval( Time );
 	}
-
-	return bHasScaleKeys;
-
 }
 
 /**
diff --git a/Engine/Source/Runtime/MovieSceneTracks/Private/TrackInstances/MovieScene3DTransformTrackInstance.cpp b/Engine/Source/Runtime/MovieSceneTracks/Private/TrackInstances/MovieScene3DTransformTrackInstance.cpp
--- a/Engine/Source/Runtime/MovieSceneTracks/Private/TrackInstances/MovieScene3DTransformTrackInstance.cpp
+++ b/Engine/Source/Runtime/MovieSceneTracks/Private/TrackInstances/MovieScene3DTransformTrackInstance.cpp
@@ -49,9 +49,9 @@ void FMovieScene3DTransformTrackInstance::Update( float Position, float LastPosi
 	FVector Translation;
 	FRotator Rotation;
 	FVector Scale;
-	bool bHasTranslationKeys = false;
-	bool bHasRotationKeys = false;
-	bool bHasScaleKeys = false;
+	TArray<bool> bHasTranslationKeys;
+	TArray<bool> bHasRotationKeys;
+	TArray<bool> bHasScaleKeys;
 
 	if( TransformTrack->Eval( Position, LastPosition, Translation, Rotation, Scale, bHasTranslationKeys, bHasRotationKeys, bHasScaleKeys ) )
 	{
@@ -62,14 +62,24 @@ void FMovieScene3DTransformTrackInstance::Update( float Position, float LastPosi
 			if (SceneComponent != NULL)
 			{
 				// Set the relative translation and rotation.  Note they are set once instead of individually to avoid a redundant component transform update.
-				SceneComponent->SetRelativeLocationAndRotation(
-					bHasTranslationKeys ? Translation : SceneComponent->RelativeLocation,
-					bHasRotationKeys ? Rotation : SceneComponent->RelativeRotation );
+				FVector NewTranslation;
+				NewTranslation[0] = bHasTranslationKeys[0] ? Translation[0] : SceneComponent->RelativeLocation[0];
+				NewTranslation[1] = bHasTranslationKeys[1] ? Translation[1] : SceneComponent->RelativeLocation[1];
+				NewTranslation[2] = bHasTranslationKeys[2] ? Translation[2] : SceneComponent->RelativeLocation[2];
 
-				if( bHasScaleKeys )
-				{
-					SceneComponent->SetRelativeScale3D( Scale );
-				}
+				FRotator NewRotation;
+				NewRotation.Roll = bHasRotationKeys[0] ? Rotation.Roll : SceneComponent->RelativeRotation.Roll;
+				NewRotation.Pitch = bHasRotationKeys[1] ? Rotation.Pitch : SceneComponent->RelativeRotation.Pitch;
+				NewRotation.Yaw = bHasRotationKeys[2] ? Rotation.Yaw : SceneComponent->RelativeRotation.Yaw;
+
+				SceneComponent->SetRelativeLocationAndRotation(NewTranslation, NewRotation);
+
+				FVector NewScale;
+				NewScale[0] = bHasScaleKeys[0] ? Scale[0] : SceneComponent->RelativeScale3D[0];
+				NewScale[1] = bHasScaleKeys[1] ? Scale[1] : SceneComponent->RelativeScale3D[1];
+				NewScale[2] = bHasScaleKeys[2] ? Scale[2] : SceneComponent->RelativeScale3D[2];
+
+				SceneComponent->SetRelativeScale3D(NewScale);
 			}
 		}
 	}
diff --git a/Engine/Source/Runtime/MovieSceneTracks/Private/Tracks/MovieScene3DTransformTrack.cpp b/Engine/Source/Runtime/MovieSceneTracks/Private/Tracks/MovieScene3DTransformTrack.cpp
--- a/Engine/Source/Runtime/MovieSceneTracks/Private/Tracks/MovieScene3DTransformTrack.cpp
+++ b/Engine/Source/Runtime/MovieSceneTracks/Private/Tracks/MovieScene3DTransformTrack.cpp
@@ -95,7 +95,7 @@ bool UMovieScene3DTransformTrack::AddKeyToSection( const FGuid& ObjectHandle, co
 }
 
 
-bool UMovieScene3DTransformTrack::Eval( float Position, float LastPosition, FVector& OutTranslation, FRotator& OutRotation, FVector& OutScale, bool& OutHasTranslationKeys, bool& OutHasRotationKeys, bool& OutHasScaleKeys ) const
+bool UMovieScene3DTransformTrack::Eval( float Position, float LastPosition, FVector& OutTranslation, FRotator& OutRotation, FVector& OutScale, TArray<bool>& OutHasTranslationKeys, TArray<bool>& OutHasRotationKeys, TArray<bool>& OutHasScaleKeys ) const
 {
 	const UMovieSceneSection* Section = MovieSceneHelpers::FindSectionAtTime( Sections, Position );
 
@@ -104,9 +104,9 @@ bool UMovieScene3DTransformTrack::Eval( float Position, float LastPosition, FVec
 		const UMovieScene3DTransformSection* TransformSection = CastChecked<UMovieScene3DTransformSection>( Section );
 
 		// Evalulate translation,rotation, and scale curves.  If no keys were found on one of these, that component of the transform will remain unchained
-		OutHasTranslationKeys = TransformSection->EvalTranslation( Position, OutTranslation );
-		OutHasRotationKeys = TransformSection->EvalRotation( Position, OutRotation );
-		OutHasScaleKeys = TransformSection->EvalScale( Position, OutScale );
+		TransformSection->EvalTranslation( Position, OutTranslation, OutHasTranslationKeys );
+		TransformSection->EvalRotation( Position, OutRotation, OutHasRotationKeys );
+		TransformSection->EvalScale( Position, OutScale, OutHasScaleKeys );
 	}
 
 	return Section != NULL;
diff --git a/Engine/Source/Runtime/MovieSceneTracks/Public/Sections/MovieScene3DTransformSection.h b/Engine/Source/Runtime/MovieSceneTracks/Public/Sections/MovieScene3DTransformSection.h
--- a/Engine/Source/Runtime/MovieSceneTracks/Public/Sections/MovieScene3DTransformSection.h
+++ b/Engine/Source/Runtime/MovieSceneTracks/Public/Sections/MovieScene3DTransformSection.h
@@ -27,27 +27,27 @@ public:
 	 *
 	 * @param Time				The position in time within the movie scene
 	 * @param OutTranslation	The evaluated translation.  Note: will remain unchanged if there were no keys to evaluate
-	 * @return true if there were any keys to evaluate.  
+	 * @param OutHasTranslationKeys True if the out has translation keys.
 	 */
-	bool EvalTranslation( float Time, FVector& OutTranslation ) const;
+	void EvalTranslation( float Time, FVector& OutTranslation, TArray<bool>& OutHasTranslationKeys ) const;
 
 	/**
 	 * Evaluates the rotation component of the transform
 	 *
 	 * @param Time				The position in time within the movie scene
 	 * @param OutRotation		The evaluated rotation.  Note: will remain unchanged if there were no keys to evaluate
-	 * @return true if there were any keys to evaluate.  
+	 * @param OutHasRotationKeys True if the out has rotation keys.
 	 */
-	bool EvalRotation( float Time, FRotator& OutRotation ) const;
+	void EvalRotation( float Time, FRotator& OutRotation, TArray<bool>& OutHasRotationKeys ) const;
 
 	/**
 	 * Evaluates the scale component of the transform
 	 *
 	 * @param Time				The position in time within the movie scene
 	 * @param OutScale			The evaluated scale.  Note: will remain unchanged if there were no keys to evaluate
-	 * @return true if there were any keys to evaluate.  
+	 * @param OutHasScaleKeys   True if the out has scale keys.
 	 */
-	bool EvalScale( float Time, FVector& OutScale ) const;
+	void EvalScale( float Time, FVector& OutScale, TArray<bool>& OutHasScaleKeys ) const;
 
 	/** 
 	 * Returns the translation curve for a specific axis
diff --git a/Engine/Source/Runtime/MovieSceneTracks/Public/Tracks/MovieScene3DTransformTrack.h b/Engine/Source/Runtime/MovieSceneTracks/Public/Tracks/MovieScene3DTransformTrack.h
--- a/Engine/Source/Runtime/MovieSceneTracks/Public/Tracks/MovieScene3DTransformTrack.h
+++ b/Engine/Source/Runtime/MovieSceneTracks/Public/Tracks/MovieScene3DTransformTrack.h
@@ -170,5 +170,5 @@ public:
 	 * @param OutHasScaleKeys 	true if a scale key was evaluated
 	 * @return true if anything was evaluated. Note: if false is returned OutBool remains unchanged
 	 */
-	virtual bool Eval( float Position, float LastPosition, FVector& OutTranslation, FRotator& OutRotation, FVector& OutScale, bool& OutHasTranslationKeys, bool& OutHasRotationKeys, bool& OutHasScaleKeys ) const;
+	virtual bool Eval( float Position, float LastPosition, FVector& OutTranslation, FRotator& OutRotation, FVector& OutScale, TArray<bool>& OutHasTranslationKeys, TArray<bool>& OutHasRotationKeys, TArray<bool>& OutHasScaleKeys ) const;
 };
