Ensure color values on color tracks are properly saved and restored on objects.  ColorTrack now looks directly for FColors instead of assuming all colors are linear colors

diff --git a/Engine/Source/Editor/MovieSceneTools/Private/TrackEditors/PropertyTrackEditor.cpp b/Engine/Source/Editor/MovieSceneTools/Private/TrackEditors/PropertyTrackEditor.cpp
--- a/Engine/Source/Editor/MovieSceneTools/Private/TrackEditors/PropertyTrackEditor.cpp
+++ b/Engine/Source/Editor/MovieSceneTools/Private/TrackEditors/PropertyTrackEditor.cpp
@@ -106,6 +106,8 @@ private:
 		// Get the default color of the first instance
 		FLinearColor DefaultColor(0.0f, 0.0f, 0.0f, 0.0f);
 
+		static const FName SlateColorName("SlateColor");
+
 		const TArray<FMovieSceneBinding>& MovieSceneBindings = Sequencer->GetFocusedMovieScene()->GetBindings();
 
 		bool bFoundColor = false;
@@ -123,16 +125,21 @@ private:
 					for (UObject* RuntimeObject : RuntimeObjects)
 					{
 						UProperty* Property = RuntimeObject->GetClass()->FindPropertyByName(CastChecked<UMovieSceneColorTrack>(Track)->GetPropertyName());
-						if (Property != nullptr)
+						UStructProperty* ColorStructProp = Cast<UStructProperty>( Property );
+						if (ColorStructProp && ColorStructProp->Struct )
 						{
-							if (CastChecked<UMovieSceneColorTrack>(Track)->IsSlateColor())
+							if (ColorStructProp->Struct->GetFName() == SlateColorName )
 							{
 								DefaultColor = (*Property->ContainerPtrToValuePtr<FSlateColor>(RuntimeObject)).GetSpecifiedColor();
 							}
-							else
+							else if( ColorStructProp->Struct->GetFName() == NAME_LinearColor )
 							{
 								DefaultColor = *Property->ContainerPtrToValuePtr<FLinearColor>(RuntimeObject);
 							}
+							else
+							{
+								DefaultColor = Property->ContainerPtrToValuePtr<FColor>(RuntimeObject)->ReinterpretAsLinear();
+							}
 							bFoundColor = true;
 							break;
 						}
diff --git a/Engine/Source/Runtime/MovieScene/Private/MovieSceneCommonHelpers.cpp b/Engine/Source/Runtime/MovieScene/Private/MovieSceneCommonHelpers.cpp
--- a/Engine/Source/Runtime/MovieScene/Private/MovieSceneCommonHelpers.cpp
+++ b/Engine/Source/Runtime/MovieScene/Private/MovieSceneCommonHelpers.cpp
@@ -137,7 +137,7 @@ void FTrackInstancePropertyBindings::CallFunction( UObject* InRuntimeObject, voi
 }
 
 
-FTrackInstancePropertyBindings::FPropertyAddress FTrackInstancePropertyBindings::FindPropertyRecursive( UObject* Object, void* BasePointer, UStruct* InStruct, TArray<FString>& InPropertyNames, uint32 Index ) const
+FTrackInstancePropertyBindings::FPropertyAddress FTrackInstancePropertyBindings::FindPropertyRecursive( const UObject* Object, void* BasePointer, UStruct* InStruct, TArray<FString>& InPropertyNames, uint32 Index ) const
 {
 	UProperty* Property = FindField<UProperty>(InStruct, *InPropertyNames[Index]);
 	
@@ -170,7 +170,7 @@ FTrackInstancePropertyBindings::FPropertyAddress FTrackInstancePropertyBindings:
 }
 
 
-FTrackInstancePropertyBindings::FPropertyAddress FTrackInstancePropertyBindings::FindProperty( UObject* InObject, const FString& InPropertyPath ) const
+FTrackInstancePropertyBindings::FPropertyAddress FTrackInstancePropertyBindings::FindProperty( const UObject* InObject, const FString& InPropertyPath ) const
 {
 	TArray<FString> PropertyNames;
 
@@ -202,8 +202,26 @@ void FTrackInstancePropertyBindings::UpdateBindings( const TArray<UObject*>& InR
 		}
 		else
 		{
-			// Dont call potentially invalid functions
+			// Don't call potentially invalid functions
 			RuntimeObjectToFunctionMap.Remove(Object);
 		}
 	}
 }
+
+UProperty* FTrackInstancePropertyBindings::GetProperty(const UObject* Object) const
+{
+	UProperty* Prop = nullptr;
+
+	FPropertyAndFunction PropAndFunction = RuntimeObjectToFunctionMap.FindRef(Object);
+	if( PropAndFunction.PropertyAddress.Property ) 
+	{
+		Prop = PropAndFunction.PropertyAddress.Property;
+	}
+	else
+	{
+		FPropertyAddress PropertyAddress = FindProperty( Object, PropertyPath );
+		Prop = PropertyAddress.Property;
+	}
+
+	return Prop;
+}
diff --git a/Engine/Source/Runtime/MovieScene/Public/MovieSceneCommonHelpers.h b/Engine/Source/Runtime/MovieScene/Public/MovieSceneCommonHelpers.h
--- a/Engine/Source/Runtime/MovieScene/Public/MovieSceneCommonHelpers.h
+++ b/Engine/Source/Runtime/MovieScene/Public/MovieSceneCommonHelpers.h
@@ -60,6 +60,14 @@ public:
 	 */
 	void UpdateBindings( const TArray<UObject*>& InRuntimeObjects );
 
+	/**
+	 * Gets the UProperty that is bound to the track instance
+	 *
+	 * @param Object	The Object that owns the property
+	 * @return			The property on the object if it exists
+	 */
+	UProperty* GetProperty(const UObject* Object) const;
+
 	/**
 	 * Gets the current value of a property on an object
 	 *
@@ -106,8 +114,8 @@ private:
 		{}
 	};
 
-	FPropertyAddress FindPropertyRecursive(UObject* Object, void* BasePointer, UStruct* InStruct, TArray<FString>& InPropertyNames, uint32 Index) const;
-	FPropertyAddress FindProperty(UObject* Object, const FString& InPropertyPath) const;
+	FPropertyAddress FindPropertyRecursive(const UObject* Object, void* BasePointer, UStruct* InStruct, TArray<FString>& InPropertyNames, uint32 Index) const;
+	FPropertyAddress FindProperty(const UObject* Object, const FString& InPropertyPath) const;
 
 private:
 	/** Mapping of objects to bound functions that will be called to update data on the track */
diff --git a/Engine/Source/Runtime/MovieSceneTracks/Private/TrackInstances/MovieSceneColorTrackInstance.cpp b/Engine/Source/Runtime/MovieSceneTracks/Private/TrackInstances/MovieSceneColorTrackInstance.cpp
--- a/Engine/Source/Runtime/MovieSceneTracks/Private/TrackInstances/MovieSceneColorTrackInstance.cpp
+++ b/Engine/Source/Runtime/MovieSceneTracks/Private/TrackInstances/MovieSceneColorTrackInstance.cpp
@@ -8,9 +8,9 @@
 
 
 FMovieSceneColorTrackInstance::FMovieSceneColorTrackInstance( UMovieSceneColorTrack& InColorTrack )
+	: ColorTrack( &InColorTrack )
+	, ColorType( EColorType::Linear )
 {
-	ColorTrack = &InColorTrack;
-
 	PropertyBindings = MakeShareable( new FTrackInstancePropertyBindings( ColorTrack->GetPropertyName(), ColorTrack->GetPropertyPath() ) );
 }
 
@@ -18,16 +18,28 @@ void FMovieSceneColorTrackInstance::SaveState(const TArray<UObject*>& RuntimeObj
 {
 	for( UObject* Object : RuntimeObjects )
 	{
-		if( ColorTrack->IsSlateColor() )
+		if( ColorType == EColorType::Slate )
 		{
 			FSlateColor ColorValue = PropertyBindings->GetCurrentValue<FSlateColor>(Object);
 			InitSlateColorMap.Add(Object, ColorValue);
 		}
-		else
+		else if( ColorType == EColorType::Linear )
 		{
+			UProperty* Property = PropertyBindings->GetProperty(Object);
 			FLinearColor ColorValue = PropertyBindings->GetCurrentValue<FLinearColor>(Object);
 			InitLinearColorMap.Add(Object, ColorValue);
 		}
+		else if( ColorType == EColorType::RegularColor )
+		{
+			UProperty* Property = PropertyBindings->GetProperty(Object);
+			FColor ColorValue = PropertyBindings->GetCurrentValue<FColor>(Object);
+			FLinearColor LinearColorValue = ColorValue.ReinterpretAsLinear();
+			InitLinearColorMap.Add(Object, ColorValue);
+		}
+		else
+		{
+			checkf(0, TEXT("Invalid Color Type"));
+		}
 	}
 }
 
@@ -40,22 +52,27 @@ void FMovieSceneColorTrackInstance::RestoreState(const TArray<UObject*>& Runtime
 			continue;
 		}
 
-		if( ColorTrack->IsSlateColor() )
+		if( ColorType == EColorType::Slate )
 		{
-			FSlateColor *ColorValue = InitSlateColorMap.Find(Object);
+			FSlateColor* ColorValue = InitSlateColorMap.Find(Object);
 			if (ColorValue != NULL)
 			{
 				PropertyBindings->CallFunction(Object, ColorValue);
 			}
 		}
-		else
+		else if( ColorType == EColorType::Linear || ColorType == EColorType::RegularColor )
 		{
-			FLinearColor *ColorValue = InitLinearColorMap.Find(Object);
+			// todo 
+			FLinearColor* ColorValue = InitLinearColorMap.Find(Object);
 			if (ColorValue != NULL)
 			{
 				PropertyBindings->CallFunction(Object, ColorValue);
 			}
 		}
+		else
+		{
+			checkf(0, TEXT("Invalid Color Type"));
+		}
 	}
 
 	PropertyBindings->UpdateBindings( RuntimeObjects );
@@ -65,7 +82,7 @@ void FMovieSceneColorTrackInstance::Update( float Position, float LastPosition,
 {
 	for(UObject* Object : RuntimeObjects)
 	{
-		if( ColorTrack->IsSlateColor() )
+		if( ColorType == EColorType::Slate )
 		{
 			FSlateColor ColorValue = PropertyBindings->GetCurrentValue<FSlateColor>(Object);
 			FLinearColor LinearColor = ColorValue.GetSpecifiedColor();
@@ -77,7 +94,7 @@ void FMovieSceneColorTrackInstance::Update( float Position, float LastPosition,
 		}
 		else
 		{
-			FLinearColor ColorValue = PropertyBindings->GetCurrentValue<FLinearColor>(Object);
+			FLinearColor ColorValue = ColorType == EColorType::Linear ? PropertyBindings->GetCurrentValue<FLinearColor>(Object) : PropertyBindings->GetCurrentValue<FColor>(Object).ReinterpretAsLinear();
 			if(ColorTrack->Eval(Position, LastPosition, ColorValue))
 			{
 				PropertyBindings->CallFunction(Object, &ColorValue);
@@ -88,5 +105,33 @@ void FMovieSceneColorTrackInstance::Update( float Position, float LastPosition,
 
 void FMovieSceneColorTrackInstance::RefreshInstance( const TArray<UObject*>& RuntimeObjects, class IMovieScenePlayer& Player )
 {
-	PropertyBindings->UpdateBindings( RuntimeObjects );
+	if( RuntimeObjects.Num() > 0 )
+	{
+		PropertyBindings->UpdateBindings( RuntimeObjects );
+
+		// Cache off what type of color this is. Just examine the first object since the property should be the same
+		UProperty* ColorProp = PropertyBindings->GetProperty( RuntimeObjects[0] );
+
+		const UStructProperty* StructProp = Cast<const UStructProperty>(ColorProp);
+		if (StructProp && StructProp->Struct)
+		{
+			FName StructName = StructProp->Struct->GetFName();
+
+			static const FName SlateColor("SlateColor");
+
+			if( StructName == NAME_Color )
+			{
+				ColorType = EColorType::RegularColor;
+			}
+			else if( StructName == SlateColor )
+			{
+				ColorType = EColorType::Slate;
+			}
+			else
+			{
+				ColorType = EColorType::Linear;
+			}
+		}
+	}
+
 }
diff --git a/Engine/Source/Runtime/MovieSceneTracks/Private/Tracks/MovieSceneColorTrack.cpp b/Engine/Source/Runtime/MovieSceneTracks/Private/Tracks/MovieSceneColorTrack.cpp
--- a/Engine/Source/Runtime/MovieSceneTracks/Private/Tracks/MovieSceneColorTrack.cpp
+++ b/Engine/Source/Runtime/MovieSceneTracks/Private/Tracks/MovieSceneColorTrack.cpp
@@ -8,7 +8,6 @@
 
 UMovieSceneColorTrack::UMovieSceneColorTrack( const FObjectInitializer& ObjectInitializer )
 	: Super( ObjectInitializer )
-	, bIsSlateColor( false )
 {
 }
 
@@ -25,8 +24,6 @@ TSharedPtr<IMovieSceneTrackInstance> UMovieSceneColorTrack::CreateInstance()
 
 bool UMovieSceneColorTrack::AddKeyToSection( float Time, const FColorKey& Key )
 {
-	bIsSlateColor = Key.bIsSlateColor;
-
 	const UMovieSceneSection* NearestSection = MovieSceneHelpers::FindNearestSectionAtTime( Sections, Time );
 	if (!NearestSection || Key.bAddKeyEvenIfUnchanged || CastChecked<UMovieSceneColorSection>(NearestSection)->NewKeyIsNewData(Time, Key.Value))
 	{
diff --git a/Engine/Source/Runtime/MovieSceneTracks/Public/TrackInstances/MovieSceneColorTrackInstance.h b/Engine/Source/Runtime/MovieSceneTracks/Public/TrackInstances/MovieSceneColorTrackInstance.h
--- a/Engine/Source/Runtime/MovieSceneTracks/Public/TrackInstances/MovieSceneColorTrackInstance.h
+++ b/Engine/Source/Runtime/MovieSceneTracks/Public/TrackInstances/MovieSceneColorTrackInstance.h
@@ -16,7 +16,6 @@ class FMovieSceneColorTrackInstance
 	: public IMovieSceneTrackInstance
 {
 public:
-
 	FMovieSceneColorTrackInstance( UMovieSceneColorTrack& InColorTrack );
 
 	/** IMovieSceneTrackInstance interface */
@@ -27,6 +26,15 @@ public:
 	virtual void ClearInstance( IMovieScenePlayer& Player ) override {}
 
 private:
+	enum class EColorType : uint8
+	{
+		/** FSlateColor */
+		Slate, 
+		/** FLinearColor */
+		Linear,
+		/** FColor */
+		RegularColor,
+	};
 
 	/** The track being instanced */
 	UMovieSceneColorTrack* ColorTrack;
@@ -37,4 +45,7 @@ private:
 	/** Map from object to initial state */
 	TMap< TWeakObjectPtr<UObject>, FSlateColor > InitSlateColorMap;
 	TMap< TWeakObjectPtr<UObject>, FLinearColor > InitLinearColorMap;
+
+	/** The type of color primitive being animated */
+	EColorType ColorType;
 };
diff --git a/Engine/Source/Runtime/MovieSceneTracks/Public/Tracks/MovieSceneColorTrack.h b/Engine/Source/Runtime/MovieSceneTracks/Public/Tracks/MovieSceneColorTrack.h
--- a/Engine/Source/Runtime/MovieSceneTracks/Public/Tracks/MovieSceneColorTrack.h
+++ b/Engine/Source/Runtime/MovieSceneTracks/Public/Tracks/MovieSceneColorTrack.h
@@ -46,10 +46,7 @@ public:
 	 * @return true if anything was evaluated. Note: if false is returned OutColor remains unchanged
 	 */
 	virtual bool Eval( float Position, float LastPostion, FLinearColor& OutColor ) const;
-
-	bool IsSlateColor() const { return bIsSlateColor; }
-
 private:
 	UPROPERTY()
-	bool bIsSlateColor;
+	bool bIsSlateColor_DEPRECATED;
 };
