Sequencer: Fix non-unity build.

diff --git a/Engine/Source/Runtime/MovieSceneTracks/Private/Tracks/SubMovieSceneTrack.cpp b/Engine/Source/Runtime/MovieSceneTracks/Private/Tracks/SubMovieSceneTrack.cpp
--- a/Engine/Source/Runtime/MovieSceneTracks/Private/Tracks/SubMovieSceneTrack.cpp
+++ b/Engine/Source/Runtime/MovieSceneTracks/Private/Tracks/SubMovieSceneTrack.cpp
@@ -1,6 +1,7 @@
 // Copyright 1998-2015 Epic Games, Inc. All Rights Reserved.
 
 #include "MovieSceneTracksPrivatePCH.h"
+#include "MovieSceneSequence.h"
 #include "SubMovieSceneTrack.h"
 #include "IMovieScenePlayer.h"
 #include "SubMovieSceneSection.h"
