MovieScene: Started to implement event tracks

diff --git a/Engine/Source/Runtime/MovieSceneTracks/Private/Sections/MovieSceneEventSection.cpp b/Engine/Source/Runtime/MovieSceneTracks/Private/Sections/MovieSceneEventSection.cpp
new file mode 100644
--- /dev/null
+++ b/Engine/Source/Runtime/MovieSceneTracks/Private/Sections/MovieSceneEventSection.cpp
@@ -0,0 +1,45 @@
+// Copyright 1998-2015 Epic Games, Inc. All Rights Reserved.
+
+#include "MovieSceneTracksPrivatePCH.h"
+#include "MovieSceneEventSection.h"
+
+
+/* UMovieSceneSection overrides
+ *****************************************************************************/
+
+void UMovieSceneEventSection::DilateSection(float DilationFactor, float Origin, TSet<FKeyHandle>& KeyHandles)
+{
+	Super::DilateSection(DilationFactor, Origin, KeyHandles);
+
+	for (auto& Key : Keys)
+	{
+		Key.Time = Key.Time * DilationFactor - Origin;
+	}
+}
+
+
+void UMovieSceneEventSection::GetKeyHandles(TSet<FKeyHandle>& KeyHandles) const
+{
+}
+
+
+void UMovieSceneEventSection::GetSnapTimes(TArray<float>& OutSnapTimes, bool bGetSectionBorders) const
+{
+	Super::GetSnapTimes(OutSnapTimes, bGetSectionBorders);
+
+	for (auto& Key : Keys)
+	{
+		OutSnapTimes.Add(Key.Time);
+	}
+}
+
+
+void UMovieSceneEventSection::MoveSection(float DeltaPosition, TSet<FKeyHandle>& KeyHandles)
+{
+	Super::MoveSection(DeltaPosition, KeyHandles);
+
+	for (auto& Key : Keys)
+	{
+		Key.Time += DeltaPosition;
+	}
+}
diff --git a/Engine/Source/Runtime/MovieSceneTracks/Private/TrackInstances/MovieSceneEventTrackInstance.cpp b/Engine/Source/Runtime/MovieSceneTracks/Private/TrackInstances/MovieSceneEventTrackInstance.cpp
new file mode 100644
--- /dev/null
+++ b/Engine/Source/Runtime/MovieSceneTracks/Private/TrackInstances/MovieSceneEventTrackInstance.cpp
@@ -0,0 +1,30 @@
+// Copyright 1998-2015 Epic Games, Inc. All Rights Reserved.
+
+#include "MovieSceneTracksPrivatePCH.h"
+#include "MovieSceneEventTrack.h"
+#include "MovieSceneEventTrackInstance.h"
+
+
+/* FMovieSceneEventTrackInstance structors
+ *****************************************************************************/
+
+FMovieSceneEventTrackInstance::FMovieSceneEventTrackInstance(UMovieSceneEventTrack& InEventTrack)
+{
+	EventTrack = &InEventTrack;
+}
+
+
+/* IMovieSceneTrackInstance interface
+ *****************************************************************************/
+
+void FMovieSceneEventTrackInstance::Update(float Position, float LastPosition, const TArray<UObject*>& RuntimeObjects, IMovieScenePlayer& Player)
+{
+	if (Position > LastPosition)
+	{
+		EventTrack->TriggerEvents(TRange<float>(LastPosition, Position), false);	
+	}
+	else if (Position < LastPosition)
+	{
+		EventTrack->TriggerEvents(TRange<float>(Position, LastPosition), true);
+	}
+}
diff --git a/Engine/Source/Runtime/MovieSceneTracks/Private/Tracks/MovieSceneEventTrack.cpp b/Engine/Source/Runtime/MovieSceneTracks/Private/Tracks/MovieSceneEventTrack.cpp
new file mode 100644
--- /dev/null
+++ b/Engine/Source/Runtime/MovieSceneTracks/Private/Tracks/MovieSceneEventTrack.cpp
@@ -0,0 +1,126 @@
+// Copyright 1998-2015 Epic Games, Inc. All Rights Reserved.
+
+#include "MovieSceneTracksPrivatePCH.h"
+#include "MovieSceneEventSection.h"
+#include "MovieSceneEventTrack.h"
+#include "MovieSceneEventTrackInstance.h"
+
+
+/* UMovieSceneEventTrack interface
+ *****************************************************************************/
+
+bool UMovieSceneEventTrack::AddKeyToSection(float Time, FName EventName)
+{
+	// @todo sequencer: gmp: add key to event section
+	// should be inserted into sorted array
+	// can sections overlap?
+	return false;
+}
+
+
+void UMovieSceneEventTrack::TriggerEvents(TRange<float> TimeRange, bool Backwards)
+{
+	if ((!Backwards && !bFireEventsWhenForwards) ||
+		(Backwards && !bFireEventsWhenBackwards))
+	{
+		return;
+	}
+
+	// @todo sequencer: gmp: implement event track triggering
+	/*
+	ULevel* Level = GetLevel();
+	ALevelScriptActor* LevelScriptActor = Level->LevelScriptActor;
+
+	if (LevelScriptActor == nullptr)
+	{
+		return;
+	}
+
+	// for each key...
+	for ()
+	{
+		UFunction* EventFunction = LevelScriptActor->FindFunction(EventName);
+
+		if (EventFunction != nullptr)
+		{
+			if(EventFunction->NumParms == 0)
+			{
+				LevelScriptActor->ProcessEvent(EventFunction, nullptr);
+			}
+			else
+			{
+				// @todo sequencer: gmp: add external log category for MovieScene
+				//UE_LOG(LogMovieScene, Log, TEXT("NotifyEventTriggered: Function '%s' does not have zero parameters."), *EventFuncName.ToString());
+			}
+		}
+		else
+		{
+			// @todo sequencer: gmp: add external log category for MovieScene
+			//UE_LOG(LogMovieScene, Log, TEXT("NotifyEventTriggered: Unable to find function '%s'"), *EventFuncName.ToString());
+		}
+	}*/
+}
+
+
+/* UMovieSceneTrack interface
+ *****************************************************************************/
+
+TSharedPtr<IMovieSceneTrackInstance> UMovieSceneEventTrack::CreateInstance()
+{
+	return MakeShareable(new FMovieSceneEventTrackInstance(*this));
+}
+
+
+UMovieSceneSection* UMovieSceneEventTrack::CreateNewSection()
+{
+	return NewObject<UMovieSceneSection>(this, UMovieSceneEventSection::StaticClass(), NAME_None, RF_Transactional);
+}
+
+
+const TArray<UMovieSceneSection*>& UMovieSceneEventTrack::GetAllSections() const
+{
+	return Sections;
+}
+
+
+TRange<float> UMovieSceneEventTrack::GetSectionBoundaries() const
+{
+	TRange<float> SectionBoundaries = TRange<float>::Empty();
+
+	for (auto& Section : Sections)
+	{
+		SectionBoundaries = TRange<float>::Hull(SectionBoundaries, Section->GetRange());
+	}
+
+	return SectionBoundaries;
+}
+
+
+FName UMovieSceneEventTrack::GetTrackName() const
+{
+	return TrackName;
+}
+
+
+bool UMovieSceneEventTrack::HasSection(UMovieSceneSection* Section) const
+{
+	return Sections.Contains(Section);
+}
+
+
+bool UMovieSceneEventTrack::IsEmpty() const
+{
+	return (Sections.Num() == 0);
+}
+
+
+void UMovieSceneEventTrack::RemoveAllAnimationData()
+{
+	Sections.Empty();
+}
+
+
+void UMovieSceneEventTrack::RemoveSection(UMovieSceneSection* Section)
+{
+	Sections.Remove(Section);
+}
diff --git a/Engine/Source/Runtime/MovieSceneTracks/Public/Sections/MovieSceneEventSection.h b/Engine/Source/Runtime/MovieSceneTracks/Public/Sections/MovieSceneEventSection.h
new file mode 100644
--- /dev/null
+++ b/Engine/Source/Runtime/MovieSceneTracks/Public/Sections/MovieSceneEventSection.h
@@ -0,0 +1,49 @@
+// Copyright 1998-2015 Epic Games, Inc. All Rights Reserved.
+
+#pragma once
+
+#include "MovieSceneEventSection.generated.h"
+
+
+/**
+ * Structure for event section keys.
+ */
+USTRUCT()
+struct FMovieSceneEventSectionKey
+{
+	GENERATED_BODY()
+
+	/** The name of the event to be triggered. */
+	UPROPERTY(EditAnywhere, Category=EventTrackKey)
+	FName EventName;
+
+	/** The time at which the event should be triggered. */
+	UPROPERTY()
+	float Time;
+};
+
+
+/**
+ * Implements a section in movie scene event tracks.
+ */
+UCLASS( MinimalAPI )
+class UMovieSceneEventSection
+	: public UMovieSceneSection
+{
+	GENERATED_BODY()
+
+public:
+
+	// UMovieSceneSection interface
+
+	virtual void DilateSection(float DilationFactor, float Origin, TSet<FKeyHandle>& KeyHandles) override;
+	virtual void GetKeyHandles(TSet<FKeyHandle>& KeyHandles) const override;
+	virtual void GetSnapTimes(TArray<float>& OutSnapTimes, bool bGetSectionBorders) const override;
+	virtual void MoveSection(float DeltaPosition, TSet<FKeyHandle>& KeyHandles) override;
+
+private:
+
+	/** The section's keys. */
+	UPROPERTY()
+	TArray<FMovieSceneEventSectionKey> Keys;
+};
diff --git a/Engine/Source/Runtime/MovieSceneTracks/Public/TrackInstances/MovieSceneEventTrackInstance.h b/Engine/Source/Runtime/MovieSceneTracks/Public/TrackInstances/MovieSceneEventTrackInstance.h
new file mode 100644
--- /dev/null
+++ b/Engine/Source/Runtime/MovieSceneTracks/Public/TrackInstances/MovieSceneEventTrackInstance.h
@@ -0,0 +1,43 @@
+// Copyright 1998-2015 Epic Games, Inc. All Rights Reserved.
+
+#pragma once
+
+#include "IMovieSceneTrackInstance.h"
+
+
+class UMovieSceneEventTrack;
+
+
+/**
+ * Instance of a UMovieSceneEventTrack
+ */
+class FMovieSceneEventTrackInstance
+	: public IMovieSceneTrackInstance
+{
+public:
+
+	/**
+	 * Creates and initializes a new instance.
+	 *
+	 * @param InEventTrack The event track to create an instance for.
+	 */
+	FMovieSceneEventTrackInstance(UMovieSceneEventTrack& InEventTrack );
+
+	/** Virtual destructor. */
+	virtual ~FMovieSceneEventTrackInstance() { }
+
+public:
+
+	// IMovieSceneTrackInstance interface
+
+	virtual void ClearInstance(IMovieScenePlayer& Player) override { }
+	virtual void RefreshInstance(const TArray<UObject*>& RuntimeObjects, IMovieScenePlayer& Player) override { }
+	virtual void RestoreState(const TArray<UObject*>& RuntimeObjects) override { }
+	virtual void SaveState(const TArray<UObject*>& RuntimeObjects) override { }
+	virtual void Update(float Position, float LastPosition, const TArray<UObject*>& RuntimeObjects, IMovieScenePlayer& Player) override;
+
+private:
+
+	/** Track that is being instanced */
+	UMovieSceneEventTrack* EventTrack;
+};
diff --git a/Engine/Source/Runtime/MovieSceneTracks/Public/Tracks/MovieSceneEventTrack.h b/Engine/Source/Runtime/MovieSceneTracks/Public/Tracks/MovieSceneEventTrack.h
new file mode 100644
--- /dev/null
+++ b/Engine/Source/Runtime/MovieSceneTracks/Public/Tracks/MovieSceneEventTrack.h
@@ -0,0 +1,73 @@
+// Copyright 1998-2015 Epic Games, Inc. All Rights Reserved.
+
+#pragma once
+
+#include "MovieSceneTrack.h"
+#include "MovieSceneEventTrack.generated.h"
+
+
+/**
+ * Implements a movie scene track that triggers discrete events during playback.
+ */
+UCLASS(MinimalAPI)
+class UMovieSceneEventTrack
+	: public UMovieSceneTrack
+{
+	GENERATED_BODY()
+
+public:
+
+	/**
+	 * Adds an event to the appropriate section.
+	 *
+	 * This method will create a new section if no appropriate one exists.
+	 *
+	 * @param Time The time at which the event should be triggered.
+	 * @param EventName The name of the event to be triggered.
+	 */
+	bool AddKeyToSection(float Time, FName EventName);
+
+	/**
+	 * Trigger the events that fall into the given time range.
+	 *
+	 * @param TimeRange The time range to trigger events in.
+	 * @param Backwards Whether the events should be triggered in reverse order.
+	 */
+	void TriggerEvents(TRange<float> TimeRange, bool Backwards);
+
+public:
+
+	// UMovieSceneTrack interface
+
+	virtual TSharedPtr<IMovieSceneTrackInstance> CreateInstance() override;
+	virtual UMovieSceneSection* CreateNewSection() override;
+	virtual const TArray<UMovieSceneSection*>& GetAllSections() const override;
+	virtual TRange<float> GetSectionBoundaries() const override;
+	virtual FName GetTrackName() const override;
+	virtual bool HasSection(UMovieSceneSection* Section) const override;
+	virtual bool IsEmpty() const override;
+	virtual void RemoveAllAnimationData() override;
+	virtual void RemoveSection(UMovieSceneSection* Section) override;
+
+private:
+
+	/** If events should be fired when passed playing the sequence forwards. */
+	UPROPERTY(EditAnywhere, Category=TrackEvent)
+	uint32 bFireEventsWhenForwards:1;
+
+	/** If events should be fired when passed playing the sequence backwards. */
+	UPROPERTY(EditAnywhere, Category=TrackEvent)
+	uint32 bFireEventsWhenBackwards:1;
+
+	/** If checked each key's event name is the exact name of the custom event function in level script that will be called */
+	UPROPERTY(EditAnywhere, Category=TrackEvent)
+	uint32 bUseCustomEventName:1;
+
+	/** The track's sections. */
+	UPROPERTY()
+	TArray<UMovieSceneSection*> Sections;
+
+	/** Name of this track. */
+	UPROPERTY()
+	FName TrackName;
+};
