Sequencer - Refactor the FMovieScenePropertyTrack - Fix Win32 build.

diff --git a/Engine/Source/Runtime/MovieSceneTracks/Private/Tracks/MovieSceneVectorTrack.cpp b/Engine/Source/Runtime/MovieSceneTracks/Private/Tracks/MovieSceneVectorTrack.cpp
--- a/Engine/Source/Runtime/MovieSceneTracks/Private/Tracks/MovieSceneVectorTrack.cpp
+++ b/Engine/Source/Runtime/MovieSceneTracks/Private/Tracks/MovieSceneVectorTrack.cpp
@@ -6,7 +6,7 @@
 #include "IMovieScenePlayer.h"
 #include "MovieSceneVectorTrackInstance.h"
 
-FVectorKey::FVectorKey( FVector2D InValue, FName InCurveName, bool InbAddKeyEvenIfUnChanged )
+FVectorKey::FVectorKey( const FVector2D& InValue, FName InCurveName, bool InbAddKeyEvenIfUnChanged )
 {
 	Value.X = InValue.X;
 	Value.Y = InValue.Y;
@@ -17,7 +17,7 @@ FVectorKey::FVectorKey( FVector2D InValue, FName InCurveName, bool InbAddKeyEven
 	bAddKeyEvenIfUnchanged = InbAddKeyEvenIfUnChanged;
 }
 
-FVectorKey::FVectorKey( FVector InValue, FName InCurveName, bool InbAddKeyEvenIfUnChanged )
+FVectorKey::FVectorKey( const FVector& InValue, FName InCurveName, bool InbAddKeyEvenIfUnChanged )
 {
 	Value.X = InValue.X;
 	Value.Y = InValue.Y;
@@ -28,7 +28,7 @@ FVectorKey::FVectorKey( FVector InValue, FName InCurveName, bool InbAddKeyEvenIf
 	bAddKeyEvenIfUnchanged = InbAddKeyEvenIfUnChanged;
 }
 
-FVectorKey::FVectorKey( FVector4 InValue, FName InCurveName, bool InbAddKeyEvenIfUnChanged )
+FVectorKey::FVectorKey( const FVector4& InValue, FName InCurveName, bool InbAddKeyEvenIfUnChanged )
 {
 	Value = InValue;
 	ChannelsUsed = 4;
@@ -73,7 +73,7 @@ bool UMovieSceneVectorTrack::AddKeyToSection( float Time, const FVector4& Value,
 	return false;
 }
 
-bool UMovieSceneVectorTrack::AddKeyToSection( float Time, FVectorKey Key )
+bool UMovieSceneVectorTrack::AddKeyToSection( float Time, const FVectorKey& Key )
 {
 	return AddKeyToSection(Time, Key.Value, Key.ChannelsUsed, Key.CurveName, Key.bAddKeyEvenIfUnchanged );
 }
diff --git a/Engine/Source/Runtime/MovieSceneTracks/Public/Tracks/MovieSceneVectorTrack.h b/Engine/Source/Runtime/MovieSceneTracks/Public/Tracks/MovieSceneVectorTrack.h
--- a/Engine/Source/Runtime/MovieSceneTracks/Public/Tracks/MovieSceneVectorTrack.h
+++ b/Engine/Source/Runtime/MovieSceneTracks/Public/Tracks/MovieSceneVectorTrack.h
@@ -12,9 +12,9 @@
 struct MOVIESCENETRACKS_API FVectorKey
 {
 	FVectorKey() { }
-	FVectorKey( FVector InValue, FName InCurveName, bool InbAddKeyEvenIfUnChanged );
-	FVectorKey( FVector2D InValue, FName InCurveName, bool InbAddKeyEvenIfUnChanged );
-	FVectorKey( FVector4 InValue, FName InCurveName, bool InbAddKeyEvenIfUnChanged );
+	FVectorKey( const FVector& InValue, FName InCurveName, bool InbAddKeyEvenIfUnChanged );
+	FVectorKey( const FVector2D& InValue, FName InCurveName, bool InbAddKeyEvenIfUnChanged );
+	FVectorKey( const FVector4& InValue, FName InCurveName, bool InbAddKeyEvenIfUnChanged );
 
 	FVector4 Value;
 	uint32 ChannelsUsed;
@@ -37,15 +37,11 @@ public:
 	/**
 	 * Adds a key to a section.  Will create the section if it doesn't exist
 	 *
-	 * @param ObjectHandle		Handle to the object(s) being changed
-	 * @param InPropertyName	The name of the property being manipulated.  @todo Sequencer - Could be a UFunction name
 	 * @param Time				The time relative to the owning movie scene where the section should be
-	 * @param InKey				The vector key to add
-	 * @param InChannelsUsed	The number of channels used, 2, 3, or 4, determining which type of vector this is
+	 * @param Key				The vector key to add
 	 * @return True if the key was successfully added.
 	 */
-
-	virtual bool AddKeyToSection( float Time, FVectorKey Key );
+	virtual bool AddKeyToSection( float Time, const FVectorKey& Key );
 	
 	/**
 	 * Evaluates the track at the playback position
