Sequencer: Fix crash when scrubbing audio sometimes.

diff --git a/Engine/Source/Runtime/MovieSceneTracks/Private/TrackInstances/MovieSceneAudioTrackInstance.cpp b/Engine/Source/Runtime/MovieSceneTracks/Private/TrackInstances/MovieSceneAudioTrackInstance.cpp
--- a/Engine/Source/Runtime/MovieSceneTracks/Private/TrackInstances/MovieSceneAudioTrackInstance.cpp
+++ b/Engine/Source/Runtime/MovieSceneTracks/Private/TrackInstances/MovieSceneAudioTrackInstance.cpp
@@ -39,19 +39,22 @@ void FMovieSceneAudioTrackInstance::Update( float Position, float LastPosition,
 
 				for (int32 ActorIndex = 0; ActorIndex < Actors.Num(); ++ActorIndex)
 				{
-					UAudioComponent* Component = GetAudioComponent(Actors[ActorIndex], RowIndex);
+					TWeakObjectPtr<UAudioComponent> Component = GetAudioComponent(Actors[ActorIndex], RowIndex);
 
 					bool bComponentIsPlaying = false;
-					for (int32 i = 0; i < MovieSceneAudioSections.Num(); ++i)
+					if (Component.IsValid())
 					{
-						UMovieSceneAudioSection* AudioSection = MovieSceneAudioSections[i];
-						if (AudioSection->IsTimeWithinAudioRange(Position))
+						for (int32 i = 0; i < MovieSceneAudioSections.Num(); ++i)
 						{
-							if (!AudioSection->IsTimeWithinAudioRange(LastPosition) || !Component->IsPlaying())
+							UMovieSceneAudioSection* AudioSection = MovieSceneAudioSections[i];
+							if (AudioSection->IsTimeWithinAudioRange(Position))
 							{
-								PlaySound(AudioSection, Component, Position);
+								if (!AudioSection->IsTimeWithinAudioRange(LastPosition) || !Component->IsPlaying())
+								{
+									PlaySound(AudioSection, Component, Position);
+								}
+								bComponentIsPlaying = true;
 							}
-							bComponentIsPlaying = true;
 						}
 					}
 
@@ -79,14 +82,16 @@ void FMovieSceneAudioTrackInstance::Update( float Position, float LastPosition,
 				
 				for (int32 ActorIndex = 0; ActorIndex < Actors.Num(); ++ActorIndex)
 				{
-					UAudioComponent* Component = GetAudioComponent(Actors[ActorIndex], RowIndex);
-
-					if (AudioSection->IsTimeWithinAudioRange(Position) && !Component->IsPlaying())
+					TWeakObjectPtr<UAudioComponent> Component = GetAudioComponent(Actors[ActorIndex], RowIndex);
+					if (Component.IsValid())
 					{
-						PlaySound(AudioSection, Component, Position);
-						// Fade out the sound at the same volume in order to simply
-						// set a short duration on the sound, far from ideal soln
-						Component->FadeOut(AudioTrackConstants::ScrubDuration, 1.f);
+						if (AudioSection->IsTimeWithinAudioRange(Position) && !Component->IsPlaying())
+						{
+							PlaySound(AudioSection, Component, Position);
+							// Fade out the sound at the same volume in order to simply
+							// set a short duration on the sound, far from ideal soln
+							Component->FadeOut(AudioTrackConstants::ScrubDuration, 1.f);
+						}
 					}
 				}
 			}
@@ -105,39 +110,48 @@ void FMovieSceneAudioTrackInstance::Update( float Position, float LastPosition,
 		{
 			for (int32 ActorIndex = 0; ActorIndex < Actors.Num(); ++ActorIndex)
 			{
-				UAudioComponent* Component = GetAudioComponent(Actors[ActorIndex], RowIndex);
-				if (Component->IsPlaying())
+				TWeakObjectPtr<UAudioComponent> Component = GetAudioComponent(Actors[ActorIndex], RowIndex);
+				if (Component.IsValid())
 				{
-					FAudioDevice* AudioDevice = Component->GetAudioDevice();
-					FActiveSound* ActiveSound = AudioDevice->FindActiveSound(Component);
-					ActiveSound->bLocationDefined = true;
-					ActiveSound->Transform = Actors[ActorIndex]->GetTransform();
+					if (Component->IsPlaying())
+					{
+						FAudioDevice* AudioDevice = Component->GetAudioDevice();
+						FActiveSound* ActiveSound = AudioDevice->FindActiveSound(Component.Get());
+						ActiveSound->bLocationDefined = true;
+						ActiveSound->Transform = Actors[ActorIndex]->GetTransform();
+					}
 				}
 			}
 		}
 	}
 }
 
-void FMovieSceneAudioTrackInstance::PlaySound(UMovieSceneAudioSection* AudioSection, UAudioComponent* Component, float Time)
+void FMovieSceneAudioTrackInstance::PlaySound(UMovieSceneAudioSection* AudioSection, TWeakObjectPtr<UAudioComponent> Component, float Time)
 {
-	float PitchMultiplier = 1.f / AudioSection->GetAudioDilationFactor();
-	Component->bAllowSpatialization = !AudioTrack->IsAMasterTrack();
-	Component->Stop();
-	Component->SetSound(AudioSection->GetSound());
-	Component->SetVolumeMultiplier(1.f);
-	Component->SetPitchMultiplier(PitchMultiplier);
-	Component->bIsUISound = true;
-	Component->Play(Time - AudioSection->GetAudioStartTime());
+	if (Component.IsValid())
+	{
+		float PitchMultiplier = 1.f / AudioSection->GetAudioDilationFactor();
+		Component->bAllowSpatialization = !AudioTrack->IsAMasterTrack();
+		Component->Stop();
+		Component->SetSound(AudioSection->GetSound());
+		Component->SetVolumeMultiplier(1.f);
+		Component->SetPitchMultiplier(PitchMultiplier);
+		Component->bIsUISound = true;
+		Component->Play(Time - AudioSection->GetAudioStartTime());
+	}
 }
 
 void FMovieSceneAudioTrackInstance::StopSound(int32 RowIndex)
 {
 	if (RowIndex >= PlaybackAudioComponents.Num()) {return;}
 
-	TMap<AActor*, UAudioComponent*>& AudioComponents = PlaybackAudioComponents[RowIndex];
-	for (TMap<AActor*, UAudioComponent*>::TIterator It(AudioComponents); It; ++It)
+	TMap<AActor*, TWeakObjectPtr<UAudioComponent>>& AudioComponents = PlaybackAudioComponents[RowIndex];
+	for (TMap<AActor*, TWeakObjectPtr<UAudioComponent>>::TIterator It(AudioComponents); It; ++It)
 	{
-		It.Value()->Stop();
+		if (It.Value().IsValid())
+		{
+			It.Value()->Stop();
+		}
 	}
 }
 
@@ -145,33 +159,38 @@ void FMovieSceneAudioTrackInstance::StopAllSounds()
 {
 	for (int32 i = 0; i < PlaybackAudioComponents.Num(); ++i)
 	{
-		TMap<AActor*, UAudioComponent*>& AudioComponents = PlaybackAudioComponents[i];
-		for (TMap<AActor*, UAudioComponent*>::TIterator It(AudioComponents); It; ++It)
+		TMap<AActor*, TWeakObjectPtr<UAudioComponent>>& AudioComponents = PlaybackAudioComponents[i];
+		for (TMap<AActor*, TWeakObjectPtr<UAudioComponent>>::TIterator It(AudioComponents); It; ++It)
 		{
-			It.Value()->Stop();
+			if (It.Value().IsValid())
+			{
+				It.Value()->Stop();
+			}
 		}
 	}
 }
 
-UAudioComponent* FMovieSceneAudioTrackInstance::GetAudioComponent(AActor* Actor, int32 RowIndex)
+TWeakObjectPtr<UAudioComponent> FMovieSceneAudioTrackInstance::GetAudioComponent(AActor* Actor, int32 RowIndex)
 {
 	if (RowIndex + 1 > PlaybackAudioComponents.Num())
 	{
 		while (PlaybackAudioComponents.Num() < RowIndex + 1)
 		{
-			TMap<AActor*, UAudioComponent*> AudioComponentMap;
+			TMap<AActor*, TWeakObjectPtr<UAudioComponent>> AudioComponentMap;
 			PlaybackAudioComponents.Add(AudioComponentMap);
 		}
 	}
 
-	UAudioComponent*& AudioComponent = PlaybackAudioComponents[RowIndex].FindOrAdd(Actor);
-	if (AudioComponent == NULL)
+	if (PlaybackAudioComponents[RowIndex].Find(Actor) == nullptr || !(*PlaybackAudioComponents[RowIndex].Find(Actor)).IsValid())
 	{
 		USoundCue* TempPlaybackAudioCue = NewObject<USoundCue>();
 		
-		AudioComponent = FAudioDevice::CreateComponent(TempPlaybackAudioCue, NULL, Actor, false, false);
+		UAudioComponent* AudioComponent = FAudioDevice::CreateComponent(TempPlaybackAudioCue, NULL, Actor, false, false);
+
+		PlaybackAudioComponents[RowIndex].Add(Actor);
+		PlaybackAudioComponents[RowIndex][Actor] = TWeakObjectPtr<UAudioComponent>(AudioComponent);
 	}
-	return AudioComponent;
+	return *PlaybackAudioComponents[RowIndex].Find(Actor);
 }
 
 TArray<AActor*> FMovieSceneAudioTrackInstance::GetRuntimeActors(const TArray<UObject*>& RuntimeObjects) const
diff --git a/Engine/Source/Runtime/MovieSceneTracks/Public/TrackInstances/MovieSceneAudioTrackInstance.h b/Engine/Source/Runtime/MovieSceneTracks/Public/TrackInstances/MovieSceneAudioTrackInstance.h
--- a/Engine/Source/Runtime/MovieSceneTracks/Public/TrackInstances/MovieSceneAudioTrackInstance.h
+++ b/Engine/Source/Runtime/MovieSceneTracks/Public/TrackInstances/MovieSceneAudioTrackInstance.h
@@ -29,7 +29,7 @@ public:
 private:
 
 	/** Plays the sound of the given section at the given time */
-	void PlaySound(class UMovieSceneAudioSection* AudioSection, UAudioComponent* Component, float Time);
+	void PlaySound(class UMovieSceneAudioSection* AudioSection, TWeakObjectPtr<UAudioComponent> Component, float Time);
 	
 	/** Stops all components for the given row index */
 	void StopSound(int32 RowIndex);
@@ -38,7 +38,7 @@ private:
 	void StopAllSounds();
 
 	/** Gets the audio component component for the actor and row index, creating it if necessary */
-	UAudioComponent* GetAudioComponent(AActor* Actor, int32 RowIndex);
+	TWeakObjectPtr<UAudioComponent> GetAudioComponent(AActor* Actor, int32 RowIndex);
 
 	/** Utility function for getting actors from objects array */
 	TArray<AActor*> GetRuntimeActors(const TArray<UObject*>& RuntimeObjects) const;
@@ -49,5 +49,5 @@ private:
 	UMovieSceneAudioTrack* AudioTrack;
 
 	/** Audio components to play our audio tracks with, one per row per actor */
-	TArray< TMap<AActor*, UAudioComponent*> > PlaybackAudioComponents;
+	TArray< TMap<AActor*, TWeakObjectPtr<UAudioComponent> > > PlaybackAudioComponents;
 };
