Sequencer: Add AddSection() to MovieSceneTrack

diff --git a/Engine/Source/Runtime/MovieScene/Public/MovieSceneTrack.h b/Engine/Source/Runtime/MovieScene/Public/MovieSceneTrack.h
--- a/Engine/Source/Runtime/MovieScene/Public/MovieSceneTrack.h
+++ b/Engine/Source/Runtime/MovieScene/Public/MovieSceneTrack.h
@@ -69,6 +69,13 @@ public:
 	 */
 	virtual bool HasSection( UMovieSceneSection* Section ) const PURE_VIRTUAL( UMovieSceneSection::HasSection, return false; );
 
+	/**
+	 * Add a section to this track
+	 *
+	 * @param Section	The section to add
+	 */
+	virtual void AddSection( UMovieSceneSection* Section ) PURE_VIRTUAL( UMovieSceneSection::AddSection, );
+
 	/**
 	 * Removes a section from this track
 	 *
diff --git a/Engine/Source/Runtime/MovieSceneTracks/Private/Tracks/MovieScene3DPathTrack.cpp b/Engine/Source/Runtime/MovieSceneTracks/Private/Tracks/MovieScene3DPathTrack.cpp
--- a/Engine/Source/Runtime/MovieSceneTracks/Private/Tracks/MovieScene3DPathTrack.cpp
+++ b/Engine/Source/Runtime/MovieSceneTracks/Private/Tracks/MovieScene3DPathTrack.cpp
@@ -34,6 +34,11 @@ bool UMovieScene3DPathTrack::HasSection( UMovieSceneSection* Section ) const
 	return PathSections.Find( Section ) != INDEX_NONE;
 }
 
+void UMovieScene3DPathTrack::AddSection( UMovieSceneSection* Section )
+{
+	PathSections.Add( Section );
+}
+
 void UMovieScene3DPathTrack::RemoveSection( UMovieSceneSection* Section )
 {
 	PathSections.Remove( Section );
diff --git a/Engine/Source/Runtime/MovieSceneTracks/Private/Tracks/MovieSceneAudioTrack.cpp b/Engine/Source/Runtime/MovieSceneTracks/Private/Tracks/MovieSceneAudioTrack.cpp
--- a/Engine/Source/Runtime/MovieSceneTracks/Private/Tracks/MovieSceneAudioTrack.cpp
+++ b/Engine/Source/Runtime/MovieSceneTracks/Private/Tracks/MovieSceneAudioTrack.cpp
@@ -41,6 +41,11 @@ bool UMovieSceneAudioTrack::HasSection( UMovieSceneSection* Section ) const
 	return AudioSections.Find( Section ) != INDEX_NONE;
 }
 
+void UMovieSceneAudioTrack::AddSection( UMovieSceneSection* Section )
+{
+	AudioSections.Add( Section );
+}
+
 void UMovieSceneAudioTrack::RemoveSection( UMovieSceneSection* Section )
 {
 	AudioSections.Remove( Section );
diff --git a/Engine/Source/Runtime/MovieSceneTracks/Private/Tracks/MovieSceneParticleTrack.cpp b/Engine/Source/Runtime/MovieSceneTracks/Private/Tracks/MovieSceneParticleTrack.cpp
--- a/Engine/Source/Runtime/MovieSceneTracks/Private/Tracks/MovieSceneParticleTrack.cpp
+++ b/Engine/Source/Runtime/MovieSceneTracks/Private/Tracks/MovieSceneParticleTrack.cpp
@@ -38,6 +38,11 @@ bool UMovieSceneParticleTrack::HasSection( UMovieSceneSection* Section ) const
 	return ParticleSections.Find( Section ) != INDEX_NONE;
 }
 
+void UMovieSceneParticleTrack::AddSection( UMovieSceneSection* Section )
+{
+	ParticleSections.Add( Section );
+}
+
 void UMovieSceneParticleTrack::RemoveSection( UMovieSceneSection* Section )
 {
 	ParticleSections.Remove( Section );
diff --git a/Engine/Source/Runtime/MovieSceneTracks/Private/Tracks/MovieScenePropertyTrack.cpp b/Engine/Source/Runtime/MovieSceneTracks/Private/Tracks/MovieScenePropertyTrack.cpp
--- a/Engine/Source/Runtime/MovieSceneTracks/Private/Tracks/MovieScenePropertyTrack.cpp
+++ b/Engine/Source/Runtime/MovieSceneTracks/Private/Tracks/MovieScenePropertyTrack.cpp
@@ -38,12 +38,16 @@ bool UMovieScenePropertyTrack::HasSection( UMovieSceneSection* Section ) const
 }
 
 
+void UMovieScenePropertyTrack::AddSection( UMovieSceneSection* Section ) 
+{
+	Sections.Add( Section );
+}
+
 void UMovieScenePropertyTrack::RemoveSection( UMovieSceneSection* Section ) 
 {
 	Sections.Remove( Section );
 }
 
-
 bool UMovieScenePropertyTrack::IsEmpty() const
 {
 	return Sections.Num() == 0;
diff --git a/Engine/Source/Runtime/MovieSceneTracks/Private/Tracks/MovieSceneShotTrack.cpp b/Engine/Source/Runtime/MovieSceneTracks/Private/Tracks/MovieSceneShotTrack.cpp
--- a/Engine/Source/Runtime/MovieSceneTracks/Private/Tracks/MovieSceneShotTrack.cpp
+++ b/Engine/Source/Runtime/MovieSceneTracks/Private/Tracks/MovieSceneShotTrack.cpp
@@ -26,6 +26,11 @@ TSharedPtr<IMovieSceneTrackInstance> UMovieSceneShotTrack::CreateInstance()
 	return MakeShareable( new FMovieSceneShotTrackInstance( *this ) ); 
 }
 
+void UMovieSceneShotTrack::AddSection( UMovieSceneSection* Section )
+{
+	Super::AddSection( Section );
+}
+
 void UMovieSceneShotTrack::RemoveSection( UMovieSceneSection* Section )
 {
 	Super::RemoveSection( Section );
diff --git a/Engine/Source/Runtime/MovieSceneTracks/Private/Tracks/SubMovieSceneTrack.cpp b/Engine/Source/Runtime/MovieSceneTracks/Private/Tracks/SubMovieSceneTrack.cpp
--- a/Engine/Source/Runtime/MovieSceneTracks/Private/Tracks/SubMovieSceneTrack.cpp
+++ b/Engine/Source/Runtime/MovieSceneTracks/Private/Tracks/SubMovieSceneTrack.cpp
@@ -40,6 +40,15 @@ const TArray<UMovieSceneSection*>& USubMovieSceneTrack::GetAllSections() const
 	return SubMovieSceneSections;
 }
 
+void USubMovieSceneTrack::AddSection( UMovieSceneSection* Section )
+{
+	USubMovieSceneSection* SubMovieSceneSection = Cast<USubMovieSceneSection>( Section );
+	if(SubMovieSceneSection)
+	{
+		SubMovieSceneSections.Add( SubMovieSceneSection );
+	}
+}
+
 void USubMovieSceneTrack::RemoveSection( UMovieSceneSection* Section )
 {
 	USubMovieSceneSection* SubMovieSceneSection = Cast<USubMovieSceneSection>( Section );
diff --git a/Engine/Source/Runtime/MovieSceneTracks/Public/Tracks/MovieScene3DPathTrack.h b/Engine/Source/Runtime/MovieSceneTracks/Public/Tracks/MovieScene3DPathTrack.h
--- a/Engine/Source/Runtime/MovieSceneTracks/Public/Tracks/MovieScene3DPathTrack.h
+++ b/Engine/Source/Runtime/MovieSceneTracks/Public/Tracks/MovieScene3DPathTrack.h
@@ -19,6 +19,7 @@ public:
 	virtual TSharedPtr<IMovieSceneTrackInstance> CreateInstance() override;
 	virtual void RemoveAllAnimationData() override;
 	virtual bool HasSection( UMovieSceneSection* Section ) const override;
+	virtual void AddSection( UMovieSceneSection* Section ) override;
 	virtual void RemoveSection( UMovieSceneSection* Section ) override;
 	virtual bool IsEmpty() const override;
 	virtual TRange<float> GetSectionBoundaries() const override;
diff --git a/Engine/Source/Runtime/MovieSceneTracks/Public/Tracks/MovieSceneAudioTrack.h b/Engine/Source/Runtime/MovieSceneTracks/Public/Tracks/MovieSceneAudioTrack.h
--- a/Engine/Source/Runtime/MovieSceneTracks/Public/Tracks/MovieSceneAudioTrack.h
+++ b/Engine/Source/Runtime/MovieSceneTracks/Public/Tracks/MovieSceneAudioTrack.h
@@ -26,6 +26,7 @@ public:
 	virtual TSharedPtr<IMovieSceneTrackInstance> CreateInstance() override;
 	virtual void RemoveAllAnimationData() override;
 	virtual bool HasSection( UMovieSceneSection* Section ) const override;
+	virtual void AddSection( UMovieSceneSection* Section ) override;
 	virtual void RemoveSection( UMovieSceneSection* Section ) override;
 	virtual bool IsEmpty() const override;
 	virtual TRange<float> GetSectionBoundaries() const override;
diff --git a/Engine/Source/Runtime/MovieSceneTracks/Public/Tracks/MovieSceneParticleTrack.h b/Engine/Source/Runtime/MovieSceneTracks/Public/Tracks/MovieSceneParticleTrack.h
--- a/Engine/Source/Runtime/MovieSceneTracks/Public/Tracks/MovieSceneParticleTrack.h
+++ b/Engine/Source/Runtime/MovieSceneTracks/Public/Tracks/MovieSceneParticleTrack.h
@@ -21,6 +21,7 @@ public:
 	virtual TSharedPtr<IMovieSceneTrackInstance> CreateInstance() override;
 	virtual void RemoveAllAnimationData() override;
 	virtual bool HasSection( UMovieSceneSection* Section ) const override;
+	virtual void AddSection( UMovieSceneSection* Section ) override;
 	virtual void RemoveSection( UMovieSceneSection* Section ) override;
 	virtual bool IsEmpty() const override;
 	virtual TRange<float> GetSectionBoundaries() const override;
diff --git a/Engine/Source/Runtime/MovieSceneTracks/Public/Tracks/MovieScenePropertyTrack.h b/Engine/Source/Runtime/MovieSceneTracks/Public/Tracks/MovieScenePropertyTrack.h
--- a/Engine/Source/Runtime/MovieSceneTracks/Public/Tracks/MovieScenePropertyTrack.h
+++ b/Engine/Source/Runtime/MovieSceneTracks/Public/Tracks/MovieScenePropertyTrack.h
@@ -18,6 +18,7 @@ public:
 	virtual FName GetTrackName() const override { return PropertyName; }
 	virtual void RemoveAllAnimationData() override;
 	virtual bool HasSection( UMovieSceneSection* Section ) const override;
+	virtual void AddSection( UMovieSceneSection* Section ) override;
 	virtual void RemoveSection( UMovieSceneSection* Section ) override;
 	virtual bool IsEmpty() const override;
 	virtual TRange<float> GetSectionBoundaries() const override;
diff --git a/Engine/Source/Runtime/MovieSceneTracks/Public/Tracks/MovieSceneShotTrack.h b/Engine/Source/Runtime/MovieSceneTracks/Public/Tracks/MovieSceneShotTrack.h
--- a/Engine/Source/Runtime/MovieSceneTracks/Public/Tracks/MovieSceneShotTrack.h
+++ b/Engine/Source/Runtime/MovieSceneTracks/Public/Tracks/MovieSceneShotTrack.h
@@ -19,6 +19,7 @@ public:
 	/** UMovieSceneTrack interface */
 	virtual FName GetTrackName() const override;
 	virtual TSharedPtr<IMovieSceneTrackInstance> CreateInstance() override;
+	virtual void AddSection( UMovieSceneSection* Section ) override;
 	virtual void RemoveSection( UMovieSceneSection* Section ) override;
 
 	/** 
diff --git a/Engine/Source/Runtime/MovieSceneTracks/Public/Tracks/SubMovieSceneTrack.h b/Engine/Source/Runtime/MovieSceneTracks/Public/Tracks/SubMovieSceneTrack.h
--- a/Engine/Source/Runtime/MovieSceneTracks/Public/Tracks/SubMovieSceneTrack.h
+++ b/Engine/Source/Runtime/MovieSceneTracks/Public/Tracks/SubMovieSceneTrack.h
@@ -18,6 +18,7 @@ public:
 	virtual FName GetTrackName() const override;
 	virtual TSharedPtr<IMovieSceneTrackInstance> CreateInstance() override;
 	virtual const TArray<UMovieSceneSection*>& GetAllSections() const override;
+	virtual void AddSection( UMovieSceneSection* Section ) override;
 	virtual void RemoveSection( UMovieSceneSection* Section ) override;
 	virtual void RemoveAllAnimationData() override;
 	virtual bool IsEmpty() const override;
