Sequencer: Path track refactor to derive from constraint track, in anticipation of attach track.

diff --git a/Engine/Source/Editor/MovieSceneTools/Private/TrackEditors/PathTrackEditor.cpp b/Engine/Source/Editor/MovieSceneTools/Private/TrackEditors/PathTrackEditor.cpp
--- a/Engine/Source/Editor/MovieSceneTools/Private/TrackEditors/PathTrackEditor.cpp
+++ b/Engine/Source/Editor/MovieSceneTools/Private/TrackEditors/PathTrackEditor.cpp
@@ -284,7 +284,7 @@ void F3DPathTrackEditor::SetPath(UMovieSceneSection* Section, AActor* ActorWithS
 
 	if (SplineId.IsValid())
 	{
-		PathSection->SetPathId(SplineId);
+		PathSection->SetConstraintId(SplineId);
 	}
 }
 
@@ -333,7 +333,7 @@ void F3DPathTrackEditor::AddKeyInternal( float KeyTime, const TArray<UObject*> O
 						}
 					}
 
-					Cast<UMovieScene3DPathTrack>(Track)->AddPath( KeyTime, PathEndTime, SplineId );
+					Cast<UMovieScene3DPathTrack>(Track)->AddConstraint( KeyTime, PathEndTime, SplineId );
 				}
 			}
 		}
diff --git a/Engine/Source/Runtime/MovieSceneTracks/Private/Sections/MovieScene3DConstraintSection.cpp b/Engine/Source/Runtime/MovieSceneTracks/Private/Sections/MovieScene3DConstraintSection.cpp
new file mode 100644
--- /dev/null
+++ b/Engine/Source/Runtime/MovieSceneTracks/Private/Sections/MovieScene3DConstraintSection.cpp
@@ -0,0 +1,21 @@
+// Copyright 1998-2015 Epic Games, Inc. All Rights Reserved.
+
+#include "MovieSceneTracksPrivatePCH.h"
+#include "MovieScene3DConstraintSection.h"
+
+UMovieScene3DConstraintSection::UMovieScene3DConstraintSection( const FObjectInitializer& ObjectInitializer )
+	: Super( ObjectInitializer )
+{
+}
+
+void UMovieScene3DConstraintSection::SetConstraintId(const FGuid& InConstraintId)
+{
+	Modify();
+
+	ConstraintId = InConstraintId;
+}
+
+FGuid UMovieScene3DConstraintSection::GetConstraintId() const
+{
+	return ConstraintId;
+}
diff --git a/Engine/Source/Runtime/MovieSceneTracks/Private/Sections/MovieScene3DPathSection.cpp b/Engine/Source/Runtime/MovieSceneTracks/Private/Sections/MovieScene3DPathSection.cpp
--- a/Engine/Source/Runtime/MovieSceneTracks/Private/Sections/MovieScene3DPathSection.cpp
+++ b/Engine/Source/Runtime/MovieSceneTracks/Private/Sections/MovieScene3DPathSection.cpp
@@ -15,19 +15,7 @@ UMovieScene3DPathSection::UMovieScene3DPathSection( const FObjectInitializer& Ob
 {
 }
 
-void UMovieScene3DPathSection::SetPathId(const FGuid& InPathId)
-{
-	Modify();
-
-	PathId = InPathId;
-}
-
-FGuid UMovieScene3DPathSection::GetPathId() const
-{
-	return PathId;
-}
-
-void UMovieScene3DPathSection::Eval( float Position, USplineComponent* SplineComponent, FVector& OutTranslation, FRotator& OutRotation ) const
+void UMovieScene3DPathSection::Eval( USceneComponent* SceneComponent, float Position, USplineComponent* SplineComponent, FVector& OutTranslation, FRotator& OutRotation ) const
 {
 	float Timing = TimingCurve.Eval( Position );
 
@@ -115,6 +103,11 @@ void UMovieScene3DPathSection::Eval( float Position, USplineComponent* SplineCom
 		OutRotation.Pitch = 0.f;
 		OutRotation.Roll = 0.f;
 	}
+
+	if (!bFollow)
+	{
+		OutRotation = SceneComponent->GetRelativeTransform().GetRotation().Rotator();
+	}
 }
 
 void UMovieScene3DPathSection::MoveSection( float DeltaPosition, TSet<FKeyHandle>& KeyHandles )
@@ -147,7 +140,7 @@ void UMovieScene3DPathSection::GetKeyHandles(TSet<FKeyHandle>& KeyHandles) const
 void UMovieScene3DPathSection::AddPath( float Time, float SequenceEndTime, const FGuid& InPathId )
 {
 	Modify();
-	PathId = InPathId;
+	ConstraintId = InPathId;
 
 	TimingCurve.UpdateOrAddKey(Time, 0);
 	TimingCurve.UpdateOrAddKey(SequenceEndTime, 1);
diff --git a/Engine/Source/Runtime/MovieSceneTracks/Private/TrackInstances/MovieScene3DConstraintTrackInstance.cpp b/Engine/Source/Runtime/MovieSceneTracks/Private/TrackInstances/MovieScene3DConstraintTrackInstance.cpp
new file mode 100644
--- /dev/null
+++ b/Engine/Source/Runtime/MovieSceneTracks/Private/TrackInstances/MovieScene3DConstraintTrackInstance.cpp
@@ -0,0 +1,81 @@
+// Copyright 1998-2015 Epic Games, Inc. All Rights Reserved.
+
+#include "MovieSceneTracksPrivatePCH.h"
+#include "MovieSceneCommonHelpers.h"
+#include "MovieScene3DConstraintTrackInstance.h"
+#include "MovieScene3DConstraintTrack.h"
+#include "MovieScene3DConstraintSection.h"
+#include "IMovieScenePlayer.h"
+#include "Components/SplineComponent.h"
+
+FMovieScene3DConstraintTrackInstance::FMovieScene3DConstraintTrackInstance( UMovieScene3DConstraintTrack& InConstraintTrack )
+{
+	ConstraintTrack = &InConstraintTrack;
+}
+
+void FMovieScene3DConstraintTrackInstance::SaveState(const TArray<UObject*>& RuntimeObjects)
+{
+	for (int32 ObjIndex = 0; ObjIndex < RuntimeObjects.Num(); ++ObjIndex)
+	{
+		USceneComponent* SceneComponent = MovieSceneHelpers::SceneComponentFromRuntimeObject(RuntimeObjects[ObjIndex]);
+		if (SceneComponent != NULL)
+		{
+			InitTransformMap.Add(RuntimeObjects[ObjIndex], SceneComponent->GetRelativeTransform());
+		}
+	}
+}
+
+void FMovieScene3DConstraintTrackInstance::RestoreState(const TArray<UObject*>& RuntimeObjects)
+{
+	for (int32 ObjIndex = 0; ObjIndex < RuntimeObjects.Num(); ++ObjIndex)
+	{
+		if (!IsValid(RuntimeObjects[ObjIndex]))
+		{
+			continue;
+		}
+
+		USceneComponent* SceneComponent = MovieSceneHelpers::SceneComponentFromRuntimeObject(RuntimeObjects[ObjIndex]);
+		if (SceneComponent != NULL)
+		{
+			FTransform *Transform = InitTransformMap.Find(RuntimeObjects[ObjIndex]);
+			if (Transform != NULL)
+			{
+				SceneComponent->SetRelativeTransform(*Transform);
+			}
+		}
+	}
+}
+
+void FMovieScene3DConstraintTrackInstance::Update( float Position, float LastPosition, const TArray<UObject*>& RuntimeObjects, class IMovieScenePlayer& Player ) 
+{
+	UMovieScene3DConstraintSection* FirstConstraintSection = NULL;
+
+	const TArray<UMovieSceneSection*>& ConstraintSections = ConstraintTrack->GetAllSections();
+
+	for (int32 ConstraintIndex = 0; ConstraintIndex < ConstraintSections.Num(); ++ConstraintIndex)
+	{
+		UMovieScene3DConstraintSection* ConstraintSection = CastChecked<UMovieScene3DConstraintSection>(ConstraintSections[ConstraintIndex]);
+
+		if (ConstraintSection->IsTimeWithinSection(Position) &&
+			(FirstConstraintSection == NULL || FirstConstraintSection->GetRowIndex() > ConstraintSection->GetRowIndex()))
+		{
+			TArray<UObject*> ConstraintObjects;
+			FGuid ConstraintId = ConstraintSection->GetConstraintId();
+
+			if (ConstraintId.IsValid())
+			{
+				Player.GetRuntimeObjects( Player.GetRootMovieSceneSequenceInstance(), ConstraintId, ConstraintObjects);
+
+				for (int32 ConstraintObjectIndex = 0; ConstraintObjectIndex < ConstraintObjects.Num(); ++ConstraintObjectIndex)
+				{
+					AActor* Actor = Cast<AActor>(ConstraintObjects[ConstraintObjectIndex]);
+					if (Actor)
+					{
+						UpdateConstraint(Position, RuntimeObjects, Actor, ConstraintSection);
+					}	
+				}
+			}
+		}
+	}
+}
+
diff --git a/Engine/Source/Runtime/MovieSceneTracks/Private/TrackInstances/MovieScene3DPathTrackInstance.cpp b/Engine/Source/Runtime/MovieSceneTracks/Private/TrackInstances/MovieScene3DPathTrackInstance.cpp
--- a/Engine/Source/Runtime/MovieSceneTracks/Private/TrackInstances/MovieScene3DPathTrackInstance.cpp
+++ b/Engine/Source/Runtime/MovieSceneTracks/Private/TrackInstances/MovieScene3DPathTrackInstance.cpp
@@ -9,91 +9,32 @@
 #include "Components/SplineComponent.h"
 
 FMovieScene3DPathTrackInstance::FMovieScene3DPathTrackInstance( UMovieScene3DPathTrack& InPathTrack )
+	: FMovieScene3DConstraintTrackInstance( InPathTrack )
 {
-	PathTrack = &InPathTrack;
 }
 
-void FMovieScene3DPathTrackInstance::SaveState(const TArray<UObject*>& RuntimeObjects)
-{
-	for (int32 ObjIndex = 0; ObjIndex < RuntimeObjects.Num(); ++ObjIndex)
-	{
-		USceneComponent* SceneComponent = MovieSceneHelpers::SceneComponentFromRuntimeObject(RuntimeObjects[ObjIndex]);
-		if (SceneComponent != NULL)
-		{
-			InitTransformMap.Add(RuntimeObjects[ObjIndex], SceneComponent->GetRelativeTransform());
-		}
-	}
-}
-
-void FMovieScene3DPathTrackInstance::RestoreState(const TArray<UObject*>& RuntimeObjects)
-{
-	for (int32 ObjIndex = 0; ObjIndex < RuntimeObjects.Num(); ++ObjIndex)
-	{
-		if (!IsValid(RuntimeObjects[ObjIndex]))
-		{
-			continue;
-		}
-
-		USceneComponent* SceneComponent = MovieSceneHelpers::SceneComponentFromRuntimeObject(RuntimeObjects[ObjIndex]);
-		if (SceneComponent != NULL)
-		{
-			FTransform *Transform = InitTransformMap.Find(RuntimeObjects[ObjIndex]);
-			if (Transform != NULL)
-			{
-				SceneComponent->SetRelativeTransform(*Transform);
-			}
-		}
-	}
-}
-
-void FMovieScene3DPathTrackInstance::Update( float Position, float LastPosition, const TArray<UObject*>& RuntimeObjects, class IMovieScenePlayer& Player ) 
+void FMovieScene3DPathTrackInstance::UpdateConstraint( float Position, const TArray<UObject*>& RuntimeObjects, AActor* Actor, UMovieScene3DConstraintSection* ConstraintSection ) 
 {
 	FVector Translation;
 	FRotator Rotation;
 
-	UMovieScene3DPathSection* FirstPathSection = NULL;
-
-	const TArray<UMovieSceneSection*>& PathSections = PathTrack->GetAllSections();
+	TArray<USplineComponent*> SplineComponents;
+	Actor->GetComponents(SplineComponents);
 
-	for (int32 PathIndex = 0; PathIndex < PathSections.Num(); ++PathIndex)
+	if (SplineComponents.Num())
 	{
-		UMovieScene3DPathSection* PathSection = CastChecked<UMovieScene3DPathSection>(PathSections[PathIndex]);
+		UMovieScene3DPathSection* PathSection = CastChecked<UMovieScene3DPathSection>(ConstraintSection);
 
-		if (PathSection->IsTimeWithinSection(Position) &&
-			(FirstPathSection == NULL || FirstPathSection->GetRowIndex() > PathSection->GetRowIndex()))
+		for( int32 ObjIndex = 0; ObjIndex < RuntimeObjects.Num(); ++ObjIndex )
 		{
-			TArray<UObject*> PathObjects;
-			FGuid PathId = PathSection->GetPathId();
+			USceneComponent* SceneComponent = MovieSceneHelpers::SceneComponentFromRuntimeObject(RuntimeObjects[ObjIndex]);
 
-			if (PathId.IsValid())
+			if (SceneComponent != NULL)
 			{
-				Player.GetRuntimeObjects( Player.GetRootMovieSceneSequenceInstance(), PathId, PathObjects);
-
-				for (int32 PathObjectIndex = 0; PathObjectIndex < PathObjects.Num(); ++PathObjectIndex)
-				{
-					AActor* Actor = Cast<AActor>(PathObjects[PathObjectIndex]);
-					if (Actor)
-					{
-						TArray<USplineComponent*> SplineComponents;
-						Actor->GetComponents(SplineComponents);
-
-						if (SplineComponents.Num())
-						{
-							PathSection->Eval(Position, SplineComponents[0], Translation, Rotation);
-
-							for( int32 ObjIndex = 0; ObjIndex < RuntimeObjects.Num(); ++ObjIndex )
-							{
-								USceneComponent* SceneComponent = MovieSceneHelpers::SceneComponentFromRuntimeObject(RuntimeObjects[ObjIndex]);
+				PathSection->Eval(SceneComponent, Position, SplineComponents[0], Translation, Rotation);
 
-								if (SceneComponent != NULL)
-								{
-									// Set the relative translation and rotation.  Note they are set once instead of individually to avoid a redundant component transform update.
-									SceneComponent->SetRelativeLocationAndRotation(Translation, Rotation);
-								}
-							}
-						}
-					}	
-				}
+				// Set the relative translation and rotation.  Note they are set once instead of individually to avoid a redundant component transform update.
+				SceneComponent->SetRelativeLocationAndRotation(Translation, Rotation);
 			}
 		}
 	}
diff --git a/Engine/Source/Runtime/MovieSceneTracks/Private/Tracks/MovieScene3DConstraintTrack.cpp b/Engine/Source/Runtime/MovieSceneTracks/Private/Tracks/MovieScene3DConstraintTrack.cpp
new file mode 100644
--- /dev/null
+++ b/Engine/Source/Runtime/MovieSceneTracks/Private/Tracks/MovieScene3DConstraintTrack.cpp
@@ -0,0 +1,51 @@
+// Copyright 1998-2015 Epic Games, Inc. All Rights Reserved.
+
+#include "MovieSceneTracksPrivatePCH.h"
+#include "MovieScene3DConstraintSection.h"
+#include "MovieScene3DConstraintTrack.h"
+#include "IMovieScenePlayer.h"
+#include "MovieScene3DConstraintTrackInstance.h"
+
+UMovieScene3DConstraintTrack::UMovieScene3DConstraintTrack( const FObjectInitializer& ObjectInitializer )
+	: Super( ObjectInitializer )
+{
+}
+
+const TArray<UMovieSceneSection*>& UMovieScene3DConstraintTrack::GetAllSections() const
+{
+	return ConstraintSections;
+}
+
+void UMovieScene3DConstraintTrack::RemoveAllAnimationData()
+{
+}
+
+bool UMovieScene3DConstraintTrack::HasSection( UMovieSceneSection* Section ) const
+{
+	return ConstraintSections.Find( Section ) != INDEX_NONE;
+}
+
+void UMovieScene3DConstraintTrack::AddSection( UMovieSceneSection* Section )
+{
+	ConstraintSections.Add( Section );
+}
+
+void UMovieScene3DConstraintTrack::RemoveSection( UMovieSceneSection* Section )
+{
+	ConstraintSections.Remove( Section );
+}
+
+bool UMovieScene3DConstraintTrack::IsEmpty() const
+{
+	return ConstraintSections.Num() == 0;
+}
+
+TRange<float> UMovieScene3DConstraintTrack::GetSectionBoundaries() const
+{
+	TArray< TRange<float> > Bounds;
+	for (int32 i = 0; i < ConstraintSections.Num(); ++i)
+	{
+		Bounds.Add(ConstraintSections[i]->GetRange());
+	}
+	return TRange<float>::Hull(Bounds);
+}
diff --git a/Engine/Source/Runtime/MovieSceneTracks/Private/Tracks/MovieScene3DPathTrack.cpp b/Engine/Source/Runtime/MovieSceneTracks/Private/Tracks/MovieScene3DPathTrack.cpp
--- a/Engine/Source/Runtime/MovieSceneTracks/Private/Tracks/MovieScene3DPathTrack.cpp
+++ b/Engine/Source/Runtime/MovieSceneTracks/Private/Tracks/MovieScene3DPathTrack.cpp
@@ -20,51 +20,12 @@ TSharedPtr<IMovieSceneTrackInstance> UMovieScene3DPathTrack::CreateInstance()
 	return MakeShareable( new FMovieScene3DPathTrackInstance( *this ) ); 
 }
 
-const TArray<UMovieSceneSection*>& UMovieScene3DPathTrack::GetAllSections() const
-{
-	return PathSections;
-}
-
-void UMovieScene3DPathTrack::RemoveAllAnimationData()
-{
-}
-
-bool UMovieScene3DPathTrack::HasSection( UMovieSceneSection* Section ) const
-{
-	return PathSections.Find( Section ) != INDEX_NONE;
-}
-
-void UMovieScene3DPathTrack::AddSection( UMovieSceneSection* Section )
-{
-	PathSections.Add( Section );
-}
-
-void UMovieScene3DPathTrack::RemoveSection( UMovieSceneSection* Section )
-{
-	PathSections.Remove( Section );
-}
-
-bool UMovieScene3DPathTrack::IsEmpty() const
-{
-	return PathSections.Num() == 0;
-}
-
-TRange<float> UMovieScene3DPathTrack::GetSectionBoundaries() const
-{
-	TArray< TRange<float> > Bounds;
-	for (int32 i = 0; i < PathSections.Num(); ++i)
-	{
-		Bounds.Add(PathSections[i]->GetRange());
-	}
-	return TRange<float>::Hull(Bounds);
-}
-
-void UMovieScene3DPathTrack::AddPath(float KeyTime, float PathEndTime, const FGuid& PathId)
+void UMovieScene3DPathTrack::AddConstraint(float KeyTime, float ConstraintEndTime, const FGuid& ConstraintId)
 {
 	// add the section
 	UMovieScene3DPathSection* NewSection = NewObject<UMovieScene3DPathSection>(this);
-	NewSection->AddPath(KeyTime, PathEndTime, PathId);
-	NewSection->InitialPlacement( PathSections, KeyTime, PathEndTime, SupportsMultipleRows() );
+	NewSection->AddPath(KeyTime, ConstraintEndTime, ConstraintId);
+	NewSection->InitialPlacement( ConstraintSections, KeyTime, ConstraintEndTime, SupportsMultipleRows() );
 
-	PathSections.Add(NewSection);
+	ConstraintSections.Add(NewSection);
 }
diff --git a/Engine/Source/Runtime/MovieSceneTracks/Public/Sections/MovieScene3DConstraintSection.h b/Engine/Source/Runtime/MovieSceneTracks/Public/Sections/MovieScene3DConstraintSection.h
new file mode 100644
--- /dev/null
+++ b/Engine/Source/Runtime/MovieSceneTracks/Public/Sections/MovieScene3DConstraintSection.h
@@ -0,0 +1,29 @@
+// Copyright 1998-2015 Epic Games, Inc. All Rights Reserved.
+
+#pragma once
+
+#include "MovieSceneSection.h"
+#include "MovieScene3DConstraintSection.generated.h"
+
+/**
+ * Base class for 3D constraint section
+ */
+UCLASS(MinimalAPI)
+class UMovieScene3DConstraintSection
+	: public UMovieSceneSection
+{
+	GENERATED_UCLASS_BODY()
+
+public:
+
+	/** Sets the constraint id for this section */
+	virtual void SetConstraintId(const FGuid& InId);
+
+	/** Gets the constraint id for this section */
+	virtual FGuid GetConstraintId() const;
+	
+protected:
+	/** The possessable guid that this constraint uses */
+	UPROPERTY()
+	FGuid ConstraintId;
+};
\ No newline at end of file
diff --git a/Engine/Source/Runtime/MovieSceneTracks/Public/Sections/MovieScene3DPathSection.h b/Engine/Source/Runtime/MovieSceneTracks/Public/Sections/MovieScene3DPathSection.h
--- a/Engine/Source/Runtime/MovieSceneTracks/Public/Sections/MovieScene3DPathSection.h
+++ b/Engine/Source/Runtime/MovieSceneTracks/Public/Sections/MovieScene3DPathSection.h
@@ -2,7 +2,7 @@
 
 #pragma once
 
-#include "MovieSceneSection.h"
+#include "MovieScene3DConstraintSection.h"
 #include "MovieScene3DPathSection.generated.h"
 
 class USplineComponent;
@@ -23,7 +23,7 @@ enum class MovieScene3DPathSection_Axis : uint8
  */
 UCLASS(MinimalAPI)
 class UMovieScene3DPathSection
-	: public UMovieSceneSection
+	: public UMovieScene3DConstraintSection
 {
 	GENERATED_UCLASS_BODY()
 
@@ -34,18 +34,12 @@ public:
 	virtual void DilateSection( float DilationFactor, float Origin, TSet<FKeyHandle>& KeyHandles ) override;
 	virtual void GetKeyHandles(TSet<FKeyHandle>& KeyHandles) const override;
 
-	/** Sets the path id for this path section */
-	virtual void SetPathId(const FGuid& InId);
-
-	/** Gets the path id for this path section */
-	virtual FGuid GetPathId() const;
-
 	/**
 	 * Evaluates the path track
 	 *
 	 * @param Time				The position in time within the movie scene
 	 */
-	void Eval( float Time, USplineComponent* SplineComponent, FVector& OutTranslation, FRotator& OutRotation ) const;
+	void Eval( USceneComponent* SceneComponent, float Time, USplineComponent* SplineComponent, FVector& OutTranslation, FRotator& OutRotation ) const;
 
 	/** 
 	 * Adds a path to the section
@@ -64,10 +58,6 @@ public:
 	MOVIESCENETRACKS_API FRichCurve& GetTimingCurve() { return TimingCurve; }
 	
 private:
-	/** The path possessable that this path uses */
-	UPROPERTY()
-	FGuid PathId;
-
 	/** Timing Curve */
 	UPROPERTY()
 	FRichCurve TimingCurve;
diff --git a/Engine/Source/Runtime/MovieSceneTracks/Public/TrackInstances/MovieScene3DConstraintTrackInstance.h b/Engine/Source/Runtime/MovieSceneTracks/Public/TrackInstances/MovieScene3DConstraintTrackInstance.h
new file mode 100644
--- /dev/null
+++ b/Engine/Source/Runtime/MovieSceneTracks/Public/TrackInstances/MovieScene3DConstraintTrackInstance.h
@@ -0,0 +1,34 @@
+// Copyright 1998-2015 Epic Games, Inc. All Rights Reserved.
+
+#pragma once
+
+#include "IMovieSceneTrackInstance.h"
+
+class UMovieScene3DConstraintTrack;
+
+/**
+ * Instance of a UMovieScene3DConstraintTrack
+ */
+class FMovieScene3DConstraintTrackInstance
+	: public IMovieSceneTrackInstance
+{
+public:
+	FMovieScene3DConstraintTrackInstance( UMovieScene3DConstraintTrack& InConstraintTrack );
+
+	/** IMovieSceneTrackInstance interface */
+	virtual void SaveState (const TArray<UObject*>& RuntimeObjects) override;
+	virtual void RestoreState (const TArray<UObject*>& RuntimeObjects) override;
+	virtual void Update( float Position, float LastPosition, const TArray<UObject*>& RuntimeObjects, IMovieScenePlayer& Player ) override;
+	virtual void RefreshInstance( const TArray<UObject*>& RuntimeObjects, IMovieScenePlayer& Player ) override {}
+	virtual void ClearInstance( IMovieScenePlayer& Player ) override {}
+	virtual float EvalOrder() { return 1.f; }
+
+	virtual void UpdateConstraint(float Position, const TArray<UObject*>& RuntimeObjects, AActor* Actor, UMovieScene3DConstraintSection* ConstraintSection) = 0;
+
+protected:
+	/** Track that is being instanced */
+	UMovieScene3DConstraintTrack* ConstraintTrack;
+
+	/** Map from object to initial state */
+	TMap< TWeakObjectPtr<UObject>, FTransform > InitTransformMap;
+};
diff --git a/Engine/Source/Runtime/MovieSceneTracks/Public/TrackInstances/MovieScene3DPathTrackInstance.h b/Engine/Source/Runtime/MovieSceneTracks/Public/TrackInstances/MovieScene3DPathTrackInstance.h
--- a/Engine/Source/Runtime/MovieSceneTracks/Public/TrackInstances/MovieScene3DPathTrackInstance.h
+++ b/Engine/Source/Runtime/MovieSceneTracks/Public/TrackInstances/MovieScene3DPathTrackInstance.h
@@ -2,7 +2,7 @@
 
 #pragma once
 
-#include "IMovieSceneTrackInstance.h"
+#include "MovieScene3DConstraintTrackInstance.h"
 
 class UMovieScene3DPathTrack;
 
@@ -10,23 +10,11 @@ class UMovieScene3DPathTrack;
  * Instance of a UMovieScene3DPathTrack
  */
 class FMovieScene3DPathTrackInstance
-	: public IMovieSceneTrackInstance
+	: public FMovieScene3DConstraintTrackInstance
 {
 public:
 	FMovieScene3DPathTrackInstance( UMovieScene3DPathTrack& InPathTrack );
 
-	/** IMovieSceneTrackInstance interface */
-	virtual void SaveState (const TArray<UObject*>& RuntimeObjects) override;
-	virtual void RestoreState (const TArray<UObject*>& RuntimeObjects) override;
-	virtual void Update( float Position, float LastPosition, const TArray<UObject*>& RuntimeObjects, IMovieScenePlayer& Player ) override;
-	virtual void RefreshInstance( const TArray<UObject*>& RuntimeObjects, IMovieScenePlayer& Player ) override {}
-	virtual void ClearInstance( IMovieScenePlayer& Player ) override {}
-	virtual float EvalOrder() { return 1.f; }
-
-private:
-	/** Track that is being instanced */
-	UMovieScene3DPathTrack* PathTrack;
-
-	/** Map from object to initial state */
-	TMap< TWeakObjectPtr<UObject>, FTransform > InitTransformMap;
+	/** FMovieScene3DConstraintTrackInstance interface */
+	virtual void UpdateConstraint(float Position, const TArray<UObject*>& RuntimeObjects, AActor* Actor, UMovieScene3DConstraintSection* ConstraintSection) override;
 };
diff --git a/Engine/Source/Runtime/MovieSceneTracks/Public/Tracks/MovieScene3DConstraintTrack.h b/Engine/Source/Runtime/MovieSceneTracks/Public/Tracks/MovieScene3DConstraintTrack.h
new file mode 100644
--- /dev/null
+++ b/Engine/Source/Runtime/MovieSceneTracks/Public/Tracks/MovieScene3DConstraintTrack.h
@@ -0,0 +1,39 @@
+// Copyright 1998-2015 Epic Games, Inc. All Rights Reserved.
+
+#pragma once
+
+#include "MovieScene.h"
+#include "MovieSceneTrack.h"
+#include "MovieScene3DConstraintTrack.generated.h"
+
+/**
+ * Base class for constraint tracks - tracks that are dependent upon other objects
+ */
+UCLASS( MinimalAPI )
+class UMovieScene3DConstraintTrack : public UMovieSceneTrack
+{
+	GENERATED_UCLASS_BODY()
+public:
+	/** UMovieSceneTrack interface */
+    virtual void RemoveAllAnimationData() override;
+	virtual bool HasSection( UMovieSceneSection* Section ) const override;
+	virtual void AddSection( UMovieSceneSection* Section ) override;
+	virtual void RemoveSection( UMovieSceneSection* Section ) override;
+	virtual bool IsEmpty() const override;
+	virtual TRange<float> GetSectionBoundaries() const override;
+	virtual const TArray<UMovieSceneSection*>& GetAllSections() const override;
+
+	/**
+	 * Adds a constraint
+	 *
+	 * @param Time				    The time relative to the owning movie scene where the section should be
+	 * @param ConstraintEndTime     Set the constraint to end at this time
+	 * @param ConstraintId			The id to the constraint
+	 */
+	virtual void AddConstraint( float Time, float ConstraintEndTime, const FGuid& ConstraintId ) {}
+
+protected:
+	/** List of all constraint sections */
+	UPROPERTY()
+	TArray<UMovieSceneSection*> ConstraintSections;
+};
diff --git a/Engine/Source/Runtime/MovieSceneTracks/Public/Tracks/MovieScene3DPathTrack.h b/Engine/Source/Runtime/MovieSceneTracks/Public/Tracks/MovieScene3DPathTrack.h
--- a/Engine/Source/Runtime/MovieSceneTracks/Public/Tracks/MovieScene3DPathTrack.h
+++ b/Engine/Source/Runtime/MovieSceneTracks/Public/Tracks/MovieScene3DPathTrack.h
@@ -3,39 +3,27 @@
 #pragma once
 
 #include "MovieScene.h"
-#include "MovieSceneTrack.h"
+#include "MovieScene3DConstraintTrack.h"
 #include "MovieScene3DPathTrack.generated.h"
 
 /**
  * Handles manipulation of path tracks in a movie scene
  */
 UCLASS( MinimalAPI )
-class UMovieScene3DPathTrack : public UMovieSceneTrack
+class UMovieScene3DPathTrack : public UMovieScene3DConstraintTrack
 {
 	GENERATED_UCLASS_BODY()
 public:
 	/** UMovieSceneTrack interface */
 	virtual FName GetTrackName() const override;
 	virtual TSharedPtr<IMovieSceneTrackInstance> CreateInstance() override;
-	virtual void RemoveAllAnimationData() override;
-	virtual bool HasSection( UMovieSceneSection* Section ) const override;
-	virtual void AddSection( UMovieSceneSection* Section ) override;
-	virtual void RemoveSection( UMovieSceneSection* Section ) override;
-	virtual bool IsEmpty() const override;
-	virtual TRange<float> GetSectionBoundaries() const override;
-	virtual const TArray<UMovieSceneSection*>& GetAllSections() const override;
 
 	/**
-	 * Adds a path
+	 * Adds a constraint
 	 *
-	 * @param Time				The time relative to the owning movie scene where the section should be
-	 * @param PathEndTime       Set the path to end at this time
-	 * @param PathId			The id to the spline
+	 * @param Time				    The time relative to the owning movie scene where the section should be
+	 * @param ConstraintEndTime     Set the constraint to end at this time
+	 * @param ConstraintId			The id to the constraint
 	 */
-	virtual void AddPath( float Time, float PathEndTime, const FGuid& PathId );
-
-private:
-	/** List of all path sections */
-	UPROPERTY()
-	TArray<UMovieSceneSection*> PathSections;
+	virtual void AddConstraint( float Time, float ConstraintEndTime, const FGuid& ConstraintId ) override;
 };
