Sequencer: Fix non unity build.

diff --git a/Engine/Source/Runtime/MovieSceneTracks/Public/TrackInstances/MovieScene3DAttachTrackInstance.h b/Engine/Source/Runtime/MovieSceneTracks/Public/TrackInstances/MovieScene3DAttachTrackInstance.h
--- a/Engine/Source/Runtime/MovieSceneTracks/Public/TrackInstances/MovieScene3DAttachTrackInstance.h
+++ b/Engine/Source/Runtime/MovieSceneTracks/Public/TrackInstances/MovieScene3DAttachTrackInstance.h
@@ -5,6 +5,7 @@
 #include "MovieScene3DConstraintTrackInstance.h"
 
 class UMovieScene3DAttachTrack;
+class UMovieScene3DConstraintSection;
 
 /**
  * Instance of a UMovieScene3DAttachTrack
diff --git a/Engine/Source/Runtime/MovieSceneTracks/Public/TrackInstances/MovieScene3DConstraintTrackInstance.h b/Engine/Source/Runtime/MovieSceneTracks/Public/TrackInstances/MovieScene3DConstraintTrackInstance.h
--- a/Engine/Source/Runtime/MovieSceneTracks/Public/TrackInstances/MovieScene3DConstraintTrackInstance.h
+++ b/Engine/Source/Runtime/MovieSceneTracks/Public/TrackInstances/MovieScene3DConstraintTrackInstance.h
@@ -5,6 +5,7 @@
 #include "IMovieSceneTrackInstance.h"
 
 class UMovieScene3DConstraintTrack;
+class UMovieScene3DConstraintSection;
 
 /**
  * Instance of a UMovieScene3DConstraintTrack
diff --git a/Engine/Source/Runtime/MovieSceneTracks/Public/TrackInstances/MovieScene3DPathTrackInstance.h b/Engine/Source/Runtime/MovieSceneTracks/Public/TrackInstances/MovieScene3DPathTrackInstance.h
--- a/Engine/Source/Runtime/MovieSceneTracks/Public/TrackInstances/MovieScene3DPathTrackInstance.h
+++ b/Engine/Source/Runtime/MovieSceneTracks/Public/TrackInstances/MovieScene3DPathTrackInstance.h
@@ -5,6 +5,7 @@
 #include "MovieScene3DConstraintTrackInstance.h"
 
 class UMovieScene3DPathTrack;
+class UMovieScene3DConstraintSection;
 
 /**
  * Instance of a UMovieScene3DPathTrack
