Sequencer: Fix section evaluation so that the evaluation time is clamped to the nearest section's start or end time if the evaluation time is not within the section.

diff --git a/Engine/Source/Runtime/MovieScene/Private/MovieSceneCommonHelpers.cpp b/Engine/Source/Runtime/MovieScene/Private/MovieSceneCommonHelpers.cpp
--- a/Engine/Source/Runtime/MovieScene/Private/MovieSceneCommonHelpers.cpp
+++ b/Engine/Source/Runtime/MovieScene/Private/MovieSceneCommonHelpers.cpp
@@ -31,15 +31,18 @@ TArray<UMovieSceneSection*> MovieSceneHelpers::GetTraversedSections( const TArra
 
 UMovieSceneSection* MovieSceneHelpers::FindSectionAtTime( const TArray<UMovieSceneSection*>& Sections, float Time )
 {
-	if( Sections.Num() == 1 )
-	{
-		return Sections[0];
-	}
-	else
+	for( int32 SectionIndex = 0; SectionIndex < Sections.Num(); ++SectionIndex )
 	{
-		return MovieSceneHelpers::FindNearestSectionAtTime( Sections, Time );
+		UMovieSceneSection* Section = Sections[SectionIndex];
+
+		//@todo Sequencer - There can be multiple sections overlapping in time. Returning instantly does not account for that.
+		if( Section->IsTimeWithinSection( Time ) )
+		{
+			return Section;
+		}
 	}
 
+	return NULL;
 }
 
 
diff --git a/Engine/Source/Runtime/MovieSceneTracks/Private/Tracks/MovieScene3DTransformTrack.cpp b/Engine/Source/Runtime/MovieSceneTracks/Private/Tracks/MovieScene3DTransformTrack.cpp
--- a/Engine/Source/Runtime/MovieSceneTracks/Private/Tracks/MovieScene3DTransformTrack.cpp
+++ b/Engine/Source/Runtime/MovieSceneTracks/Private/Tracks/MovieScene3DTransformTrack.cpp
@@ -74,7 +74,7 @@ TSharedPtr<IMovieSceneTrackInstance> UMovieScene3DTransformTrack::CreateInstance
 
 bool UMovieScene3DTransformTrack::AddKeyToSection( const FGuid& ObjectHandle, const FTransformKey& InKey, const bool bUnwindRotation, F3DTransformTrackKey::Type KeyType )
 {
-	const UMovieSceneSection* NearestSection = MovieSceneHelpers::FindSectionAtTime(Sections, InKey.GetKeyTime());
+	const UMovieSceneSection* NearestSection = MovieSceneHelpers::FindNearestSectionAtTime(Sections, InKey.GetKeyTime());
 	if (!NearestSection || CastChecked<UMovieScene3DTransformSection>(NearestSection)->NewKeyIsNewData(InKey))
 	{
 		Modify();
@@ -97,13 +97,15 @@ bool UMovieScene3DTransformTrack::AddKeyToSection( const FGuid& ObjectHandle, co
 
 bool UMovieScene3DTransformTrack::Eval( float Position, float LastPosition, FVector& OutTranslation, FRotator& OutRotation, FVector& OutScale, TArray<bool>& OutHasTranslationKeys, TArray<bool>& OutHasRotationKeys, TArray<bool>& OutHasScaleKeys ) const
 {
-	const UMovieSceneSection* Section = MovieSceneHelpers::FindSectionAtTime( Sections, Position );
+	const UMovieSceneSection* Section = MovieSceneHelpers::FindNearestSectionAtTime( Sections, Position );
 
 	if( Section )
 	{
 		const UMovieScene3DTransformSection* TransformSection = CastChecked<UMovieScene3DTransformSection>( Section );
 
-		// Evalulate translation,rotation, and scale curves.  If no keys were found on one of these, that component of the transform will remain unchained
+		Position = FMath::Clamp(Position, Section->GetStartTime(), Section->GetEndTime());
+
+		// Evaluate translation,rotation, and scale curves.  If no keys were found on one of these, that component of the transform will remain unchained
 		TransformSection->EvalTranslation( Position, OutTranslation, OutHasTranslationKeys );
 		TransformSection->EvalRotation( Position, OutRotation, OutHasRotationKeys );
 		TransformSection->EvalScale( Position, OutScale, OutHasScaleKeys );
diff --git a/Engine/Source/Runtime/MovieSceneTracks/Private/Tracks/MovieSceneBoolTrack.cpp b/Engine/Source/Runtime/MovieSceneTracks/Private/Tracks/MovieSceneBoolTrack.cpp
--- a/Engine/Source/Runtime/MovieSceneTracks/Private/Tracks/MovieSceneBoolTrack.cpp
+++ b/Engine/Source/Runtime/MovieSceneTracks/Private/Tracks/MovieSceneBoolTrack.cpp
@@ -40,10 +40,12 @@ bool UMovieSceneBoolTrack::AddKeyToSection( float Time, bool Value )
 
 bool UMovieSceneBoolTrack::Eval( float Position, float LastPostion, bool& OutBool ) const
 {	
-	const UMovieSceneSection* Section = MovieSceneHelpers::FindSectionAtTime( Sections, Position );
+	const UMovieSceneSection* Section = MovieSceneHelpers::FindNearestSectionAtTime( Sections, Position );
 
 	if( Section )
 	{
+		Position = FMath::Clamp(Position, Section->GetStartTime(), Section->GetEndTime());
+
 		OutBool = CastChecked<UMovieSceneBoolSection>( Section )->Eval( Position );
 	}
 
diff --git a/Engine/Source/Runtime/MovieSceneTracks/Private/Tracks/MovieSceneByteTrack.cpp b/Engine/Source/Runtime/MovieSceneTracks/Private/Tracks/MovieSceneByteTrack.cpp
--- a/Engine/Source/Runtime/MovieSceneTracks/Private/Tracks/MovieSceneByteTrack.cpp
+++ b/Engine/Source/Runtime/MovieSceneTracks/Private/Tracks/MovieSceneByteTrack.cpp
@@ -40,10 +40,12 @@ bool UMovieSceneByteTrack::AddKeyToSection( float Time, uint8 Value )
 
 bool UMovieSceneByteTrack::Eval( float Position, float LastPostion, uint8& OutByte ) const
 {	
-	const UMovieSceneSection* Section = MovieSceneHelpers::FindSectionAtTime( Sections, Position );
+	const UMovieSceneSection* Section = MovieSceneHelpers::FindNearestSectionAtTime( Sections, Position );
 
 	if( Section )
 	{
+		Position = FMath::Clamp(Position, Section->GetStartTime(), Section->GetEndTime());
+
 		OutByte = CastChecked<UMovieSceneByteSection>( Section )->Eval( Position );
 	}
 
diff --git a/Engine/Source/Runtime/MovieSceneTracks/Private/Tracks/MovieSceneColorTrack.cpp b/Engine/Source/Runtime/MovieSceneTracks/Private/Tracks/MovieSceneColorTrack.cpp
--- a/Engine/Source/Runtime/MovieSceneTracks/Private/Tracks/MovieSceneColorTrack.cpp
+++ b/Engine/Source/Runtime/MovieSceneTracks/Private/Tracks/MovieSceneColorTrack.cpp
@@ -40,10 +40,12 @@ bool UMovieSceneColorTrack::AddKeyToSection( float Time, const FColorKey& Key )
 
 bool UMovieSceneColorTrack::Eval( float Position, float LastPosition, FLinearColor& OutColor ) const
 {
-	const UMovieSceneSection* Section = MovieSceneHelpers::FindSectionAtTime( Sections, Position );
+	const UMovieSceneSection* Section = MovieSceneHelpers::FindNearestSectionAtTime( Sections, Position );
 
 	if( Section )
 	{
+		Position = FMath::Clamp(Position, Section->GetStartTime(), Section->GetEndTime());
+
 		OutColor = CastChecked<UMovieSceneColorSection>( Section )->Eval( Position, OutColor );
 	}
 
diff --git a/Engine/Source/Runtime/MovieSceneTracks/Private/Tracks/MovieSceneFloatTrack.cpp b/Engine/Source/Runtime/MovieSceneTracks/Private/Tracks/MovieSceneFloatTrack.cpp
--- a/Engine/Source/Runtime/MovieSceneTracks/Private/Tracks/MovieSceneFloatTrack.cpp
+++ b/Engine/Source/Runtime/MovieSceneTracks/Private/Tracks/MovieSceneFloatTrack.cpp
@@ -41,10 +41,12 @@ bool UMovieSceneFloatTrack::AddKeyToSection( float Time, float Value )
 
 bool UMovieSceneFloatTrack::Eval( float Position, float LastPosition, float& OutFloat ) const
 {
-	const UMovieSceneSection* Section = MovieSceneHelpers::FindSectionAtTime( Sections, Position );
+	const UMovieSceneSection* Section = MovieSceneHelpers::FindNearestSectionAtTime( Sections, Position );
 
 	if( Section )
 	{
+		Position = FMath::Clamp(Position, Section->GetStartTime(), Section->GetEndTime());
+
 		OutFloat = CastChecked<UMovieSceneFloatSection>( Section )->Eval( Position );
 	}
 
diff --git a/Engine/Source/Runtime/MovieSceneTracks/Private/Tracks/MovieSceneVectorTrack.cpp b/Engine/Source/Runtime/MovieSceneTracks/Private/Tracks/MovieSceneVectorTrack.cpp
--- a/Engine/Source/Runtime/MovieSceneTracks/Private/Tracks/MovieSceneVectorTrack.cpp
+++ b/Engine/Source/Runtime/MovieSceneTracks/Private/Tracks/MovieSceneVectorTrack.cpp
@@ -80,10 +80,12 @@ bool UMovieSceneVectorTrack::AddKeyToSection( float Time, const FVectorKey& Key
 
 bool UMovieSceneVectorTrack::Eval( float Position, float LastPosition, FVector4& InOutVector ) const
 {
-	const UMovieSceneSection* Section = MovieSceneHelpers::FindSectionAtTime( Sections, Position );
+	const UMovieSceneSection* Section = MovieSceneHelpers::FindNearestSectionAtTime( Sections, Position );
 
 	if( Section )
 	{
+		Position = FMath::Clamp(Position, Section->GetStartTime(), Section->GetEndTime());
+
 		InOutVector = CastChecked<UMovieSceneVectorSection>( Section )->Eval( Position, InOutVector );
 	}
 
diff --git a/Engine/Source/Runtime/UMG/Private/Animation/MovieScene2DTransformTrack.cpp b/Engine/Source/Runtime/UMG/Private/Animation/MovieScene2DTransformTrack.cpp
--- a/Engine/Source/Runtime/UMG/Private/Animation/MovieScene2DTransformTrack.cpp
+++ b/Engine/Source/Runtime/UMG/Private/Animation/MovieScene2DTransformTrack.cpp
@@ -24,12 +24,14 @@ TSharedPtr<IMovieSceneTrackInstance> UMovieScene2DTransformTrack::CreateInstance
 
 bool UMovieScene2DTransformTrack::Eval(float Position, float LastPosition, FWidgetTransform& InOutTransform) const
 {
-	const UMovieSceneSection* Section = MovieSceneHelpers::FindSectionAtTime(Sections, Position);
+	const UMovieSceneSection* Section = MovieSceneHelpers::FindNearestSectionAtTime(Sections, Position);
 
 	if(Section)
 	{
 		const UMovieScene2DTransformSection* TransformSection = CastChecked<UMovieScene2DTransformSection>(Section);
 
+		Position = FMath::Clamp(Position, Section->GetStartTime(), Section->GetEndTime());
+
 		InOutTransform = TransformSection->Eval(Position, InOutTransform);
 	}
 
@@ -39,7 +41,7 @@ bool UMovieScene2DTransformTrack::Eval(float Position, float LastPosition, FWidg
 
 bool UMovieScene2DTransformTrack::AddKeyToSection(float Time, const F2DTransformKey& TransformKey)
 {
-	const UMovieSceneSection* NearestSection = MovieSceneHelpers::FindSectionAtTime(Sections, Time);
+	const UMovieSceneSection* NearestSection = MovieSceneHelpers::FindNearestSectionAtTime(Sections, Time);
 	if(!NearestSection || TransformKey.bAddKeyEvenIfUnchanged || CastChecked<UMovieScene2DTransformSection>(NearestSection)->NewKeyIsNewData(Time, TransformKey.Value))
 	{
 		Modify();
diff --git a/Engine/Source/Runtime/UMG/Private/Animation/MovieSceneMarginTrack.cpp b/Engine/Source/Runtime/UMG/Private/Animation/MovieSceneMarginTrack.cpp
--- a/Engine/Source/Runtime/UMG/Private/Animation/MovieSceneMarginTrack.cpp
+++ b/Engine/Source/Runtime/UMG/Private/Animation/MovieSceneMarginTrack.cpp
@@ -42,12 +42,14 @@ bool UMovieSceneMarginTrack::AddKeyToSection( float Time, const FMarginKey& Marg
 
 bool UMovieSceneMarginTrack::Eval( float Position, float LastPosition, FMargin& InOutMargin ) const
 {
-	const UMovieSceneSection* Section = MovieSceneHelpers::FindSectionAtTime( Sections, Position );
+	const UMovieSceneSection* Section = MovieSceneHelpers::FindNearestSectionAtTime( Sections, Position );
 
 	if( Section )
 	{
 		const UMovieSceneMarginSection* MarginSection = CastChecked<UMovieSceneMarginSection>( Section );
 
+		Position = FMath::Clamp(Position, Section->GetStartTime(), Section->GetEndTime());
+
 		InOutMargin = MarginSection->Eval( Position, InOutMargin );
 	}
 
