Sequencer: Audio active/inactive.

diff --git a/Engine/Source/Runtime/MovieSceneTracks/Private/TrackInstances/MovieSceneAudioTrackInstance.cpp b/Engine/Source/Runtime/MovieSceneTracks/Private/TrackInstances/MovieSceneAudioTrackInstance.cpp
--- a/Engine/Source/Runtime/MovieSceneTracks/Private/TrackInstances/MovieSceneAudioTrackInstance.cpp
+++ b/Engine/Source/Runtime/MovieSceneTracks/Private/TrackInstances/MovieSceneAudioTrackInstance.cpp
@@ -28,8 +28,11 @@ void FMovieSceneAudioTrackInstance::Update( float Position, float LastPosition,
 			for (int32 i = 0; i < AudioSections.Num(); ++i)
 			{
 				UMovieSceneAudioSection* AudioSection = Cast<UMovieSceneAudioSection>(AudioSections[i]);
-				int32 SectionIndex = AudioSection->GetRowIndex();
-				AudioSectionsBySectionIndex.FindOrAdd(SectionIndex).Add(AudioSection);
+				if (AudioSection->IsActive())
+				{
+					int32 SectionIndex = AudioSection->GetRowIndex();
+					AudioSectionsBySectionIndex.FindOrAdd(SectionIndex).Add(AudioSection);
+				}
 			}
 
 			for (TMap<int32, TArray<UMovieSceneAudioSection*> >::TIterator It(AudioSectionsBySectionIndex); It; ++It)
@@ -78,19 +81,22 @@ void FMovieSceneAudioTrackInstance::Update( float Position, float LastPosition,
 			for (int32 i = 0; i < AudioSections.Num(); ++i)
 			{
 				auto AudioSection = Cast<UMovieSceneAudioSection>(AudioSections[i]);
-				int32 RowIndex = AudioSection->GetRowIndex();
-				
-				for (int32 ActorIndex = 0; ActorIndex < Actors.Num(); ++ActorIndex)
+				if (AudioSection->IsActive())
 				{
-					TWeakObjectPtr<UAudioComponent> Component = GetAudioComponent(Actors[ActorIndex], RowIndex);
-					if (Component.IsValid())
+					int32 RowIndex = AudioSection->GetRowIndex();
+				
+					for (int32 ActorIndex = 0; ActorIndex < Actors.Num(); ++ActorIndex)
 					{
-						if (AudioSection->IsTimeWithinAudioRange(Position) && !Component->IsPlaying())
+						TWeakObjectPtr<UAudioComponent> Component = GetAudioComponent(Actors[ActorIndex], RowIndex);
+						if (Component.IsValid())
 						{
-							PlaySound(AudioSection, Component, Position);
-							// Fade out the sound at the same volume in order to simply
-							// set a short duration on the sound, far from ideal soln
-							Component->FadeOut(AudioTrackConstants::ScrubDuration, 1.f);
+							if (AudioSection->IsTimeWithinAudioRange(Position) && !Component->IsPlaying())
+							{
+								PlaySound(AudioSection, Component, Position);
+								// Fade out the sound at the same volume in order to simply
+								// set a short duration on the sound, far from ideal soln
+								Component->FadeOut(AudioTrackConstants::ScrubDuration, 1.f);
+							}
 						}
 					}
 				}
