Sequencer: Skeletal animation and particle active/inactive.

diff --git a/Engine/Source/Runtime/MovieSceneTracks/Private/TrackInstances/MovieSceneParticleTrackInstance.cpp b/Engine/Source/Runtime/MovieSceneTracks/Private/TrackInstances/MovieSceneParticleTrackInstance.cpp
--- a/Engine/Source/Runtime/MovieSceneTracks/Private/TrackInstances/MovieSceneParticleTrackInstance.cpp
+++ b/Engine/Source/Runtime/MovieSceneTracks/Private/TrackInstances/MovieSceneParticleTrackInstance.cpp
@@ -23,22 +23,25 @@ void FMovieSceneParticleTrackInstance::Update( float Position, float LastPositio
 		for (int32 i = 0; i < Sections.Num(); ++i)
 		{
 			UMovieSceneParticleSection* Section = Cast<UMovieSceneParticleSection>(Sections[i]);
-			if (Section->GetKeyType() == EParticleKey::Trigger)
+			if (Section->IsActive())
 			{
-				if (Position > Section->GetStartTime() && LastPosition < Section->GetStartTime())
+				if (Section->GetKeyType() == EParticleKey::Trigger)
 				{
-					bTrigger = true;
-				}
-			}
-			else if (Section->GetKeyType() == EParticleKey::Toggle)
-			{
-				if (Position >= Section->GetStartTime() && Position <= Section->GetEndTime())
-				{
-					bOn = true;
+					if (Position > Section->GetStartTime() && LastPosition < Section->GetStartTime())
+					{
+						bTrigger = true;
+					}
 				}
-				else if (Position >= Section->GetEndTime() && LastPosition < Section->GetEndTime())
+				else if (Section->GetKeyType() == EParticleKey::Toggle)
 				{
-					bOff = true;
+					if (Position >= Section->GetStartTime() && Position <= Section->GetEndTime())
+					{
+						bOn = true;
+					}
+					else if (Position >= Section->GetEndTime() && LastPosition < Section->GetEndTime())
+					{
+						bOff = true;
+					}
 				}
 			}
 		}
diff --git a/Engine/Source/Runtime/MovieSceneTracks/Private/TrackInstances/MovieSceneSkeletalAnimationTrackInstance.cpp b/Engine/Source/Runtime/MovieSceneTracks/Private/TrackInstances/MovieSceneSkeletalAnimationTrackInstance.cpp
--- a/Engine/Source/Runtime/MovieSceneTracks/Private/TrackInstances/MovieSceneSkeletalAnimationTrackInstance.cpp
+++ b/Engine/Source/Runtime/MovieSceneTracks/Private/TrackInstances/MovieSceneSkeletalAnimationTrackInstance.cpp
@@ -27,7 +27,7 @@ void FMovieSceneSkeletalAnimationTrackInstance::Update( float Position, float La
 		if (AnimInterface)
 		{
 			UMovieSceneSkeletalAnimationSection* AnimSection = Cast<UMovieSceneSkeletalAnimationSection>(AnimationTrack->GetAnimSectionAtTime(Position));
-			if (AnimSection)
+			if (AnimSection && AnimSection->IsActive())
 			{
 				int32 ChannelIndex = 0;
 				FName SlotName = FName("AnimationSlot");
