Sequencer: Fix setting the color when it's an FColor. Light component's SetLightColor defaults to SRGB conversion, so invert the conversion before applying the value.

diff --git a/Engine/Source/Runtime/MovieSceneTracks/Private/TrackInstances/MovieSceneColorTrackInstance.cpp b/Engine/Source/Runtime/MovieSceneTracks/Private/TrackInstances/MovieSceneColorTrackInstance.cpp
--- a/Engine/Source/Runtime/MovieSceneTracks/Private/TrackInstances/MovieSceneColorTrackInstance.cpp
+++ b/Engine/Source/Runtime/MovieSceneTracks/Private/TrackInstances/MovieSceneColorTrackInstance.cpp
@@ -97,6 +97,13 @@ void FMovieSceneColorTrackInstance::Update( float Position, float LastPosition,
 			FLinearColor ColorValue = ColorType == EColorType::Linear ? PropertyBindings->GetCurrentValue<FLinearColor>(Object) : PropertyBindings->GetCurrentValue<FColor>(Object).ReinterpretAsLinear();
 			if(ColorTrack->Eval(Position, LastPosition, ColorValue))
 			{
+				// LightComponent's SetLightColor applies an sRGB conversion by default, so invert it here before applying the value 
+				if (ColorType == EColorType::RegularColor)
+				{
+					const bool sRGB = false;
+					ColorValue = FLinearColor(ColorValue.ToFColor(sRGB));
+				}
+
 				PropertyBindings->CallFunction(Object, &ColorValue);
 			}
 		}
