Sequencer: Allow setting properties for those that are marked Interp but don't have setter functions.

diff --git a/Engine/Source/Editor/Sequencer/Private/SequencerObjectChangeListener.cpp b/Engine/Source/Editor/Sequencer/Private/SequencerObjectChangeListener.cpp
--- a/Engine/Source/Editor/Sequencer/Private/SequencerObjectChangeListener.cpp
+++ b/Engine/Source/Editor/Sequencer/Private/SequencerObjectChangeListener.cpp
@@ -134,34 +134,31 @@ bool FSequencerObjectChangeListener::FindPropertySetter( const UClass& ObjectCla
 
 		static const FName DeprecatedFunctionName(TEXT("DeprecatedFunction"));
 		UFunction* Function = ObjectClass.FindFunctionByName(FunctionName);
+		bool bFoundValidFunction = false;
 		if( Function && !Function->HasMetaData(DeprecatedFunctionName) )
 		{
-			bFound = true;
+			bFoundValidFunction = true;
+		}
 
-			if (Sequencer.IsValid() && Sequencer.Pin()->GetKeyInterpPropertiesOnly())
+		bool bFoundValidInterp = false;
+		if (StructProperty != 0)
+		{
+			if (StructProperty->HasAnyPropertyFlags(CPF_Interp))
 			{
-				if (StructProperty != 0)
-				{
-					if (!StructProperty->HasAnyPropertyFlags(CPF_Interp))
-					{
-						bFound = false;
-					}
-				}
-				else
-				{
-					UProperty* Property = ObjectClass.FindPropertyByName(FName(*InPropertyVarName));
-
-					if (!Property || !Property->HasAnyPropertyFlags(CPF_Interp))
-					{
-						bFound = false;
-					}
-				}
+				bFoundValidInterp = true;
 			}
 		}
 		else
 		{
-			bFound = false;
+			UProperty* Property = ObjectClass.FindPropertyByName(FName(*InPropertyVarName));
+			if (Property && Property->HasAnyPropertyFlags(CPF_Interp))
+			{
+				bFoundValidInterp = true;
+			}
 		}
+		
+		// Valid either if a setter function is foud or a property with the Interp keyword is found.
+		bFound = bFoundValidFunction || bFoundValidInterp;
 	}
 
 	return bFound;
diff --git a/Engine/Source/Runtime/MovieScene/Private/MovieSceneCommonHelpers.cpp b/Engine/Source/Runtime/MovieScene/Private/MovieSceneCommonHelpers.cpp
--- a/Engine/Source/Runtime/MovieScene/Private/MovieSceneCommonHelpers.cpp
+++ b/Engine/Source/Runtime/MovieScene/Private/MovieSceneCommonHelpers.cpp
@@ -132,17 +132,6 @@ FTrackInstancePropertyBindings::FTrackInstancePropertyBindings( FName InProperty
 	}
 }
 
-
-void FTrackInstancePropertyBindings::CallFunction( UObject* InRuntimeObject, void* FunctionParams )
-{
-	FPropertyAndFunction PropAndFunction = RuntimeObjectToFunctionMap.FindRef(InRuntimeObject);
-	if(PropAndFunction.Function)
-	{
-		InRuntimeObject->ProcessEvent(PropAndFunction.Function, FunctionParams);
-	}
-}
-
-
 FTrackInstancePropertyBindings::FPropertyAddress FTrackInstancePropertyBindings::FindPropertyRecursive( const UObject* Object, void* BasePointer, UStruct* InStruct, TArray<FString>& InPropertyNames, uint32 Index ) const
 {
 	UProperty* Property = FindField<UProperty>(InStruct, *InPropertyNames[Index]);
@@ -201,16 +190,8 @@ void FTrackInstancePropertyBindings::UpdateBindings( const TArray<UObject*>& InR
 		FPropertyAndFunction PropAndFunction;
 
 		PropAndFunction.Function = Object->FindFunction(FunctionName);
-		if(PropAndFunction.Function)
-		{
-			PropAndFunction.PropertyAddress = FindProperty( Object, PropertyPath );
-			RuntimeObjectToFunctionMap.Add(Object, PropAndFunction);
-		}
-		else
-		{
-			// Don't call potentially invalid functions
-			RuntimeObjectToFunctionMap.Remove(Object);
-		}
+		PropAndFunction.PropertyAddress = FindProperty( Object, PropertyPath );
+		RuntimeObjectToFunctionMap.Add(Object, PropAndFunction);
 	}
 }
 
diff --git a/Engine/Source/Runtime/MovieScene/Public/MovieSceneCommonHelpers.h b/Engine/Source/Runtime/MovieScene/Public/MovieSceneCommonHelpers.h
--- a/Engine/Source/Runtime/MovieScene/Public/MovieSceneCommonHelpers.h
+++ b/Engine/Source/Runtime/MovieScene/Public/MovieSceneCommonHelpers.h
@@ -46,12 +46,26 @@ public:
 	FTrackInstancePropertyBindings( FName InPropertyName, const FString& InPropertyPath, const FName& InFunctionName = FName());
 
 	/**
-	 * Calls the setter function for a specific runtime object
+	 * Calls the setter function for a specific runtime object or if the setter function does not exist, the property is set directly
 	 *
 	 * @param InRuntimeObject	The runtime object whose function to call
 	 * @param FunctionParams	Parameters to pass to the function
 	 */
-	void CallFunction( UObject* InRuntimeObject, void* FunctionParams );
+	template <typename ValueType>
+	void 
+	CallFunction( UObject* InRuntimeObject, void* FunctionParams )
+	{
+		FPropertyAndFunction PropAndFunction = RuntimeObjectToFunctionMap.FindRef(InRuntimeObject);
+		if(PropAndFunction.Function)
+		{
+			InRuntimeObject->ProcessEvent(PropAndFunction.Function, FunctionParams);
+		}
+		else
+		{
+			ValueType* NewValue = (ValueType*)(FunctionParams);
+			SetCurrentValue<ValueType>(InRuntimeObject, *NewValue);
+		}
+	}
 
 	/**
 	 * Rebuilds the property and function mappings for a set of runtime objects
@@ -90,6 +104,28 @@ public:
 
 		return ValueType();
 	}
+
+	/**
+	 * Sets the current value of a property on an object
+	 *
+	 * @param Object	The object to set the property on
+	 * @param InValue   The value to set
+	 */
+	template <typename ValueType>
+	void SetCurrentValue(const UObject* Object, ValueType InValue)
+	{
+		FPropertyAndFunction PropAndFunction = RuntimeObjectToFunctionMap.FindRef(Object);
+
+		if(PropAndFunction.PropertyAddress.Address)
+		{
+			ValueType* Val = PropAndFunction.PropertyAddress.Property->ContainerPtrToValuePtr<ValueType>(PropAndFunction.PropertyAddress.Address);
+			if(Val)
+			{
+				*Val = InValue;
+			}
+		}
+	}
+
 private:
 
 	struct FPropertyAddress
diff --git a/Engine/Source/Runtime/MovieSceneTracks/Private/TrackInstances/MovieSceneBoolTrackInstance.cpp b/Engine/Source/Runtime/MovieSceneTracks/Private/TrackInstances/MovieSceneBoolTrackInstance.cpp
--- a/Engine/Source/Runtime/MovieSceneTracks/Private/TrackInstances/MovieSceneBoolTrackInstance.cpp
+++ b/Engine/Source/Runtime/MovieSceneTracks/Private/TrackInstances/MovieSceneBoolTrackInstance.cpp
@@ -35,7 +35,7 @@ void FMovieSceneBoolTrackInstance::RestoreState(const TArray<UObject*>& RuntimeO
 		bool *BoolValue = InitBoolMap.Find(Object);
 		if (BoolValue != NULL)
 		{
-			PropertyBindings->CallFunction(Object, BoolValue);
+			PropertyBindings->CallFunction<bool>(Object, BoolValue);
 		}
 	}
 
@@ -50,7 +50,7 @@ void FMovieSceneBoolTrackInstance::Update( float Position, float LastPosition, c
 	{
 		for( UObject* Object : RuntimeObjects )
 		{
-			PropertyBindings->CallFunction( Object, &BoolValue );
+			PropertyBindings->CallFunction<bool>( Object, &BoolValue );
 		}
 	}
 }
diff --git a/Engine/Source/Runtime/MovieSceneTracks/Private/TrackInstances/MovieSceneByteTrackInstance.cpp b/Engine/Source/Runtime/MovieSceneTracks/Private/TrackInstances/MovieSceneByteTrackInstance.cpp
--- a/Engine/Source/Runtime/MovieSceneTracks/Private/TrackInstances/MovieSceneByteTrackInstance.cpp
+++ b/Engine/Source/Runtime/MovieSceneTracks/Private/TrackInstances/MovieSceneByteTrackInstance.cpp
@@ -35,7 +35,7 @@ void FMovieSceneByteTrackInstance::RestoreState(const TArray<UObject*>& RuntimeO
 		uint8 *ByteValue = InitByteMap.Find(Object);
 		if (ByteValue != NULL)
 		{
-			PropertyBindings->CallFunction(Object, ByteValue);
+			PropertyBindings->CallFunction<uint8>(Object, ByteValue);
 		}
 	}
 
@@ -50,7 +50,7 @@ void FMovieSceneByteTrackInstance::Update( float Position, float LastPosition, c
 	{
 		for( UObject* Object : RuntimeObjects )
 		{
-			PropertyBindings->CallFunction( Object, &ByteValue );
+			PropertyBindings->CallFunction<uint8>( Object, &ByteValue );
 		}
 	}
 }
diff --git a/Engine/Source/Runtime/MovieSceneTracks/Private/TrackInstances/MovieSceneColorTrackInstance.cpp b/Engine/Source/Runtime/MovieSceneTracks/Private/TrackInstances/MovieSceneColorTrackInstance.cpp
--- a/Engine/Source/Runtime/MovieSceneTracks/Private/TrackInstances/MovieSceneColorTrackInstance.cpp
+++ b/Engine/Source/Runtime/MovieSceneTracks/Private/TrackInstances/MovieSceneColorTrackInstance.cpp
@@ -57,16 +57,16 @@ void FMovieSceneColorTrackInstance::RestoreState(const TArray<UObject*>& Runtime
 			FSlateColor* ColorValue = InitSlateColorMap.Find(Object);
 			if (ColorValue != NULL)
 			{
-				PropertyBindings->CallFunction(Object, ColorValue);
+				PropertyBindings->CallFunction<FSlateColor>(Object, ColorValue);
 			}
 		}
 		else if( ColorType == EColorType::Linear || ColorType == EColorType::RegularColor )
 		{
-			// todo 
+			//todo
 			FLinearColor* ColorValue = InitLinearColorMap.Find(Object);
 			if (ColorValue != NULL)
 			{
-				PropertyBindings->CallFunction(Object, ColorValue);
+				PropertyBindings->CallFunction<FLinearColor>(Object, ColorValue);
 			}
 		}
 		else
@@ -89,7 +89,7 @@ void FMovieSceneColorTrackInstance::Update( float Position, float LastPosition,
 			if(ColorTrack->Eval(Position, LastPosition, LinearColor))
 			{
 				FSlateColor NewColor(LinearColor);
-				PropertyBindings->CallFunction(Object, &NewColor);
+				PropertyBindings->CallFunction<FSlateColor>(Object, &NewColor);
 			}
 		}
 		else
@@ -104,7 +104,7 @@ void FMovieSceneColorTrackInstance::Update( float Position, float LastPosition,
 					ColorValue = FLinearColor(ColorValue.ToFColor(sRGB));
 				}
 
-				PropertyBindings->CallFunction(Object, &ColorValue);
+				PropertyBindings->CallFunction<FLinearColor>(Object, &ColorValue);
 			}
 		}
 	}
diff --git a/Engine/Source/Runtime/MovieSceneTracks/Private/TrackInstances/MovieSceneFloatTrackInstance.cpp b/Engine/Source/Runtime/MovieSceneTracks/Private/TrackInstances/MovieSceneFloatTrackInstance.cpp
--- a/Engine/Source/Runtime/MovieSceneTracks/Private/TrackInstances/MovieSceneFloatTrackInstance.cpp
+++ b/Engine/Source/Runtime/MovieSceneTracks/Private/TrackInstances/MovieSceneFloatTrackInstance.cpp
@@ -35,7 +35,7 @@ void FMovieSceneFloatTrackInstance::RestoreState(const TArray<UObject*>& Runtime
 		float *FloatValue = InitFloatMap.Find(Object);
 		if (FloatValue != NULL)
 		{
-			PropertyBindings->CallFunction(Object, FloatValue);
+			PropertyBindings->CallFunction<float>(Object, FloatValue);
 		}
 	}
 
@@ -50,7 +50,7 @@ void FMovieSceneFloatTrackInstance::Update( float Position, float LastPosition,
 	{
 		for(UObject* Object : RuntimeObjects)
 		{
-			PropertyBindings->CallFunction( Object, &FloatValue );
+			PropertyBindings->CallFunction<float>( Object, &FloatValue );
 		}
 	}
 }
diff --git a/Engine/Source/Runtime/MovieSceneTracks/Private/TrackInstances/MovieSceneVectorTrackInstance.cpp b/Engine/Source/Runtime/MovieSceneTracks/Private/TrackInstances/MovieSceneVectorTrackInstance.cpp
--- a/Engine/Source/Runtime/MovieSceneTracks/Private/TrackInstances/MovieSceneVectorTrackInstance.cpp
+++ b/Engine/Source/Runtime/MovieSceneTracks/Private/TrackInstances/MovieSceneVectorTrackInstance.cpp
@@ -57,7 +57,7 @@ void FMovieSceneVectorTrackInstance::RestoreState(const TArray<UObject*>& Runtim
 				FVector2D *VectorValue = InitVector2DMap.Find(Object);
 				if (VectorValue != NULL)
 				{
-					PropertyBindings->CallFunction(Object, VectorValue);
+					PropertyBindings->CallFunction<FVector2D>(Object, VectorValue);
 				}
 				break;
 			}
@@ -66,7 +66,7 @@ void FMovieSceneVectorTrackInstance::RestoreState(const TArray<UObject*>& Runtim
 				FVector *VectorValue = InitVector3DMap.Find(Object);
 				if (VectorValue != NULL)
 				{
-					PropertyBindings->CallFunction(Object, VectorValue);
+					PropertyBindings->CallFunction<FVector2D>(Object, VectorValue);
 				}
 				break;
 			}
@@ -75,7 +75,7 @@ void FMovieSceneVectorTrackInstance::RestoreState(const TArray<UObject*>& Runtim
 				FVector4 *VectorValue = InitVector4DMap.Find(Object);
 				if (VectorValue != NULL)
 				{
-					PropertyBindings->CallFunction(Object, VectorValue);
+					PropertyBindings->CallFunction<FVector2D>(Object, VectorValue);
 				}
 				break;
 			}
@@ -98,7 +98,7 @@ void FMovieSceneVectorTrackInstance::Update( float Position, float LastPosition,
 				FVector2D Value(Vector.X, Vector.Y);
 				for(UObject* Object : RuntimeObjects)
 				{
-					PropertyBindings->CallFunction(Object, &Value);
+					PropertyBindings->CallFunction<FVector2D>(Object, &Value);
 				}
 				break;
 			}
@@ -107,7 +107,7 @@ void FMovieSceneVectorTrackInstance::Update( float Position, float LastPosition,
 				FVector Value(Vector.X, Vector.Y, Vector.Z);
 				for(UObject* Object : RuntimeObjects)
 				{
-					PropertyBindings->CallFunction(Object, &Value);
+					PropertyBindings->CallFunction<FVector2D>(Object, &Value);
 				}
 				break;
 			}
@@ -115,7 +115,7 @@ void FMovieSceneVectorTrackInstance::Update( float Position, float LastPosition,
 			{
 				for(UObject* Object : RuntimeObjects)
 				{
-					PropertyBindings->CallFunction(Object, &Vector);
+					PropertyBindings->CallFunction<FVector2D>(Object, &Vector);
 				}
 				break;
 			}
diff --git a/Engine/Source/Runtime/UMG/Private/Animation/MovieScene2DTransformTrackInstance.cpp b/Engine/Source/Runtime/UMG/Private/Animation/MovieScene2DTransformTrackInstance.cpp
--- a/Engine/Source/Runtime/UMG/Private/Animation/MovieScene2DTransformTrackInstance.cpp
+++ b/Engine/Source/Runtime/UMG/Private/Animation/MovieScene2DTransformTrackInstance.cpp
@@ -18,7 +18,7 @@ void FMovieScene2DTransformTrackInstance::Update( float Position, float LastPosi
 		FWidgetTransform TransformValue = PropertyBindings->GetCurrentValue<FWidgetTransform>(Object);
 		if(TransformTrack->Eval(Position, LastPosition, TransformValue))
 		{
-			PropertyBindings->CallFunction(Object, &TransformValue);
+			PropertyBindings->CallFunction<FWidgetTransform>(Object, &TransformValue);
 		}
 	}
 }
diff --git a/Engine/Source/Runtime/UMG/Private/Animation/MovieSceneMarginTrackInstance.cpp b/Engine/Source/Runtime/UMG/Private/Animation/MovieSceneMarginTrackInstance.cpp
--- a/Engine/Source/Runtime/UMG/Private/Animation/MovieSceneMarginTrackInstance.cpp
+++ b/Engine/Source/Runtime/UMG/Private/Animation/MovieSceneMarginTrackInstance.cpp
@@ -19,7 +19,7 @@ void FMovieSceneMarginTrackInstance::Update( float Position, float LastPosition,
 		FMargin MarginValue = PropertyBindings->GetCurrentValue<FMargin>( Object );
 		if(MarginTrack->Eval(Position, LastPosition, MarginValue))
 		{
-			PropertyBindings->CallFunction(Object, &MarginValue);
+			PropertyBindings->CallFunction<FMargin>(Object, &MarginValue);
 		}
 	}
 }
