MovieScene: Renamed external log category

diff --git a/Engine/Source/Runtime/MovieScene/Private/MovieSceneModule.cpp b/Engine/Source/Runtime/MovieScene/Private/MovieSceneModule.cpp
--- a/Engine/Source/Runtime/MovieScene/Private/MovieSceneModule.cpp
+++ b/Engine/Source/Runtime/MovieScene/Private/MovieSceneModule.cpp
@@ -3,21 +3,22 @@
 #include "MovieScenePrivatePCH.h"
 
 
-DEFINE_LOG_CATEGORY(LogSequencerRuntime);
+DEFINE_LOG_CATEGORY(LogMovieScene);
 
 
 /**
- * MovieScene module implementation (private)
+ * MovieScene module implementation.
  */
 class FMovieSceneModule
 	: public IMovieSceneModule
 {
 public:
 
-	/** IModuleInterface */
+	// IModuleInterface interface
+
 	virtual void StartupModule() override { }
 	virtual void ShutdownModule() override { }
 };
 
 
-IMPLEMENT_MODULE( FMovieSceneModule, MovieScene );
+IMPLEMENT_MODULE(FMovieSceneModule, MovieScene);
diff --git a/Engine/Source/Runtime/MovieScene/Public/MovieScene.h b/Engine/Source/Runtime/MovieScene/Public/MovieScene.h
--- a/Engine/Source/Runtime/MovieScene/Public/MovieScene.h
+++ b/Engine/Source/Runtime/MovieScene/Public/MovieScene.h
@@ -13,7 +13,7 @@ class UMovieSceneSection;
 class UMovieSceneTrack;
 
 
-MOVIESCENE_API DECLARE_LOG_CATEGORY_EXTERN(LogSequencerRuntime, Log, All);
+MOVIESCENE_API DECLARE_LOG_CATEGORY_EXTERN(LogMovieScene, Log, All);
 
 
 /** @todo: remove this type when support for intrinsics on TMap values is added? */
diff --git a/Engine/Source/Runtime/MovieSceneTracks/Private/TrackInstances/MovieSceneVectorTrackInstance.cpp b/Engine/Source/Runtime/MovieSceneTracks/Private/TrackInstances/MovieSceneVectorTrackInstance.cpp
--- a/Engine/Source/Runtime/MovieSceneTracks/Private/TrackInstances/MovieSceneVectorTrackInstance.cpp
+++ b/Engine/Source/Runtime/MovieSceneTracks/Private/TrackInstances/MovieSceneVectorTrackInstance.cpp
@@ -4,6 +4,7 @@
 #include "MovieSceneVectorTrackInstance.h"
 #include "MovieSceneVectorTrack.h"
 
+
 FMovieSceneVectorTrackInstance::FMovieSceneVectorTrackInstance( UMovieSceneVectorTrack& InVectorTrack )
 {
 	VectorTrack = &InVectorTrack;
@@ -120,7 +121,7 @@ void FMovieSceneVectorTrackInstance::Update( float Position, float LastPosition,
 				break;
 			}
 			default:
-				UE_LOG(LogSequencerRuntime, Warning, TEXT("Invalid number of channels(%d) for vector track"), NumChannelsUsed );
+				UE_LOG(LogMovieScene, Warning, TEXT("Invalid number of channels(%d) for vector track"), NumChannelsUsed );
 				break;
 		}
 		
