MovieScene: Code cleanup pass

NULL to nullptr, reduced nesting, line breaks

diff --git a/Engine/Source/Runtime/MovieSceneTracks/Private/MovieSceneTracksPrivatePCH.h b/Engine/Source/Runtime/MovieSceneTracks/Private/MovieSceneTracksPrivatePCH.h
--- a/Engine/Source/Runtime/MovieSceneTracks/Private/MovieSceneTracksPrivatePCH.h
+++ b/Engine/Source/Runtime/MovieSceneTracks/Private/MovieSceneTracksPrivatePCH.h
@@ -1,7 +1,7 @@
 // Copyright 1998-2015 Epic Games, Inc. All Rights Reserved.
 
-#ifndef _MOVIE_SCENE_CORE_TYPES_PRIVATE_
-#define _MOVIE_SCENE_CORE_TYPES_PRIVATE_
+#pragma once
+
 
 /* Private dependencies
  *****************************************************************************/
@@ -22,6 +22,3 @@
  *****************************************************************************/
 
 #include "MovieSceneTracksClasses.h"
-
-
-#endif
diff --git a/Engine/Source/Runtime/MovieSceneTracks/Private/Sections/MovieSceneAudioSection.cpp b/Engine/Source/Runtime/MovieSceneTracks/Private/Sections/MovieSceneAudioSection.cpp
--- a/Engine/Source/Runtime/MovieSceneTracks/Private/Sections/MovieSceneAudioSection.cpp
+++ b/Engine/Source/Runtime/MovieSceneTracks/Private/Sections/MovieSceneAudioSection.cpp
@@ -4,15 +4,17 @@
 #include "MovieSceneAudioSection.h"
 #include "SoundDefinitions.h"
 
+
 UMovieSceneAudioSection::UMovieSceneAudioSection( const FObjectInitializer& ObjectInitializer )
 	: Super( ObjectInitializer )
 {
 
-	Sound = NULL;
+	Sound = nullptr;
 	AudioStartTime = 0.f;
 	AudioDilationFactor = 1.f;
 }
 
+
 TRange<float> UMovieSceneAudioSection::GetAudioRange() const
 {
 	return !Sound ? TRange<float>::Empty() :
@@ -20,11 +22,13 @@ TRange<float> UMovieSceneAudioSection::GetAudioRange() const
 		FMath::Min(AudioStartTime + Sound->GetDuration() * AudioDilationFactor, GetEndTime()));
 }
 
+
 TRange<float> UMovieSceneAudioSection::GetAudioTrueRange() const
 {
 	return !Sound ? TRange<float>::Empty() : TRange<float>(AudioStartTime, AudioStartTime + Sound->GetDuration() * AudioDilationFactor);
 }
 
+
 void UMovieSceneAudioSection::MoveSection( float DeltaTime, TSet<FKeyHandle>& KeyHandles )
 {
 	Super::MoveSection(DeltaTime, KeyHandles);
@@ -32,6 +36,7 @@ void UMovieSceneAudioSection::MoveSection( float DeltaTime, TSet<FKeyHandle>& Ke
 	AudioStartTime += DeltaTime;
 }
 
+
 void UMovieSceneAudioSection::DilateSection( float DilationFactor, float Origin, TSet<FKeyHandle>& KeyHandles )
 {
 	AudioStartTime = (AudioStartTime - Origin) * DilationFactor + Origin;
@@ -40,10 +45,13 @@ void UMovieSceneAudioSection::DilateSection( float DilationFactor, float Origin,
 	Super::DilateSection(DilationFactor, Origin, KeyHandles);
 }
 
+
 void UMovieSceneAudioSection::GetKeyHandles(TSet<FKeyHandle>& KeyHandles) const
 {
+	// do nothing
 }
 
+
 void UMovieSceneAudioSection::GetSnapTimes(TArray<float>& OutSnapTimes, bool bGetSectionBorders) const
 {
 	Super::GetSnapTimes(OutSnapTimes, bGetSectionBorders);
@@ -51,11 +59,14 @@ void UMovieSceneAudioSection::GetSnapTimes(TArray<float>& OutSnapTimes, bool bGe
 	{
 		OutSnapTimes.Add(AudioStartTime);
 	}
+
 	float AudioEndTime = GetAudioTrueRange().GetUpperBoundValue();
+
 	if (AudioEndTime < GetEndTime())
 	{
 		OutSnapTimes.Add(AudioEndTime);
 	}
+
 	// @todo Sequencer handle snapping for time dilation
 	// @todo Don't add redundant times (can't use AddUnique due to floating point equality issues)
 }
diff --git a/Engine/Source/Runtime/MovieSceneTracks/Private/Sections/MovieSceneEventSection.cpp b/Engine/Source/Runtime/MovieSceneTracks/Private/Sections/MovieSceneEventSection.cpp
--- a/Engine/Source/Runtime/MovieSceneTracks/Private/Sections/MovieSceneEventSection.cpp
+++ b/Engine/Source/Runtime/MovieSceneTracks/Private/Sections/MovieSceneEventSection.cpp
@@ -20,6 +20,7 @@ void UMovieSceneEventSection::DilateSection(float DilationFactor, float Origin,
 
 void UMovieSceneEventSection::GetKeyHandles(TSet<FKeyHandle>& KeyHandles) const
 {
+	 // @todo gmp: implement event tracks
 }
 
 
diff --git a/Engine/Source/Runtime/MovieSceneTracks/Private/Sections/MovieSceneSkeletalAnimationSection.cpp b/Engine/Source/Runtime/MovieSceneTracks/Private/Sections/MovieSceneSkeletalAnimationSection.cpp
--- a/Engine/Source/Runtime/MovieSceneTracks/Private/Sections/MovieSceneSkeletalAnimationSection.cpp
+++ b/Engine/Source/Runtime/MovieSceneTracks/Private/Sections/MovieSceneSkeletalAnimationSection.cpp
@@ -3,14 +3,16 @@
 #include "MovieSceneTracksPrivatePCH.h"
 #include "MovieSceneSkeletalAnimationSection.h"
 
+
 UMovieSceneSkeletalAnimationSection::UMovieSceneSkeletalAnimationSection( const FObjectInitializer& ObjectInitializer )
 	: Super( ObjectInitializer )
 {
-	AnimSequence = NULL;
+	AnimSequence = nullptr;
 	AnimationStartTime = 0.f;
 	AnimationDilationFactor = 1.f;
 }
 
+
 void UMovieSceneSkeletalAnimationSection::MoveSection( float DeltaTime, TSet<FKeyHandle>& KeyHandles )
 {
 	Super::MoveSection(DeltaTime, KeyHandles);
@@ -18,6 +20,7 @@ void UMovieSceneSkeletalAnimationSection::MoveSection( float DeltaTime, TSet<FKe
 	AnimationStartTime += DeltaTime;
 }
 
+
 void UMovieSceneSkeletalAnimationSection::DilateSection( float DilationFactor, float Origin, TSet<FKeyHandle>& KeyHandles )
 {
 	AnimationStartTime = (AnimationStartTime - Origin) * DilationFactor + Origin;
@@ -26,13 +29,17 @@ void UMovieSceneSkeletalAnimationSection::DilateSection( float DilationFactor, f
 	Super::DilateSection(DilationFactor, Origin, KeyHandles);
 }
 
+
 void UMovieSceneSkeletalAnimationSection::GetKeyHandles(TSet<FKeyHandle>& KeyHandles) const
 {
+	// do nothing
 }
 
+
 void UMovieSceneSkeletalAnimationSection::GetSnapTimes(TArray<float>& OutSnapTimes, bool bGetSectionBorders) const
 {
 	Super::GetSnapTimes(OutSnapTimes, bGetSectionBorders);
+
 	float CurrentTime = GetAnimationStartTime();
 	while (CurrentTime <= GetEndTime())
 	{
@@ -40,6 +47,7 @@ void UMovieSceneSkeletalAnimationSection::GetSnapTimes(TArray<float>& OutSnapTim
 		{
 			OutSnapTimes.Add(CurrentTime);
 		}
+
 		CurrentTime += GetAnimationDuration();
 	}
 }
diff --git a/Engine/Source/Runtime/MovieSceneTracks/Private/TrackInstances/MovieScene3DAttachTrackInstance.cpp b/Engine/Source/Runtime/MovieSceneTracks/Private/TrackInstances/MovieScene3DAttachTrackInstance.cpp
--- a/Engine/Source/Runtime/MovieSceneTracks/Private/TrackInstances/MovieScene3DAttachTrackInstance.cpp
+++ b/Engine/Source/Runtime/MovieSceneTracks/Private/TrackInstances/MovieScene3DAttachTrackInstance.cpp
@@ -7,10 +7,11 @@
 #include "MovieScene3DAttachSection.h"
 #include "IMovieScenePlayer.h"
 
+
 FMovieScene3DAttachTrackInstance::FMovieScene3DAttachTrackInstance( UMovieScene3DAttachTrack& InAttachTrack )
 	: FMovieScene3DConstraintTrackInstance( InAttachTrack )
-{
-}
+{ }
+
 
 void FMovieScene3DAttachTrackInstance::UpdateConstraint( float Position, const TArray<UObject*>& RuntimeObjects, AActor* Actor, UMovieScene3DConstraintSection* ConstraintSection ) 
 {
@@ -23,7 +24,7 @@ void FMovieScene3DAttachTrackInstance::UpdateConstraint( float Position, const T
 	{
 		USceneComponent* SceneComponent = MovieSceneHelpers::SceneComponentFromRuntimeObject(RuntimeObjects[ObjIndex]);
 
-		if (SceneComponent != NULL)
+		if (SceneComponent != nullptr)
 		{
 			AttachSection->Eval(SceneComponent, Position, Actor, Translation, Rotation);
 
@@ -32,4 +33,3 @@ void FMovieScene3DAttachTrackInstance::UpdateConstraint( float Position, const T
 		}
 	}
 }
-
diff --git a/Engine/Source/Runtime/MovieSceneTracks/Private/TrackInstances/MovieScene3DConstraintTrackInstance.cpp b/Engine/Source/Runtime/MovieSceneTracks/Private/TrackInstances/MovieScene3DConstraintTrackInstance.cpp
--- a/Engine/Source/Runtime/MovieSceneTracks/Private/TrackInstances/MovieScene3DConstraintTrackInstance.cpp
+++ b/Engine/Source/Runtime/MovieSceneTracks/Private/TrackInstances/MovieScene3DConstraintTrackInstance.cpp
@@ -8,23 +8,26 @@
 #include "IMovieScenePlayer.h"
 #include "Components/SplineComponent.h"
 
+
 FMovieScene3DConstraintTrackInstance::FMovieScene3DConstraintTrackInstance( UMovieScene3DConstraintTrack& InConstraintTrack )
 {
 	ConstraintTrack = &InConstraintTrack;
 }
 
+
 void FMovieScene3DConstraintTrackInstance::SaveState(const TArray<UObject*>& RuntimeObjects)
 {
 	for (int32 ObjIndex = 0; ObjIndex < RuntimeObjects.Num(); ++ObjIndex)
 	{
 		USceneComponent* SceneComponent = MovieSceneHelpers::SceneComponentFromRuntimeObject(RuntimeObjects[ObjIndex]);
-		if (SceneComponent != NULL)
+		if (SceneComponent != nullptr)
 		{
 			InitTransformMap.Add(RuntimeObjects[ObjIndex], SceneComponent->GetRelativeTransform());
 		}
 	}
 }
 
+
 void FMovieScene3DConstraintTrackInstance::RestoreState(const TArray<UObject*>& RuntimeObjects)
 {
 	for (int32 ObjIndex = 0; ObjIndex < RuntimeObjects.Num(); ++ObjIndex)
@@ -35,10 +38,10 @@ void FMovieScene3DConstraintTrackInstance::RestoreState(const TArray<UObject*>&
 		}
 
 		USceneComponent* SceneComponent = MovieSceneHelpers::SceneComponentFromRuntimeObject(RuntimeObjects[ObjIndex]);
-		if (SceneComponent != NULL)
+		if (SceneComponent != nullptr)
 		{
 			FTransform *Transform = InitTransformMap.Find(RuntimeObjects[ObjIndex]);
-			if (Transform != NULL)
+			if (Transform != nullptr)
 			{
 				SceneComponent->SetRelativeTransform(*Transform);
 			}
@@ -46,9 +49,10 @@ void FMovieScene3DConstraintTrackInstance::RestoreState(const TArray<UObject*>&
 	}
 }
 
+
 void FMovieScene3DConstraintTrackInstance::Update( float Position, float LastPosition, const TArray<UObject*>& RuntimeObjects, class IMovieScenePlayer& Player ) 
 {
-	UMovieScene3DConstraintSection* FirstConstraintSection = NULL;
+	UMovieScene3DConstraintSection* FirstConstraintSection = nullptr;
 
 	const TArray<UMovieSceneSection*>& ConstraintSections = ConstraintTrack->GetAllSections();
 
@@ -57,7 +61,7 @@ void FMovieScene3DConstraintTrackInstance::Update( float Position, float LastPos
 		UMovieScene3DConstraintSection* ConstraintSection = CastChecked<UMovieScene3DConstraintSection>(ConstraintSections[ConstraintIndex]);
 
 		if (ConstraintSection->IsTimeWithinSection(Position) &&
-			(FirstConstraintSection == NULL || FirstConstraintSection->GetRowIndex() > ConstraintSection->GetRowIndex()))
+			(FirstConstraintSection == nullptr || FirstConstraintSection->GetRowIndex() > ConstraintSection->GetRowIndex()))
 		{
 			TArray<UObject*> ConstraintObjects;
 			FGuid ConstraintId = ConstraintSection->GetConstraintId();
@@ -78,4 +82,3 @@ void FMovieScene3DConstraintTrackInstance::Update( float Position, float LastPos
 		}
 	}
 }
-
diff --git a/Engine/Source/Runtime/MovieSceneTracks/Private/TrackInstances/MovieScene3DPathTrackInstance.cpp b/Engine/Source/Runtime/MovieSceneTracks/Private/TrackInstances/MovieScene3DPathTrackInstance.cpp
--- a/Engine/Source/Runtime/MovieSceneTracks/Private/TrackInstances/MovieScene3DPathTrackInstance.cpp
+++ b/Engine/Source/Runtime/MovieSceneTracks/Private/TrackInstances/MovieScene3DPathTrackInstance.cpp
@@ -8,10 +8,11 @@
 #include "IMovieScenePlayer.h"
 #include "Components/SplineComponent.h"
 
+
 FMovieScene3DPathTrackInstance::FMovieScene3DPathTrackInstance( UMovieScene3DPathTrack& InPathTrack )
 	: FMovieScene3DConstraintTrackInstance( InPathTrack )
-{
-}
+{ }
+
 
 void FMovieScene3DPathTrackInstance::UpdateConstraint( float Position, const TArray<UObject*>& RuntimeObjects, AActor* Actor, UMovieScene3DConstraintSection* ConstraintSection ) 
 {
@@ -29,7 +30,7 @@ void FMovieScene3DPathTrackInstance::UpdateConstraint( float Position, const TAr
 		{
 			USceneComponent* SceneComponent = MovieSceneHelpers::SceneComponentFromRuntimeObject(RuntimeObjects[ObjIndex]);
 
-			if (SceneComponent != NULL)
+			if (SceneComponent != nullptr)
 			{
 				PathSection->Eval(SceneComponent, Position, SplineComponents[0], Translation, Rotation);
 
@@ -39,4 +40,3 @@ void FMovieScene3DPathTrackInstance::UpdateConstraint( float Position, const TAr
 		}
 	}
 }
-
diff --git a/Engine/Source/Runtime/MovieSceneTracks/Private/TrackInstances/MovieScene3DTransformTrackInstance.cpp b/Engine/Source/Runtime/MovieSceneTracks/Private/TrackInstances/MovieScene3DTransformTrackInstance.cpp
--- a/Engine/Source/Runtime/MovieSceneTracks/Private/TrackInstances/MovieScene3DTransformTrackInstance.cpp
+++ b/Engine/Source/Runtime/MovieSceneTracks/Private/TrackInstances/MovieScene3DTransformTrackInstance.cpp
@@ -11,18 +11,20 @@ FMovieScene3DTransformTrackInstance::FMovieScene3DTransformTrackInstance( UMovie
 	TransformTrack = &InTransformTrack;
 }
 
+
 void FMovieScene3DTransformTrackInstance::SaveState(const TArray<UObject*>& RuntimeObjects)
 {
 	for (int32 ObjIndex = 0; ObjIndex < RuntimeObjects.Num(); ++ObjIndex)
 	{
 		USceneComponent* SceneComponent = MovieSceneHelpers::SceneComponentFromRuntimeObject(RuntimeObjects[ObjIndex]);
-		if (SceneComponent != NULL)
+		if (SceneComponent != nullptr)
 		{
 			InitTransformMap.Add(RuntimeObjects[ObjIndex], SceneComponent->GetRelativeTransform());
 		}
 	}
 }
 
+
 void FMovieScene3DTransformTrackInstance::RestoreState(const TArray<UObject*>& RuntimeObjects)
 {
 	for (int32 ObjIndex = 0; ObjIndex < RuntimeObjects.Num(); ++ObjIndex)
@@ -33,10 +35,10 @@ void FMovieScene3DTransformTrackInstance::RestoreState(const TArray<UObject*>& R
 		}
 
 		USceneComponent* SceneComponent = MovieSceneHelpers::SceneComponentFromRuntimeObject(RuntimeObjects[ObjIndex]);
-		if (SceneComponent != NULL)
+		if (SceneComponent != nullptr)
 		{
 			FTransform *Transform = InitTransformMap.Find(RuntimeObjects[ObjIndex]);
-			if (Transform != NULL)
+			if (Transform != nullptr)
 			{
 				SceneComponent->SetRelativeTransform(*Transform);
 			}
@@ -44,6 +46,7 @@ void FMovieScene3DTransformTrackInstance::RestoreState(const TArray<UObject*>& R
 	}
 }
 
+
 void FMovieScene3DTransformTrackInstance::Update( float Position, float LastPosition, const TArray<UObject*>& RuntimeObjects, class IMovieScenePlayer& Player ) 
 {
 	FVector Translation;
@@ -56,7 +59,7 @@ void FMovieScene3DTransformTrackInstance::Update( float Position, float LastPosi
 		{
 			USceneComponent* SceneComponent = MovieSceneHelpers::SceneComponentFromRuntimeObject(RuntimeObjects[ObjIndex]);
 
-			if (SceneComponent != NULL)
+			if (SceneComponent != nullptr)
 			{
 				SceneComponent->SetRelativeLocationAndRotation(Translation, Rotation);
 
@@ -65,4 +68,3 @@ void FMovieScene3DTransformTrackInstance::Update( float Position, float LastPosi
 		}
 	}
 }
-
diff --git a/Engine/Source/Runtime/MovieSceneTracks/Private/TrackInstances/MovieSceneAudioTrackInstance.cpp b/Engine/Source/Runtime/MovieSceneTracks/Private/TrackInstances/MovieSceneAudioTrackInstance.cpp
--- a/Engine/Source/Runtime/MovieSceneTracks/Private/TrackInstances/MovieSceneAudioTrackInstance.cpp
+++ b/Engine/Source/Runtime/MovieSceneTracks/Private/TrackInstances/MovieSceneAudioTrackInstance.cpp
@@ -9,11 +9,13 @@
 #include "MovieSceneAudioTrack.h"
 #include "MovieSceneAudioSection.h"
 
+
 FMovieSceneAudioTrackInstance::FMovieSceneAudioTrackInstance( UMovieSceneAudioTrack& InAudioTrack )
 {
 	AudioTrack = &InAudioTrack;
 }
 
+
 void FMovieSceneAudioTrackInstance::Update( float Position, float LastPosition, const TArray<UObject*>& RuntimeObjects, class IMovieScenePlayer& Player ) 
 {
 	const TArray<UMovieSceneSection*>& AudioSections = AudioTrack->GetAudioSections();
@@ -132,24 +134,32 @@ void FMovieSceneAudioTrackInstance::Update( float Position, float LastPosition,
 	}
 }
 
+
 void FMovieSceneAudioTrackInstance::PlaySound(UMovieSceneAudioSection* AudioSection, TWeakObjectPtr<UAudioComponent> Component, float Time)
 {
-	if (Component.IsValid())
+	if (!Component.IsValid())
 	{
-		float PitchMultiplier = 1.f / AudioSection->GetAudioDilationFactor();
-		Component->bAllowSpatialization = !AudioTrack->IsAMasterTrack();
-		Component->Stop();
-		Component->SetSound(AudioSection->GetSound());
-		Component->SetVolumeMultiplier(1.f);
-		Component->SetPitchMultiplier(PitchMultiplier);
-		Component->bIsUISound = true;
-		Component->Play(Time - AudioSection->GetAudioStartTime());
+		return;
 	}
+
+	float PitchMultiplier = 1.f / AudioSection->GetAudioDilationFactor();
+	
+	Component->bAllowSpatialization = !AudioTrack->IsAMasterTrack();
+	Component->Stop();
+	Component->SetSound(AudioSection->GetSound());
+	Component->SetVolumeMultiplier(1.f);
+	Component->SetPitchMultiplier(PitchMultiplier);
+	Component->bIsUISound = true;
+	Component->Play(Time - AudioSection->GetAudioStartTime());
 }
 
+
 void FMovieSceneAudioTrackInstance::StopSound(int32 RowIndex)
 {
-	if (RowIndex >= PlaybackAudioComponents.Num()) {return;}
+	if (RowIndex >= PlaybackAudioComponents.Num())
+	{
+		return;
+	}
 
 	TMap<AActor*, TWeakObjectPtr<UAudioComponent>>& AudioComponents = PlaybackAudioComponents[RowIndex];
 	for (TMap<AActor*, TWeakObjectPtr<UAudioComponent>>::TIterator It(AudioComponents); It; ++It)
@@ -161,6 +171,7 @@ void FMovieSceneAudioTrackInstance::StopSound(int32 RowIndex)
 	}
 }
 
+
 void FMovieSceneAudioTrackInstance::StopAllSounds()
 {
 	for (int32 i = 0; i < PlaybackAudioComponents.Num(); ++i)
@@ -176,6 +187,7 @@ void FMovieSceneAudioTrackInstance::StopAllSounds()
 	}
 }
 
+
 TWeakObjectPtr<UAudioComponent> FMovieSceneAudioTrackInstance::GetAudioComponent(AActor* Actor, int32 RowIndex)
 {
 	if (RowIndex + 1 > PlaybackAudioComponents.Num())
@@ -191,17 +203,20 @@ TWeakObjectPtr<UAudioComponent> FMovieSceneAudioTrackInstance::GetAudioComponent
 	{
 		USoundCue* TempPlaybackAudioCue = NewObject<USoundCue>();
 		
-		UAudioComponent* AudioComponent = FAudioDevice::CreateComponent(TempPlaybackAudioCue, NULL, Actor, false, false);
+		UAudioComponent* AudioComponent = FAudioDevice::CreateComponent(TempPlaybackAudioCue, nullptr, Actor, false, false);
 
 		PlaybackAudioComponents[RowIndex].Add(Actor);
 		PlaybackAudioComponents[RowIndex][Actor] = TWeakObjectPtr<UAudioComponent>(AudioComponent);
 	}
+
 	return *PlaybackAudioComponents[RowIndex].Find(Actor);
 }
 
+
 TArray<AActor*> FMovieSceneAudioTrackInstance::GetRuntimeActors(const TArray<UObject*>& RuntimeObjects) const
 {
 	TArray<AActor*> Actors;
+
 	for (int32 ObjectIndex = 0; ObjectIndex < RuntimeObjects.Num(); ++ObjectIndex)
 	{
 		if (RuntimeObjects[ObjectIndex]->IsA<AActor>())
@@ -209,10 +224,12 @@ TArray<AActor*> FMovieSceneAudioTrackInstance::GetRuntimeActors(const TArray<UOb
 			Actors.Add(Cast<AActor>(RuntimeObjects[ObjectIndex]));
 		}
 	}
+
 	if (AudioTrack->IsAMasterTrack())
 	{
 		check(Actors.Num() == 0);
-		Actors.Add(NULL);
+		Actors.Add(nullptr);
 	}
+
 	return Actors;
 }
diff --git a/Engine/Source/Runtime/MovieSceneTracks/Private/TrackInstances/MovieSceneBoolTrackInstance.cpp b/Engine/Source/Runtime/MovieSceneTracks/Private/TrackInstances/MovieSceneBoolTrackInstance.cpp
--- a/Engine/Source/Runtime/MovieSceneTracks/Private/TrackInstances/MovieSceneBoolTrackInstance.cpp
+++ b/Engine/Source/Runtime/MovieSceneTracks/Private/TrackInstances/MovieSceneBoolTrackInstance.cpp
@@ -33,7 +33,8 @@ void FMovieSceneBoolTrackInstance::RestoreState(const TArray<UObject*>& RuntimeO
 		}
 
 		bool *BoolValue = InitBoolMap.Find(Object);
-		if (BoolValue != NULL)
+
+		if (BoolValue != nullptr)
 		{
 			PropertyBindings->CallFunction<bool>(Object, BoolValue);
 		}
@@ -46,6 +47,7 @@ void FMovieSceneBoolTrackInstance::RestoreState(const TArray<UObject*>& RuntimeO
 void FMovieSceneBoolTrackInstance::Update( float Position, float LastPosition, const TArray<UObject*>& RuntimeObjects, class IMovieScenePlayer& Player ) 
 {
 	bool BoolValue = false;
+
 	if( BoolTrack->Eval( Position, LastPosition, BoolValue ) )
 	{
 		for( UObject* Object : RuntimeObjects )
diff --git a/Engine/Source/Runtime/MovieSceneTracks/Private/TrackInstances/MovieSceneByteTrackInstance.cpp b/Engine/Source/Runtime/MovieSceneTracks/Private/TrackInstances/MovieSceneByteTrackInstance.cpp
--- a/Engine/Source/Runtime/MovieSceneTracks/Private/TrackInstances/MovieSceneByteTrackInstance.cpp
+++ b/Engine/Source/Runtime/MovieSceneTracks/Private/TrackInstances/MovieSceneByteTrackInstance.cpp
@@ -33,7 +33,7 @@ void FMovieSceneByteTrackInstance::RestoreState(const TArray<UObject*>& RuntimeO
 		}
 
 		uint8 *ByteValue = InitByteMap.Find(Object);
-		if (ByteValue != NULL)
+		if (ByteValue != nullptr)
 		{
 			PropertyBindings->CallFunction<uint8>(Object, ByteValue);
 		}
diff --git a/Engine/Source/Runtime/MovieSceneTracks/Private/TrackInstances/MovieSceneColorTrackInstance.cpp b/Engine/Source/Runtime/MovieSceneTracks/Private/TrackInstances/MovieSceneColorTrackInstance.cpp
--- a/Engine/Source/Runtime/MovieSceneTracks/Private/TrackInstances/MovieSceneColorTrackInstance.cpp
+++ b/Engine/Source/Runtime/MovieSceneTracks/Private/TrackInstances/MovieSceneColorTrackInstance.cpp
@@ -14,6 +14,7 @@ FMovieSceneColorTrackInstance::FMovieSceneColorTrackInstance( UMovieSceneColorTr
 	PropertyBindings = MakeShareable( new FTrackInstancePropertyBindings( ColorTrack->GetPropertyName(), ColorTrack->GetPropertyPath() ) );
 }
 
+
 void FMovieSceneColorTrackInstance::SaveState(const TArray<UObject*>& RuntimeObjects)
 {
 	for( UObject* Object : RuntimeObjects )
@@ -43,6 +44,7 @@ void FMovieSceneColorTrackInstance::SaveState(const TArray<UObject*>& RuntimeObj
 	}
 }
 
+
 void FMovieSceneColorTrackInstance::RestoreState(const TArray<UObject*>& RuntimeObjects)
 {
 	for( UObject* Object : RuntimeObjects )
@@ -55,7 +57,7 @@ void FMovieSceneColorTrackInstance::RestoreState(const TArray<UObject*>& Runtime
 		if( ColorType == EColorType::Slate )
 		{
 			FSlateColor* ColorValue = InitSlateColorMap.Find(Object);
-			if (ColorValue != NULL)
+			if (ColorValue != nullptr)
 			{
 				PropertyBindings->CallFunction<FSlateColor>(Object, ColorValue);
 			}
@@ -64,7 +66,7 @@ void FMovieSceneColorTrackInstance::RestoreState(const TArray<UObject*>& Runtime
 		{
 			//todo
 			FLinearColor* ColorValue = InitLinearColorMap.Find(Object);
-			if (ColorValue != NULL)
+			if (ColorValue != nullptr)
 			{
 				PropertyBindings->CallFunction<FLinearColor>(Object, ColorValue);
 			}
@@ -78,6 +80,7 @@ void FMovieSceneColorTrackInstance::RestoreState(const TArray<UObject*>& Runtime
 	PropertyBindings->UpdateBindings( RuntimeObjects );
 }
 
+
 void FMovieSceneColorTrackInstance::Update( float Position, float LastPosition, const TArray<UObject*>& RuntimeObjects, class IMovieScenePlayer& Player ) 
 {
 	for(UObject* Object : RuntimeObjects)
@@ -110,35 +113,37 @@ void FMovieSceneColorTrackInstance::Update( float Position, float LastPosition,
 	}
 }
 
+
 void FMovieSceneColorTrackInstance::RefreshInstance( const TArray<UObject*>& RuntimeObjects, class IMovieScenePlayer& Player )
 {
-	if( RuntimeObjects.Num() > 0 )
+	if( RuntimeObjects.Num() == 0 )
 	{
-		PropertyBindings->UpdateBindings( RuntimeObjects );
+		return;
+	}
 
-		// Cache off what type of color this is. Just examine the first object since the property should be the same
-		UProperty* ColorProp = PropertyBindings->GetProperty( RuntimeObjects[0] );
+	PropertyBindings->UpdateBindings( RuntimeObjects );
 
-		const UStructProperty* StructProp = Cast<const UStructProperty>(ColorProp);
-		if (StructProp && StructProp->Struct)
-		{
-			FName StructName = StructProp->Struct->GetFName();
+	// Cache off what type of color this is. Just examine the first object since the property should be the same
+	UProperty* ColorProp = PropertyBindings->GetProperty( RuntimeObjects[0] );
 
-			static const FName SlateColor("SlateColor");
+	const UStructProperty* StructProp = Cast<const UStructProperty>(ColorProp);
+	if (StructProp && StructProp->Struct)
+	{
+		FName StructName = StructProp->Struct->GetFName();
 
-			if( StructName == NAME_Color )
-			{
-				ColorType = EColorType::RegularColor;
-			}
-			else if( StructName == SlateColor )
-			{
-				ColorType = EColorType::Slate;
-			}
-			else
-			{
-				ColorType = EColorType::Linear;
-			}
+		static const FName SlateColor("SlateColor");
+
+		if( StructName == NAME_Color )
+		{
+			ColorType = EColorType::RegularColor;
+		}
+		else if( StructName == SlateColor )
+		{
+			ColorType = EColorType::Slate;
+		}
+		else
+		{
+			ColorType = EColorType::Linear;
 		}
 	}
-
 }
diff --git a/Engine/Source/Runtime/MovieSceneTracks/Private/TrackInstances/MovieSceneFloatTrackInstance.cpp b/Engine/Source/Runtime/MovieSceneTracks/Private/TrackInstances/MovieSceneFloatTrackInstance.cpp
--- a/Engine/Source/Runtime/MovieSceneTracks/Private/TrackInstances/MovieSceneFloatTrackInstance.cpp
+++ b/Engine/Source/Runtime/MovieSceneTracks/Private/TrackInstances/MovieSceneFloatTrackInstance.cpp
@@ -33,7 +33,7 @@ void FMovieSceneFloatTrackInstance::RestoreState(const TArray<UObject*>& Runtime
 		}
 
 		float *FloatValue = InitFloatMap.Find(Object);
-		if (FloatValue != NULL)
+		if (FloatValue != nullptr)
 		{
 			PropertyBindings->CallFunction<float>(Object, FloatValue);
 		}
diff --git a/Engine/Source/Runtime/MovieSceneTracks/Private/TrackInstances/MovieSceneSkeletalAnimationTrackInstance.cpp b/Engine/Source/Runtime/MovieSceneTracks/Private/TrackInstances/MovieSceneSkeletalAnimationTrackInstance.cpp
--- a/Engine/Source/Runtime/MovieSceneTracks/Private/TrackInstances/MovieSceneSkeletalAnimationTrackInstance.cpp
+++ b/Engine/Source/Runtime/MovieSceneTracks/Private/TrackInstances/MovieSceneSkeletalAnimationTrackInstance.cpp
@@ -7,16 +7,19 @@
 #include "IMovieScenePlayer.h"
 #include "Matinee/MatineeAnimInterface.h"
 
+
 FMovieSceneSkeletalAnimationTrackInstance::FMovieSceneSkeletalAnimationTrackInstance( UMovieSceneSkeletalAnimationTrack& InAnimationTrack )
 {
 	AnimationTrack = &InAnimationTrack;
 }
 
+
 FMovieSceneSkeletalAnimationTrackInstance::~FMovieSceneSkeletalAnimationTrackInstance()
 {
 	// @todo Sequencer Need to find some way to call PreviewFinishAnimControl (needs the runtime objects)
 }
 
+
 void FMovieSceneSkeletalAnimationTrackInstance::Update( float Position, float LastPosition, const TArray<UObject*>& RuntimeObjects, class IMovieScenePlayer& Player ) 
 {
 	// @todo Sequencer gameplay update has a different code path than editor update for animation
@@ -34,7 +37,7 @@ void FMovieSceneSkeletalAnimationTrackInstance::Update( float Position, float La
 				UAnimSequence* AnimSequence = AnimSection->GetAnimSequence();
 				float AnimPosition = FMath::Fmod(Position - AnimSection->GetAnimationStartTime(), AnimSection->GetAnimationDuration()) / AnimSection->GetAnimationDilationFactor();
 
-				AnimInterface->PreviewBeginAnimControl(NULL);
+				AnimInterface->PreviewBeginAnimControl(nullptr);
 				AnimInterface->PreviewSetAnimPosition(SlotName, ChannelIndex, AnimSequence, AnimPosition, true, false, 0.f);
 			}
 		}
diff --git a/Engine/Source/Runtime/MovieSceneTracks/Private/TrackInstances/MovieSceneVectorTrackInstance.cpp b/Engine/Source/Runtime/MovieSceneTracks/Private/TrackInstances/MovieSceneVectorTrackInstance.cpp
--- a/Engine/Source/Runtime/MovieSceneTracks/Private/TrackInstances/MovieSceneVectorTrackInstance.cpp
+++ b/Engine/Source/Runtime/MovieSceneTracks/Private/TrackInstances/MovieSceneVectorTrackInstance.cpp
@@ -8,10 +8,10 @@
 FMovieSceneVectorTrackInstance::FMovieSceneVectorTrackInstance( UMovieSceneVectorTrack& InVectorTrack )
 {
 	VectorTrack = &InVectorTrack;
-
 	PropertyBindings = MakeShareable( new FTrackInstancePropertyBindings( VectorTrack->GetPropertyName(), VectorTrack->GetPropertyPath() ) );
 }
 
+
 void FMovieSceneVectorTrackInstance::SaveState(const TArray<UObject*>& RuntimeObjects)
 {
 	int32 NumChannelsUsed = VectorTrack->GetNumChannelsUsed();
@@ -41,6 +41,7 @@ void FMovieSceneVectorTrackInstance::SaveState(const TArray<UObject*>& RuntimeOb
 	}
 }
 
+
 void FMovieSceneVectorTrackInstance::RestoreState(const TArray<UObject*>& RuntimeObjects)
 {
 	int32 NumChannelsUsed = VectorTrack->GetNumChannelsUsed();
@@ -56,7 +57,7 @@ void FMovieSceneVectorTrackInstance::RestoreState(const TArray<UObject*>& Runtim
 			case 2:
 			{
 				FVector2D *VectorValue = InitVector2DMap.Find(Object);
-				if (VectorValue != NULL)
+				if (VectorValue != nullptr)
 				{
 					PropertyBindings->CallFunction<FVector2D>(Object, VectorValue);
 				}
@@ -65,7 +66,7 @@ void FMovieSceneVectorTrackInstance::RestoreState(const TArray<UObject*>& Runtim
 			case 3:
 			{
 				FVector *VectorValue = InitVector3DMap.Find(Object);
-				if (VectorValue != NULL)
+				if (VectorValue != nullptr)
 				{
 					PropertyBindings->CallFunction<FVector2D>(Object, VectorValue);
 				}
@@ -74,7 +75,7 @@ void FMovieSceneVectorTrackInstance::RestoreState(const TArray<UObject*>& Runtim
 			case 4:
 			{
 				FVector4 *VectorValue = InitVector4DMap.Find(Object);
-				if (VectorValue != NULL)
+				if (VectorValue != nullptr)
 				{
 					PropertyBindings->CallFunction<FVector2D>(Object, VectorValue);
 				}
@@ -86,6 +87,7 @@ void FMovieSceneVectorTrackInstance::RestoreState(const TArray<UObject*>& Runtim
 	PropertyBindings->UpdateBindings( RuntimeObjects );
 }
 
+
 void FMovieSceneVectorTrackInstance::Update( float Position, float LastPosition, const TArray<UObject*>& RuntimeObjects, class IMovieScenePlayer& Player ) 
 {
 	FVector4 Vector;
diff --git a/Engine/Source/Runtime/MovieSceneTracks/Private/TrackInstances/MovieSceneVisibilityTrackInstance.cpp b/Engine/Source/Runtime/MovieSceneTracks/Private/TrackInstances/MovieSceneVisibilityTrackInstance.cpp
--- a/Engine/Source/Runtime/MovieSceneTracks/Private/TrackInstances/MovieSceneVisibilityTrackInstance.cpp
+++ b/Engine/Source/Runtime/MovieSceneTracks/Private/TrackInstances/MovieSceneVisibilityTrackInstance.cpp
@@ -4,20 +4,21 @@
 #include "MovieSceneVisibilityTrackInstance.h"
 #include "MovieSceneVisibilityTrack.h"
 
+
 FMovieSceneVisibilityTrackInstance::FMovieSceneVisibilityTrackInstance(UMovieSceneVisibilityTrack& InVisibilityTrack)
 {
 	VisibilityTrack = &InVisibilityTrack;
-
 	PropertyBindings = MakeShareable( new FTrackInstancePropertyBindings( VisibilityTrack->GetPropertyName(), VisibilityTrack->GetPropertyPath() ) );
 }
 
+
 void FMovieSceneVisibilityTrackInstance::SaveState(const TArray<UObject*>& RuntimeObjects)
 {
  	for( UObject* Object : RuntimeObjects )
  	{
 		AActor* Actor = Cast<AActor>(Object);
 
-		if (Actor != NULL)
+		if (Actor != nullptr)
 		{
 #if WITH_EDITOR
 			if (GIsEditor && !Actor->GetWorld()->IsPlayInEditor())
@@ -31,6 +32,7 @@ void FMovieSceneVisibilityTrackInstance::SaveState(const TArray<UObject*>& Runti
  	}
 }
 
+
 void FMovieSceneVisibilityTrackInstance::RestoreState(const TArray<UObject*>& RuntimeObjects)
 {
  	for( UObject* Object : RuntimeObjects )
@@ -42,10 +44,10 @@ void FMovieSceneVisibilityTrackInstance::RestoreState(const TArray<UObject*>& Ru
 
 		AActor* Actor = Cast<AActor>(Object);
 
-		if (Actor != NULL)
+		if (Actor != nullptr)
 		{
 			bool *HiddenInGameValue = InitHiddenInGameMap.Find(Object);
-			if (HiddenInGameValue != NULL)
+			if (HiddenInGameValue != nullptr)
 			{
 				Actor->SetActorHiddenInGame(*HiddenInGameValue);
 			}
@@ -54,7 +56,7 @@ void FMovieSceneVisibilityTrackInstance::RestoreState(const TArray<UObject*>& Ru
 			if (GIsEditor && !Actor->GetWorld()->IsPlayInEditor())
 			{
 				bool *HiddenInEditorValue = InitHiddenInEditorMap.Find(Object);
-				if (HiddenInEditorValue != NULL)
+				if (HiddenInEditorValue != nullptr)
 				{
 					if (Actor->IsTemporarilyHiddenInEditor() != *HiddenInEditorValue)
 					{
@@ -67,6 +69,7 @@ void FMovieSceneVisibilityTrackInstance::RestoreState(const TArray<UObject*>& Ru
  	}
 }
 
+
 void FMovieSceneVisibilityTrackInstance::Update( float Position, float LastPosition, const TArray<UObject*>& RuntimeObjects, class IMovieScenePlayer& Player ) 
 {
  	bool Visible = false;
@@ -79,7 +82,7 @@ void FMovieSceneVisibilityTrackInstance::Update( float Position, float LastPosit
  		{
 			AActor* Actor = Cast<AActor>(Object);
 
-			if (Actor != NULL)
+			if (Actor != nullptr)
 			{
 				Actor->SetActorHiddenInGame(Visible);
 
@@ -99,8 +102,8 @@ void FMovieSceneVisibilityTrackInstance::Update( float Position, float LastPosit
 	}
 }
 
+
 void FMovieSceneVisibilityTrackInstance::RefreshInstance( const TArray<UObject*>& RuntimeObjects, class IMovieScenePlayer& Player )
 {
 	PropertyBindings->UpdateBindings( RuntimeObjects );
 }
-
diff --git a/Engine/Source/Runtime/MovieSceneTracks/Private/Tracks/MovieScene3DConstraintTrack.cpp b/Engine/Source/Runtime/MovieSceneTracks/Private/Tracks/MovieScene3DConstraintTrack.cpp
--- a/Engine/Source/Runtime/MovieSceneTracks/Private/Tracks/MovieScene3DConstraintTrack.cpp
+++ b/Engine/Source/Runtime/MovieSceneTracks/Private/Tracks/MovieScene3DConstraintTrack.cpp
@@ -6,46 +6,56 @@
 #include "IMovieScenePlayer.h"
 #include "MovieScene3DConstraintTrackInstance.h"
 
+
 UMovieScene3DConstraintTrack::UMovieScene3DConstraintTrack( const FObjectInitializer& ObjectInitializer )
 	: Super( ObjectInitializer )
-{
-}
+{ }
+
 
 const TArray<UMovieSceneSection*>& UMovieScene3DConstraintTrack::GetAllSections() const
 {
 	return ConstraintSections;
 }
 
+
 void UMovieScene3DConstraintTrack::RemoveAllAnimationData()
 {
+	// do nothing
 }
 
+
 bool UMovieScene3DConstraintTrack::HasSection( UMovieSceneSection* Section ) const
 {
 	return ConstraintSections.Find( Section ) != INDEX_NONE;
 }
 
+
 void UMovieScene3DConstraintTrack::AddSection( UMovieSceneSection* Section )
 {
 	ConstraintSections.Add( Section );
 }
 
+
 void UMovieScene3DConstraintTrack::RemoveSection( UMovieSceneSection* Section )
 {
 	ConstraintSections.Remove( Section );
 }
 
+
 bool UMovieScene3DConstraintTrack::IsEmpty() const
 {
 	return ConstraintSections.Num() == 0;
 }
 
+
 TRange<float> UMovieScene3DConstraintTrack::GetSectionBoundaries() const
 {
 	TArray< TRange<float> > Bounds;
+
 	for (int32 i = 0; i < ConstraintSections.Num(); ++i)
 	{
 		Bounds.Add(ConstraintSections[i]->GetRange());
 	}
+
 	return TRange<float>::Hull(Bounds);
 }
diff --git a/Engine/Source/Runtime/MovieSceneTracks/Private/Tracks/MovieScene3DPathTrack.cpp b/Engine/Source/Runtime/MovieSceneTracks/Private/Tracks/MovieScene3DPathTrack.cpp
--- a/Engine/Source/Runtime/MovieSceneTracks/Private/Tracks/MovieScene3DPathTrack.cpp
+++ b/Engine/Source/Runtime/MovieSceneTracks/Private/Tracks/MovieScene3DPathTrack.cpp
@@ -6,26 +6,31 @@
 #include "IMovieScenePlayer.h"
 #include "MovieScene3DPathTrackInstance.h"
 
+
 UMovieScene3DPathTrack::UMovieScene3DPathTrack( const FObjectInitializer& ObjectInitializer )
 	: Super( ObjectInitializer )
-{
-}
+{ }
+
+
 FName UMovieScene3DPathTrack::GetTrackName() const
 {
 	return FName("Path");
 }
 
+
 TSharedPtr<IMovieSceneTrackInstance> UMovieScene3DPathTrack::CreateInstance()
 {
 	return MakeShareable( new FMovieScene3DPathTrackInstance( *this ) ); 
 }
 
+
 void UMovieScene3DPathTrack::AddConstraint(float KeyTime, float ConstraintEndTime, const FGuid& ConstraintId)
 {
-	// add the section
 	UMovieScene3DPathSection* NewSection = NewObject<UMovieScene3DPathSection>(this);
-	NewSection->AddPath(KeyTime, ConstraintEndTime, ConstraintId);
-	NewSection->InitialPlacement( ConstraintSections, KeyTime, ConstraintEndTime, SupportsMultipleRows() );
+	{
+		NewSection->AddPath(KeyTime, ConstraintEndTime, ConstraintId);
+		NewSection->InitialPlacement( ConstraintSections, KeyTime, ConstraintEndTime, SupportsMultipleRows() );
+	}
 
 	ConstraintSections.Add(NewSection);
 }
diff --git a/Engine/Source/Runtime/MovieSceneTracks/Private/Tracks/MovieScene3DTransformTrack.cpp b/Engine/Source/Runtime/MovieSceneTracks/Private/Tracks/MovieScene3DTransformTrack.cpp
--- a/Engine/Source/Runtime/MovieSceneTracks/Private/Tracks/MovieScene3DTransformTrack.cpp
+++ b/Engine/Source/Runtime/MovieSceneTracks/Private/Tracks/MovieScene3DTransformTrack.cpp
@@ -6,24 +6,28 @@
 #include "IMovieScenePlayer.h"
 #include "MovieScene3DTransformTrackInstance.h"
 
+
 bool FTransformKey::ShouldKeyTranslation( EAxis::Type Axis) const
 {
 	// If the previous transform is not valid we have nothing to compare against so assume we always key. Otherwise check for differences
 	return !PreviousTransform.IsValid() || GetVectorComponentIfDifferent( Axis, NewTransform.Translation, PreviousTransform.Translation );
 }
 
+
 bool FTransformKey::ShouldKeyRotation( EAxis::Type Axis ) const
 {
 	// If the previous transform is not valid we have nothing to compare against so assume we always key. Otherwise check for differences
 	return !PreviousTransform.IsValid() || GetVectorComponentIfDifferent( Axis, NewTransform.Rotation.Euler(), PreviousTransform.Rotation.Euler() );
 }
 
+
 bool FTransformKey::ShouldKeyScale( EAxis::Type Axis ) const
 {
 	// If the previous transform is not valid we have nothing to compare against so assume we always key. Otherwise check for differences
 	return !PreviousTransform.IsValid() || GetVectorComponentIfDifferent( Axis, NewTransform.Scale, PreviousTransform.Scale );
 }
 
+
 bool FTransformKey::ShouldKeyAny() const
 {
 	return ShouldKeyTranslation(EAxis::X) || ShouldKeyTranslation(EAxis::Y) || ShouldKeyTranslation(EAxis::Z) ||
@@ -31,6 +35,7 @@ bool FTransformKey::ShouldKeyAny() const
 		ShouldKeyScale(EAxis::X) || ShouldKeyScale(EAxis::Y) || ShouldKeyScale(EAxis::Z);
 }
 
+
 bool FTransformKey::GetVectorComponentIfDifferent( EAxis::Type Axis, const FVector& Current, const FVector& Previous ) const
 {
 	bool bShouldAddKey = false;
@@ -57,21 +62,24 @@ bool FTransformKey::GetVectorComponentIfDifferent( EAxis::Type Axis, const FVect
 	return bShouldAddKey;
 }
 
+
 UMovieScene3DTransformTrack::UMovieScene3DTransformTrack( const FObjectInitializer& ObjectInitializer )
 	: Super( ObjectInitializer )
-{
-}
+{ }
+
 
 UMovieSceneSection* UMovieScene3DTransformTrack::CreateNewSection()
 {
 	return NewObject<UMovieSceneSection>(this, UMovieScene3DTransformSection::StaticClass(), NAME_None, RF_Transactional);
 }
 
+
 TSharedPtr<IMovieSceneTrackInstance> UMovieScene3DTransformTrack::CreateInstance()
 {
 	return MakeShareable( new FMovieScene3DTransformTrackInstance( *this ) );
 }
 
+
 bool UMovieScene3DTransformTrack::AddKeyToSection( const FGuid& ObjectHandle, const FTransformKey& InKey, const bool bUnwindRotation, F3DTransformTrackKey::Type KeyType )
 {
 	const UMovieSceneSection* NearestSection = MovieSceneHelpers::FindNearestSectionAtTime(Sections, InKey.GetKeyTime());
@@ -117,9 +125,10 @@ bool UMovieScene3DTransformTrack::Eval( float Position, float LastPosition, FVec
 		TransformSection->EvalScale( Position, OutScale );
 	}
 
-	return Section != NULL;
+	return Section != nullptr;
 }
 
+
 bool UMovieScene3DTransformTrack::CanKeyTrack(const FTransformKey& InKey ) const
 {
 	const UMovieSceneSection* NearestSection = MovieSceneHelpers::FindNearestSectionAtTime(Sections, InKey.GetKeyTime());
@@ -128,4 +137,4 @@ bool UMovieScene3DTransformTrack::CanKeyTrack(const FTransformKey& InKey ) const
 		return true;
 	}
 	return false;
-}
\ No newline at end of file
+}
diff --git a/Engine/Source/Runtime/MovieSceneTracks/Private/Tracks/MovieSceneAudioTrack.cpp b/Engine/Source/Runtime/MovieSceneTracks/Private/Tracks/MovieSceneAudioTrack.cpp
--- a/Engine/Source/Runtime/MovieSceneTracks/Private/Tracks/MovieSceneAudioTrack.cpp
+++ b/Engine/Source/Runtime/MovieSceneTracks/Private/Tracks/MovieSceneAudioTrack.cpp
@@ -9,12 +9,14 @@
 #include "Runtime/Engine/Public/AudioDecompress.h"
 #include "MovieSceneAudioTrackInstance.h"
 
+
 #define LOCTEXT_NAMESPACE "MovieSceneAudioTrack"
 
+
 UMovieSceneAudioTrack::UMovieSceneAudioTrack( const FObjectInitializer& ObjectInitializer )
 	: Super( ObjectInitializer )
-{
-}
+{ }
+
 
 FName UMovieSceneAudioTrack::GetTrackName() const
 {
@@ -27,35 +29,43 @@ TSharedPtr<IMovieSceneTrackInstance> UMovieSceneAudioTrack::CreateInstance()
 	return MakeShareable( new FMovieSceneAudioTrackInstance( *this ) ); 
 }
 
+
 const TArray<UMovieSceneSection*>& UMovieSceneAudioTrack::GetAllSections() const
 {
 	return AudioSections;
 }
 
+
 void UMovieSceneAudioTrack::RemoveAllAnimationData()
 {
+	// do nothing
 }
 
+
 bool UMovieSceneAudioTrack::HasSection( UMovieSceneSection* Section ) const
 {
 	return AudioSections.Find( Section ) != INDEX_NONE;
 }
 
+
 void UMovieSceneAudioTrack::AddSection( UMovieSceneSection* Section )
 {
 	AudioSections.Add( Section );
 }
 
+
 void UMovieSceneAudioTrack::RemoveSection( UMovieSceneSection* Section )
 {
 	AudioSections.Remove( Section );
 }
 
+
 bool UMovieSceneAudioTrack::IsEmpty() const
 {
 	return AudioSections.Num() == 0;
 }
 
+
 TRange<float> UMovieSceneAudioTrack::GetSectionBoundaries() const
 {
 	TArray< TRange<float> > Bounds;
@@ -66,6 +76,7 @@ TRange<float> UMovieSceneAudioTrack::GetSectionBoundaries() const
 	return TRange<float>::Hull(Bounds);
 }
 
+
 void UMovieSceneAudioTrack::AddNewSound(USoundBase* Sound, float Time)
 {
 	check(Sound);
@@ -89,11 +100,11 @@ void UMovieSceneAudioTrack::AddNewSound(USoundBase* Sound, float Time)
 	AudioSections.Add(NewSection);
 }
 
+
 bool UMovieSceneAudioTrack::IsAMasterTrack() const
 {
 	return Cast<UMovieScene>(GetOuter())->IsAMasterTrack(this);
 }
 
 
-
 #undef LOCTEXT_NAMESPACE
diff --git a/Engine/Source/Runtime/MovieSceneTracks/Private/Tracks/MovieSceneByteTrack.cpp b/Engine/Source/Runtime/MovieSceneTracks/Private/Tracks/MovieSceneByteTrack.cpp
--- a/Engine/Source/Runtime/MovieSceneTracks/Private/Tracks/MovieSceneByteTrack.cpp
+++ b/Engine/Source/Runtime/MovieSceneTracks/Private/Tracks/MovieSceneByteTrack.cpp
@@ -9,19 +9,21 @@
 
 UMovieSceneByteTrack::UMovieSceneByteTrack( const FObjectInitializer& ObjectInitializer )
 	: Super( ObjectInitializer )
-{
-}
+{ }
+
 
 UMovieSceneSection* UMovieSceneByteTrack::CreateNewSection()
 {
 	return NewObject<UMovieSceneSection>(this, UMovieSceneByteSection::StaticClass(), NAME_None, RF_Transactional);
 }
 
+
 TSharedPtr<IMovieSceneTrackInstance> UMovieSceneByteTrack::CreateInstance()
 {
 	return MakeShareable( new FMovieSceneByteTrackInstance( *this ) );
 }
 
+
 bool UMovieSceneByteTrack::AddKeyToSection( float Time, uint8 Value, FKeyParams KeyParams )
 {
 	const UMovieSceneSection* NearestSection = MovieSceneHelpers::FindNearestSectionAtTime( Sections, Time );
@@ -30,14 +32,15 @@ bool UMovieSceneByteTrack::AddKeyToSection( float Time, uint8 Value, FKeyParams
 		Modify();
 
 		UMovieSceneByteSection* NewSection = CastChecked<UMovieSceneByteSection>(FindOrAddSection(  Time ));
-	
 		NewSection->AddKey( Time, Value, KeyParams );
 
 		return true;
 	}
+
 	return false;
 }
 
+
 bool UMovieSceneByteTrack::Eval( float Position, float LastPostion, uint8& OutByte ) const
 {	
 	const UMovieSceneSection* Section = MovieSceneHelpers::FindNearestSectionAtTime( Sections, Position );
@@ -49,14 +52,16 @@ bool UMovieSceneByteTrack::Eval( float Position, float LastPostion, uint8& OutBy
 		OutByte = CastChecked<UMovieSceneByteSection>( Section )->Eval( Position );
 	}
 
-	return Section != NULL;
+	return Section != nullptr;
 }
 
+
 void UMovieSceneByteTrack::SetEnum(UEnum* InEnum)
 {
 	Enum = InEnum;
 }
 
+
 class UEnum* UMovieSceneByteTrack::GetEnum() const
 {
 	return Enum;
diff --git a/Engine/Source/Runtime/MovieSceneTracks/Private/Tracks/MovieSceneColorTrack.cpp b/Engine/Source/Runtime/MovieSceneTracks/Private/Tracks/MovieSceneColorTrack.cpp
--- a/Engine/Source/Runtime/MovieSceneTracks/Private/Tracks/MovieSceneColorTrack.cpp
+++ b/Engine/Source/Runtime/MovieSceneTracks/Private/Tracks/MovieSceneColorTrack.cpp
@@ -6,16 +6,18 @@
 #include "IMovieScenePlayer.h"
 #include "MovieSceneColorTrackInstance.h"
 
+
 UMovieSceneColorTrack::UMovieSceneColorTrack( const FObjectInitializer& ObjectInitializer )
 	: Super( ObjectInitializer )
-{
-}
+{ }
+
 
 UMovieSceneSection* UMovieSceneColorTrack::CreateNewSection()
 {
 	return NewObject<UMovieSceneSection>(this, UMovieSceneColorSection::StaticClass(), NAME_None, RF_Transactional);
 }
 
+
 TSharedPtr<IMovieSceneTrackInstance> UMovieSceneColorTrack::CreateInstance()
 {
 	return MakeShareable( new FMovieSceneColorTrackInstance( *this ) ); 
@@ -28,12 +30,13 @@ bool UMovieSceneColorTrack::AddKeyToSection( float Time, const FColorKey& Key, F
 	if (!NearestSection || KeyParams.bAddKeyEvenIfUnchanged || CastChecked<UMovieSceneColorSection>(NearestSection)->NewKeyIsNewData(Time, Key.Value, KeyParams))
 	{
 		Modify();
-		UMovieSceneColorSection* NewSection = CastChecked<UMovieSceneColorSection>( FindOrAddSection( Time ) );
 
+		UMovieSceneColorSection* NewSection = CastChecked<UMovieSceneColorSection>( FindOrAddSection( Time ) );
 		NewSection->AddKey( Time, Key, KeyParams );
 
 		return true;
 	}
+
 	return false;
 }
 
@@ -45,9 +48,8 @@ bool UMovieSceneColorTrack::Eval( float Position, float LastPosition, FLinearCol
 	if( Section )
 	{
 		Position = FMath::Clamp(Position, Section->GetStartTime(), Section->GetEndTime());
-
 		OutColor = CastChecked<UMovieSceneColorSection>( Section )->Eval( Position, OutColor );
 	}
 
-	return Section != NULL;
+	return Section != nullptr;
 }
\ No newline at end of file
diff --git a/Engine/Source/Runtime/MovieSceneTracks/Private/Tracks/MovieSceneFloatTrack.cpp b/Engine/Source/Runtime/MovieSceneTracks/Private/Tracks/MovieSceneFloatTrack.cpp
--- a/Engine/Source/Runtime/MovieSceneTracks/Private/Tracks/MovieSceneFloatTrack.cpp
+++ b/Engine/Source/Runtime/MovieSceneTracks/Private/Tracks/MovieSceneFloatTrack.cpp
@@ -9,14 +9,15 @@
 
 UMovieSceneFloatTrack::UMovieSceneFloatTrack( const FObjectInitializer& ObjectInitializer )
 	: Super( ObjectInitializer )
-{
-}
+{ }
+
 
 UMovieSceneSection* UMovieSceneFloatTrack::CreateNewSection()
 {
 	return NewObject<UMovieSceneSection>(this, UMovieSceneFloatSection::StaticClass(), NAME_None, RF_Transactional);
 }
 
+
 TSharedPtr<IMovieSceneTrackInstance> UMovieSceneFloatTrack::CreateInstance()
 {
 	return MakeShareable( new FMovieSceneFloatTrackInstance( *this ) ); 
@@ -39,6 +40,7 @@ bool UMovieSceneFloatTrack::AddKeyToSection( float Time, float Value, FKeyParams
 	return false;
 }
 
+
 bool UMovieSceneFloatTrack::Eval( float Position, float LastPosition, float& OutFloat ) const
 {
 	const UMovieSceneSection* Section = MovieSceneHelpers::FindNearestSectionAtTime( Sections, Position );
@@ -50,5 +52,5 @@ bool UMovieSceneFloatTrack::Eval( float Position, float LastPosition, float& Out
 		OutFloat = CastChecked<UMovieSceneFloatSection>( Section )->Eval( Position );
 	}
 
-	return Section != NULL;
+	return Section != nullptr;
 }
diff --git a/Engine/Source/Runtime/MovieSceneTracks/Private/Tracks/MovieSceneParticleTrack.cpp b/Engine/Source/Runtime/MovieSceneTracks/Private/Tracks/MovieSceneParticleTrack.cpp
--- a/Engine/Source/Runtime/MovieSceneTracks/Private/Tracks/MovieSceneParticleTrack.cpp
+++ b/Engine/Source/Runtime/MovieSceneTracks/Private/Tracks/MovieSceneParticleTrack.cpp
@@ -6,12 +6,14 @@
 #include "IMovieScenePlayer.h"
 #include "MovieSceneParticleTrackInstance.h"
 
+
 #define LOCTEXT_NAMESPACE "MovieSceneParticleTrack"
 
+
 UMovieSceneParticleTrack::UMovieSceneParticleTrack( const FObjectInitializer& ObjectInitializer )
 	: Super( ObjectInitializer )
-{
-}
+{ }
+
 
 FName UMovieSceneParticleTrack::GetTrackName() const
 {
@@ -24,35 +26,43 @@ TSharedPtr<IMovieSceneTrackInstance> UMovieSceneParticleTrack::CreateInstance()
 	return MakeShareable( new FMovieSceneParticleTrackInstance( *this ) ); 
 }
 
+
 const TArray<UMovieSceneSection*>& UMovieSceneParticleTrack::GetAllSections() const
 {
 	return ParticleSections;
 }
 
+
 void UMovieSceneParticleTrack::RemoveAllAnimationData()
 {
+	// do nothing
 }
 
+
 bool UMovieSceneParticleTrack::HasSection( UMovieSceneSection* Section ) const
 {
 	return ParticleSections.Find( Section ) != INDEX_NONE;
 }
 
+
 void UMovieSceneParticleTrack::AddSection( UMovieSceneSection* Section )
 {
 	ParticleSections.Add( Section );
 }
 
+
 void UMovieSceneParticleTrack::RemoveSection( UMovieSceneSection* Section )
 {
 	ParticleSections.Remove( Section );
 }
 
+
 bool UMovieSceneParticleTrack::IsEmpty() const
 {
 	return ParticleSections.Num() == 0;
 }
 
+
 TRange<float> UMovieSceneParticleTrack::GetSectionBoundaries() const
 {
 	TArray< TRange<float> > Bounds;
@@ -63,6 +73,7 @@ TRange<float> UMovieSceneParticleTrack::GetSectionBoundaries() const
 	return TRange<float>::Hull(Bounds);
 }
 
+
 void UMovieSceneParticleTrack::AddNewParticleSystem(float KeyTime, bool bTrigger)
 {
 	EParticleKey::Type KeyType = bTrigger ? EParticleKey::Trigger : EParticleKey::Toggle;
@@ -76,4 +87,5 @@ void UMovieSceneParticleTrack::AddNewParticleSystem(float KeyTime, bool bTrigger
 	ParticleSections.Add(NewSection);
 }
 
+
 #undef LOCTEXT_NAMESPACE
diff --git a/Engine/Source/Runtime/MovieSceneTracks/Private/Tracks/MovieSceneSkeletalAnimationTrack.cpp b/Engine/Source/Runtime/MovieSceneTracks/Private/Tracks/MovieSceneSkeletalAnimationTrack.cpp
--- a/Engine/Source/Runtime/MovieSceneTracks/Private/Tracks/MovieSceneSkeletalAnimationTrack.cpp
+++ b/Engine/Source/Runtime/MovieSceneTracks/Private/Tracks/MovieSceneSkeletalAnimationTrack.cpp
@@ -6,12 +6,14 @@
 #include "IMovieScenePlayer.h"
 #include "MovieSceneSkeletalAnimationTrackInstance.h"
 
+
 #define LOCTEXT_NAMESPACE "MovieSceneSkeletalAnimationTrack"
 
+
 UMovieSceneSkeletalAnimationTrack::UMovieSceneSkeletalAnimationTrack( const FObjectInitializer& ObjectInitializer )
 	: Super( ObjectInitializer )
-{
-}
+{ }
+
 
 FName UMovieSceneSkeletalAnimationTrack::GetTrackName() const
 {
@@ -24,30 +26,37 @@ TSharedPtr<IMovieSceneTrackInstance> UMovieSceneSkeletalAnimationTrack::CreateIn
 	return MakeShareable( new FMovieSceneSkeletalAnimationTrackInstance( *this ) ); 
 }
 
+
 const TArray<UMovieSceneSection*>& UMovieSceneSkeletalAnimationTrack::GetAllSections() const
 {
 	return AnimationSections;
 }
 
+
 void UMovieSceneSkeletalAnimationTrack::RemoveAllAnimationData()
 {
+	// do nothing
 }
 
+
 bool UMovieSceneSkeletalAnimationTrack::HasSection( UMovieSceneSection* Section ) const
 {
 	return AnimationSections.Find( Section ) != INDEX_NONE;
 }
 
+
 void UMovieSceneSkeletalAnimationTrack::RemoveSection( UMovieSceneSection* Section )
 {
 	AnimationSections.Remove( Section );
 }
 
+
 bool UMovieSceneSkeletalAnimationTrack::IsEmpty() const
 {
 	return AnimationSections.Num() == 0;
 }
 
+
 TRange<float> UMovieSceneSkeletalAnimationTrack::GetSectionBoundaries() const
 {
 	TArray< TRange<float> > Bounds;
@@ -58,6 +67,7 @@ TRange<float> UMovieSceneSkeletalAnimationTrack::GetSectionBoundaries() const
 	return TRange<float>::Hull(Bounds);
 }
 
+
 void UMovieSceneSkeletalAnimationTrack::AddNewAnimation(float KeyTime, UAnimSequence* AnimSequence)
 {
 	// add the section
@@ -69,6 +79,7 @@ void UMovieSceneSkeletalAnimationTrack::AddNewAnimation(float KeyTime, UAnimSequ
 	AnimationSections.Add(NewSection);
 }
 
+
 UMovieSceneSection* UMovieSceneSkeletalAnimationTrack::GetAnimSectionAtTime(float Time)
 {
 	for (int32 i = 0; i < AnimationSections.Num(); ++i)
@@ -80,7 +91,7 @@ UMovieSceneSection* UMovieSceneSkeletalAnimationTrack::GetAnimSectionAtTime(floa
 		}
 	}
 
-	return NULL;
+	return nullptr;
 }
 
 
diff --git a/Engine/Source/Runtime/MovieSceneTracks/Private/Tracks/MovieSceneVectorTrack.cpp b/Engine/Source/Runtime/MovieSceneTracks/Private/Tracks/MovieSceneVectorTrack.cpp
--- a/Engine/Source/Runtime/MovieSceneTracks/Private/Tracks/MovieSceneVectorTrack.cpp
+++ b/Engine/Source/Runtime/MovieSceneTracks/Private/Tracks/MovieSceneVectorTrack.cpp
@@ -6,6 +6,7 @@
 #include "IMovieScenePlayer.h"
 #include "MovieSceneVectorTrackInstance.h"
 
+
 FVectorKey::FVectorKey( const FVector2D& InValue, FName InCurveName )
 {
 	Value.X = InValue.X;
@@ -16,6 +17,7 @@ FVectorKey::FVectorKey( const FVector2D& InValue, FName InCurveName )
 	CurveName = InCurveName;
 }
 
+
 FVectorKey::FVectorKey( const FVector& InValue, FName InCurveName )
 {
 	Value.X = InValue.X;
@@ -26,6 +28,7 @@ FVectorKey::FVectorKey( const FVector& InValue, FName InCurveName )
 	CurveName = InCurveName;
 }
 
+
 FVectorKey::FVectorKey( const FVector4& InValue, FName InCurveName )
 {
 	Value = InValue;
@@ -33,22 +36,26 @@ FVectorKey::FVectorKey( const FVector4& InValue, FName InCurveName )
 	CurveName = InCurveName;
 }
 
+
 UMovieSceneVectorTrack::UMovieSceneVectorTrack( const FObjectInitializer& ObjectInitializer )
 	: Super( ObjectInitializer )
 {
 	NumChannelsUsed = 0;
 }
 
+
 UMovieSceneSection* UMovieSceneVectorTrack::CreateNewSection()
 {
 	return NewObject<UMovieSceneSection>(this, UMovieSceneVectorSection::StaticClass());
 }
 
+
 TSharedPtr<IMovieSceneTrackInstance> UMovieSceneVectorTrack::CreateInstance()
 {
 	return MakeShareable( new FMovieSceneVectorTrackInstance( *this ) );
 }
 
+
 bool UMovieSceneVectorTrack::AddKeyToSection( float Time, const FVector4& Value, int32 InChannelsUsed, FName CurveName, FKeyParams KeyParams )
 {
 	const UMovieSceneSection* NearestSection = MovieSceneHelpers::FindNearestSectionAtTime( Sections, Time );
@@ -67,14 +74,17 @@ bool UMovieSceneVectorTrack::AddKeyToSection( float Time, const FVector4& Value,
 
 		return true;
 	}
+
 	return false;
 }
 
+
 bool UMovieSceneVectorTrack::AddKeyToSection( float Time, const FVectorKey& Key, FKeyParams KeyParams )
 {
 	return AddKeyToSection(Time, Key.Value, Key.ChannelsUsed, Key.CurveName, KeyParams );
 }
 
+
 bool UMovieSceneVectorTrack::Eval( float Position, float LastPosition, FVector4& InOutVector ) const
 {
 	const UMovieSceneSection* Section = MovieSceneHelpers::FindNearestSectionAtTime( Sections, Position );
@@ -86,6 +96,5 @@ bool UMovieSceneVectorTrack::Eval( float Position, float LastPosition, FVector4&
 		InOutVector = CastChecked<UMovieSceneVectorSection>( Section )->Eval( Position, InOutVector );
 	}
 
-	return Section != NULL;
+	return Section != nullptr;
 }
-
diff --git a/Engine/Source/Runtime/MovieSceneTracks/Private/Tracks/SubMovieSceneTrack.cpp b/Engine/Source/Runtime/MovieSceneTracks/Private/Tracks/SubMovieSceneTrack.cpp
--- a/Engine/Source/Runtime/MovieSceneTracks/Private/Tracks/SubMovieSceneTrack.cpp
+++ b/Engine/Source/Runtime/MovieSceneTracks/Private/Tracks/SubMovieSceneTrack.cpp
@@ -7,16 +7,18 @@
 #include "SubMovieSceneSection.h"
 #include "SubMovieSceneTrackInstance.h"
 
+
 USubMovieSceneTrack::USubMovieSceneTrack( const FObjectInitializer& ObjectInitializer )
 	: Super( ObjectInitializer )
-{
-}
+{ }
+
 
 UMovieSceneSection* USubMovieSceneTrack::CreateNewSection()
 {
 	return NewObject<USubMovieSceneSection>(this, NAME_None, RF_Transactional);
 }
 
+
 FName USubMovieSceneTrack::GetTrackName() const
 {
 	static const FName DefaultName( "Sequence" );
@@ -30,16 +32,19 @@ FName USubMovieSceneTrack::GetTrackName() const
 	return Name;
 }
 
+
 TSharedPtr<IMovieSceneTrackInstance> USubMovieSceneTrack::CreateInstance()
 {
 	return MakeShareable( new FSubMovieSceneTrackInstance( *this ) ); 
 }
 
+
 const TArray<UMovieSceneSection*>& USubMovieSceneTrack::GetAllSections() const 
 {
 	return SubMovieSceneSections;
 }
 
+
 void USubMovieSceneTrack::AddSection( UMovieSceneSection* Section )
 {
 	USubMovieSceneSection* SubMovieSceneSection = Cast<USubMovieSceneSection>( Section );
@@ -49,6 +54,7 @@ void USubMovieSceneTrack::AddSection( UMovieSceneSection* Section )
 	}
 }
 
+
 void USubMovieSceneTrack::RemoveSection( UMovieSceneSection* Section )
 {
 	USubMovieSceneSection* SubMovieSceneSection = Cast<USubMovieSceneSection>( Section );
@@ -58,16 +64,19 @@ void USubMovieSceneTrack::RemoveSection( UMovieSceneSection* Section )
 	}
 }
 
+
 void USubMovieSceneTrack::RemoveAllAnimationData()
 {
 	SubMovieSceneSections.Empty();
 }
 
+
 bool USubMovieSceneTrack::IsEmpty() const
 {
 	return SubMovieSceneSections.Num() == 0;
 }
 
+
 TRange<float> USubMovieSceneTrack::GetSectionBoundaries() const
 {
 	TArray< TRange<float> > Bounds;
@@ -78,21 +87,23 @@ TRange<float> USubMovieSceneTrack::GetSectionBoundaries() const
 	return TRange<float>::Hull(Bounds);
 }
 
+
 bool USubMovieSceneTrack::HasSection( UMovieSceneSection* Section ) const
 {
 	return SubMovieSceneSections.Contains( Section );
 }
 
+
 void USubMovieSceneTrack::AddMovieSceneSection( UMovieSceneSequence* SubMovieSceneSequence, float Time )
 {
 	Modify();
 
 	USubMovieSceneSection* NewSection = CastChecked<USubMovieSceneSection>( CreateNewSection() );
-
-	NewSection->SetMovieSceneAnimation( SubMovieSceneSequence );
-	NewSection->SetStartTime( Time );
-	NewSection->SetEndTime( Time + SubMovieSceneSequence->GetMovieScene()->GetTimeRange().Size<float>() );
+	{
+		NewSection->SetMovieSceneAnimation( SubMovieSceneSequence );
+		NewSection->SetStartTime( Time );
+		NewSection->SetEndTime( Time + SubMovieSceneSequence->GetMovieScene()->GetTimeRange().Size<float>() );
+	}
 
 	SubMovieSceneSections.Add( NewSection );
-
 }
