MovieScene: Continued to implement event tracks

diff --git a/Engine/Source/Runtime/MovieSceneTracks/Private/Sections/MovieSceneEventSection.cpp b/Engine/Source/Runtime/MovieSceneTracks/Private/Sections/MovieSceneEventSection.cpp
--- a/Engine/Source/Runtime/MovieSceneTracks/Private/Sections/MovieSceneEventSection.cpp
+++ b/Engine/Source/Runtime/MovieSceneTracks/Private/Sections/MovieSceneEventSection.cpp
@@ -4,6 +4,15 @@
 #include "MovieSceneEventSection.h"
 
 
+/* UMovieSceneSection overrides
+ *****************************************************************************/
+
+void UMovieSceneEventSection::AddKey(float Time, const FName& EventName, FKeyParams KeyParams)
+{
+	Keys.HeapPush(FMovieSceneEventSectionKey(EventName, Time));
+}
+
+
 /* UMovieSceneSection overrides
  *****************************************************************************/
 
diff --git a/Engine/Source/Runtime/MovieSceneTracks/Private/Tracks/MovieSceneEventTrack.cpp b/Engine/Source/Runtime/MovieSceneTracks/Private/Tracks/MovieSceneEventTrack.cpp
--- a/Engine/Source/Runtime/MovieSceneTracks/Private/Tracks/MovieSceneEventTrack.cpp
+++ b/Engine/Source/Runtime/MovieSceneTracks/Private/Tracks/MovieSceneEventTrack.cpp
@@ -11,10 +11,27 @@
 
 bool UMovieSceneEventTrack::AddKeyToSection(float Time, FName EventName, FKeyParams KeyParams)
 {
-	// @todo sequencer: gmp: add key to event section
-	// should be inserted into sorted array
-	// can sections overlap?
-	return false;
+	UMovieSceneSection* TargetSection = MovieSceneHelpers::FindNearestSectionAtTime(Sections, Time);
+
+	if (TargetSection == nullptr)
+	{
+		TargetSection = CreateNewSection();
+		TargetSection->SetStartTime(Time);
+		TargetSection->SetEndTime(Time);
+
+		Sections.Add(TargetSection);
+	}
+
+	UMovieSceneEventSection* EventSection = Cast<UMovieSceneEventSection>(TargetSection);
+
+	if (EventSection == nullptr)
+	{
+		return false;
+	}
+
+	EventSection->AddKey(Time, EventName, KeyParams);
+
+	return true;
 }
 
 
diff --git a/Engine/Source/Runtime/MovieSceneTracks/Public/Sections/MovieSceneEventSection.h b/Engine/Source/Runtime/MovieSceneTracks/Public/Sections/MovieSceneEventSection.h
--- a/Engine/Source/Runtime/MovieSceneTracks/Public/Sections/MovieSceneEventSection.h
+++ b/Engine/Source/Runtime/MovieSceneTracks/Public/Sections/MovieSceneEventSection.h
@@ -11,7 +11,7 @@
 USTRUCT()
 struct FMovieSceneEventSectionKey
 {
-	GENERATED_BODY()
+	GENERATED_USTRUCT_BODY()
 
 	/** The name of the event to be triggered. */
 	UPROPERTY(EditAnywhere, Category=EventTrackKey)
@@ -20,6 +20,23 @@ struct FMovieSceneEventSectionKey
 	/** The time at which the event should be triggered. */
 	UPROPERTY()
 	float Time;
+
+	/** Default constructor. */
+	FMovieSceneEventSectionKey()
+		: Time(0.0f)
+	{ }
+
+	/** Creates and initializes a new instance. */
+	FMovieSceneEventSectionKey(const FName& InEventName, float InTime)
+		: EventName(InEventName)
+		, Time(InTime)
+	{ }
+
+	/** Operator less, used to sort the heap based on time until execution. */
+	bool operator<(const FMovieSceneEventSectionKey& Other) const
+	{
+		return Time < Other.Time;
+	}
 };
 
 
@@ -32,6 +49,17 @@ class UMovieSceneEventSection
 {
 	GENERATED_BODY()
 
+public:
+
+	/**
+	 * Add a key with the specified event.
+	 *
+	 * @param Time The location in time where the key should be added.
+	 * @param EventName The name of the event to fire at the specified time.
+	 * @param KeyParams The keying parameters.
+	 */
+	void AddKey(float Time, const FName& EventName, FKeyParams KeyParams);
+
 public:
 
 	// UMovieSceneSection interface
