Sequencer: Fix keyframing transform track when only setting a keyframe on one transform track.

diff --git a/Engine/Source/Runtime/MovieScene/Public/KeyParams.h b/Engine/Source/Runtime/MovieScene/Public/KeyParams.h
--- a/Engine/Source/Runtime/MovieScene/Public/KeyParams.h
+++ b/Engine/Source/Runtime/MovieScene/Public/KeyParams.h
@@ -15,6 +15,14 @@ struct MOVIESCENE_API FKeyParams
 		bAutoKeying = false;
 	}
 
+	FKeyParams(const FKeyParams& InKeyParams)
+	{
+		bCreateHandleIfMissing = InKeyParams.bCreateHandleIfMissing;
+		bCreateTrackIfMissing = InKeyParams.bCreateTrackIfMissing;
+		bAddKeyEvenIfUnchanged = InKeyParams.bAddKeyEvenIfUnchanged;
+		bAutoKeying = InKeyParams.bAutoKeying;
+	}
+
 	/** Create handle if it doesn't exist. */
 	bool bCreateHandleIfMissing;
 	/** Create track if it doesn't exist. */
diff --git a/Engine/Source/Runtime/MovieSceneTracks/Private/Tracks/MovieScene3DTransformTrack.cpp b/Engine/Source/Runtime/MovieSceneTracks/Private/Tracks/MovieScene3DTransformTrack.cpp
--- a/Engine/Source/Runtime/MovieSceneTracks/Private/Tracks/MovieScene3DTransformTrack.cpp
+++ b/Engine/Source/Runtime/MovieSceneTracks/Private/Tracks/MovieScene3DTransformTrack.cpp
@@ -89,19 +89,16 @@ bool UMovieScene3DTransformTrack::AddKeyToSection( const FGuid& ObjectHandle, co
 
 		UMovieScene3DTransformSection* NewSection = Cast<UMovieScene3DTransformSection>( FindOrAddSection( InKey.GetKeyTime() ) );
 
-		// key each component of the transform
-		if (KeyType & F3DTransformTrackKey::Key_Translation)
-		{
-			NewSection->AddTranslationKeys( InKey );
-		}
-		if (KeyType & F3DTransformTrackKey::Key_Rotation)
-		{
-			NewSection->AddRotationKeys( InKey, bUnwindRotation );
-		}
-		if (KeyType & F3DTransformTrackKey::Key_Scale)
-		{
-			NewSection->AddScaleKeys( InKey );
-		}
+		FTransformKey Key(InKey);
+
+		Key.KeyParams.bAddKeyEvenIfUnchanged = InKey.KeyParams.bAddKeyEvenIfUnchanged && (KeyType & F3DTransformTrackKey::Key_Translation);
+		NewSection->AddTranslationKeys( Key );
+
+		Key.KeyParams.bAddKeyEvenIfUnchanged = InKey.KeyParams.bAddKeyEvenIfUnchanged && (KeyType & F3DTransformTrackKey::Key_Rotation);
+		NewSection->AddRotationKeys( Key, bUnwindRotation );
+
+		Key.KeyParams.bAddKeyEvenIfUnchanged = InKey.KeyParams.bAddKeyEvenIfUnchanged && (KeyType & F3DTransformTrackKey::Key_Scale);
+		NewSection->AddScaleKeys( Key );
 
 		return true;
 	}
