Sequencer: Exposed relevant movie scene section properties

diff --git a/Engine/Source/Runtime/MovieScene/Public/MovieSceneSection.h b/Engine/Source/Runtime/MovieScene/Public/MovieSceneSection.h
--- a/Engine/Source/Runtime/MovieScene/Public/MovieSceneSection.h
+++ b/Engine/Source/Runtime/MovieScene/Public/MovieSceneSection.h
@@ -189,11 +189,11 @@ public:
 
 private:
 	/** The start time of the section */
-	UPROPERTY()
+	UPROPERTY(EditAnywhere, Category="Section")
 	float StartTime;
 
 	/** The end time of the section */
-	UPROPERTY()
+	UPROPERTY(EditAnywhere, Category="Section")
 	float EndTime;
 
 	/** The row index that this section sits on */
diff --git a/Engine/Source/Runtime/MovieSceneTracks/Public/Sections/MovieSceneAudioSection.h b/Engine/Source/Runtime/MovieSceneTracks/Public/Sections/MovieSceneAudioSection.h
--- a/Engine/Source/Runtime/MovieSceneTracks/Public/Sections/MovieSceneAudioSection.h
+++ b/Engine/Source/Runtime/MovieSceneTracks/Public/Sections/MovieSceneAudioSection.h
@@ -64,14 +64,14 @@ public:
 
 private:
 	/** The sound cue or wave that this section plays */
-	UPROPERTY()
+	UPROPERTY(EditAnywhere, Category="Audio")
 	USoundBase* Sound;
 
 	/** The absolute time that the sound starts playing at */
-	UPROPERTY()
+	UPROPERTY(EditAnywhere, Category="Audio")
 	float AudioStartTime;
 	
 	/** The amount which this audio is time dilated by */
-	UPROPERTY()
+	UPROPERTY(EditAnywhere, Category="Audio")
 	float AudioDilationFactor;
 };
diff --git a/Engine/Source/Runtime/MovieSceneTracks/Public/Sections/MovieSceneSkeletalAnimationSection.h b/Engine/Source/Runtime/MovieSceneTracks/Public/Sections/MovieSceneSkeletalAnimationSection.h
--- a/Engine/Source/Runtime/MovieSceneTracks/Public/Sections/MovieSceneSkeletalAnimationSection.h
+++ b/Engine/Source/Runtime/MovieSceneTracks/Public/Sections/MovieSceneSkeletalAnimationSection.h
@@ -41,14 +41,14 @@ public:
 
 private:
 	/** The animation sequence this section has */
-	UPROPERTY()
+	UPROPERTY(EditAnywhere, Category="Audio")
 	class UAnimSequence* AnimSequence;
 
 	/** The absolute time that the animation starts playing at */
-	UPROPERTY()
+	UPROPERTY(EditAnywhere, Category="Audio")
 	float AnimationStartTime;
 	
 	/** The amount which this animation is time dilated by */
-	UPROPERTY()
+	UPROPERTY(EditAnywhere, Category="Audio")
 	float AnimationDilationFactor;
 };
