Sequencer: Exposed more movie scene section properties.

diff --git a/Engine/Source/Runtime/MovieScene/Public/MovieSceneSection.h b/Engine/Source/Runtime/MovieScene/Public/MovieSceneSection.h
--- a/Engine/Source/Runtime/MovieScene/Public/MovieSceneSection.h
+++ b/Engine/Source/Runtime/MovieScene/Public/MovieSceneSection.h
@@ -201,10 +201,10 @@ private:
 	int32 RowIndex;
 
 	/** Toggle whether this section is active/inactive */
-	UPROPERTY()
+	UPROPERTY(EditAnywhere, Category="Section")
 	uint32 bIsActive : 1;
 
 	/** Toggle to set this section to be infinite */
-	UPROPERTY()
+	UPROPERTY(EditAnywhere, Category="Section")
 	uint32 bIsInfinite : 1;
 };
diff --git a/Engine/Source/Runtime/MovieSceneTracks/Public/Sections/MovieScene3DAttachSection.h b/Engine/Source/Runtime/MovieSceneTracks/Public/Sections/MovieScene3DAttachSection.h
--- a/Engine/Source/Runtime/MovieSceneTracks/Public/Sections/MovieScene3DAttachSection.h
+++ b/Engine/Source/Runtime/MovieSceneTracks/Public/Sections/MovieScene3DAttachSection.h
@@ -36,21 +36,21 @@ private:
 
 	//@todo need additional actor component/bone to attach to
 
-	UPROPERTY()
+	UPROPERTY(EditAnywhere, Category="Attach")
 	uint32 bConstrainTx:1;
 
-	UPROPERTY()
+	UPROPERTY(EditAnywhere, Category="Attach")
 	uint32 bConstrainTy:1;
 
-	UPROPERTY()
+	UPROPERTY(EditAnywhere, Category="Attach")
 	uint32 bConstrainTz:1;
 
-	UPROPERTY()
+	UPROPERTY(EditAnywhere, Category="Attach")
 	uint32 bConstrainRx:1;
 
-	UPROPERTY()
+	UPROPERTY(EditAnywhere, Category="Attach")
 	uint32 bConstrainRy:1;
 
-	UPROPERTY()
+	UPROPERTY(EditAnywhere, Category="Attach")
 	uint32 bConstrainRz:1;
 };
\ No newline at end of file
diff --git a/Engine/Source/Runtime/MovieSceneTracks/Public/Sections/MovieScene3DPathSection.h b/Engine/Source/Runtime/MovieSceneTracks/Public/Sections/MovieScene3DPathSection.h
--- a/Engine/Source/Runtime/MovieSceneTracks/Public/Sections/MovieScene3DPathSection.h
+++ b/Engine/Source/Runtime/MovieSceneTracks/Public/Sections/MovieScene3DPathSection.h
@@ -59,26 +59,26 @@ public:
 	
 private:
 	/** Timing Curve */
-	UPROPERTY()
+	UPROPERTY(EditAnywhere, Category="Path")
 	FRichCurve TimingCurve;
 
 	/** Front Axis */
-	UPROPERTY()
+	UPROPERTY(EditAnywhere, Category="Path")
 	MovieScene3DPathSection_Axis FrontAxisEnum;
 
 	/** Up Axis */
-	UPROPERTY()
+	UPROPERTY(EditAnywhere, Category="Path")
 	MovieScene3DPathSection_Axis UpAxisEnum;
 
 	/** Follow Curve */
-	UPROPERTY()
+	UPROPERTY(EditAnywhere, Category="Path")
 	uint32 bFollow:1;
 
 	/** Reverse Timing */
-	UPROPERTY()
+	UPROPERTY(EditAnywhere, Category="Path")
 	uint32 bReverse:1;
 
 	/** Force Upright */
-	UPROPERTY()
+	UPROPERTY(EditAnywhere, Category="Path")
 	uint32 bForceUpright:1;
 };
\ No newline at end of file
