Sequencer: Fix animation category.

diff --git a/Engine/Source/Runtime/MovieSceneTracks/Public/Sections/MovieSceneSkeletalAnimationSection.h b/Engine/Source/Runtime/MovieSceneTracks/Public/Sections/MovieSceneSkeletalAnimationSection.h
--- a/Engine/Source/Runtime/MovieSceneTracks/Public/Sections/MovieSceneSkeletalAnimationSection.h
+++ b/Engine/Source/Runtime/MovieSceneTracks/Public/Sections/MovieSceneSkeletalAnimationSection.h
@@ -41,14 +41,14 @@ public:
 
 private:
 	/** The animation sequence this section has */
-	UPROPERTY(EditAnywhere, Category="Audio")
+	UPROPERTY(EditAnywhere, Category="Animation")
 	class UAnimSequence* AnimSequence;
 
 	/** The absolute time that the animation starts playing at */
-	UPROPERTY(EditAnywhere, Category="Audio")
+	UPROPERTY(EditAnywhere, Category="Animation")
 	float AnimationStartTime;
 	
 	/** The amount which this animation is time dilated by */
-	UPROPERTY(EditAnywhere, Category="Audio")
+	UPROPERTY(EditAnywhere, Category="Animation")
 	float AnimationDilationFactor;
 };
