Sequencer: When autokeying is off, a key at the current time shouldn't be updated if the property changes. Fixed property tracks so that they don't cause runtime instances to be updated if no key is going to be added whent he property changes.

diff --git a/Engine/Source/Editor/MovieSceneTools/Private/TrackEditors/PropertyTrackEditors/BoolPropertyTrackEditor.cpp b/Engine/Source/Editor/MovieSceneTools/Private/TrackEditors/PropertyTrackEditors/BoolPropertyTrackEditor.cpp
--- a/Engine/Source/Editor/MovieSceneTools/Private/TrackEditors/PropertyTrackEditors/BoolPropertyTrackEditor.cpp
+++ b/Engine/Source/Editor/MovieSceneTools/Private/TrackEditors/PropertyTrackEditors/BoolPropertyTrackEditor.cpp
@@ -1,6 +1,7 @@
 // Copyright 1998-2015 Epic Games, Inc. All Rights Reserved.
 
 #include "MovieSceneToolsPrivatePCH.h"
+#include "MovieSceneBoolTrack.h"
 #include "BoolPropertyTrackEditor.h"
 #include "BoolPropertySection.h"
 
@@ -14,13 +15,22 @@ TSharedRef<ISequencerSection> FBoolPropertyTrackEditor::MakeSectionInterface( UM
 	return MakeShareable(new FBoolPropertySection( SectionObject, Track->GetTrackName(), GetSequencer().Get() ));
 }
 
-bool FBoolPropertyTrackEditor::TryGenerateKeyFromPropertyChanged( const FPropertyChangedParams& PropertyChangedParams, bool& OutKey )
+bool FBoolPropertyTrackEditor::TryGenerateKeyFromPropertyChanged( const UMovieSceneTrack* InTrack, const FPropertyChangedParams& PropertyChangedParams, bool& OutKey )
 {
 	const UBoolProperty* BoolProperty = Cast<const UBoolProperty>(PropertyChangedParams.PropertyPath.Last());
 	if (BoolProperty)
 	{
 		OutKey = BoolProperty->GetPropertyValue(BoolProperty->ContainerPtrToValuePtr<void>(PropertyChangedParams.ObjectsThatChanged.Last()));
-		return true;
+	
+		if (InTrack)
+		{
+			const UMovieSceneBoolTrack* BoolTrack = CastChecked<const UMovieSceneBoolTrack>( InTrack );
+			if (BoolTrack)
+			{
+				float KeyTime =	GetTimeForKey(GetMovieSceneSequence());
+				return BoolTrack->CanKeyTrack(KeyTime, OutKey, PropertyChangedParams.KeyParams);
+			}
+		}
 	}
 
 	return false;
diff --git a/Engine/Source/Editor/MovieSceneTools/Private/TrackEditors/PropertyTrackEditors/BoolPropertyTrackEditor.h b/Engine/Source/Editor/MovieSceneTools/Private/TrackEditors/PropertyTrackEditors/BoolPropertyTrackEditor.h
--- a/Engine/Source/Editor/MovieSceneTools/Private/TrackEditors/PropertyTrackEditors/BoolPropertyTrackEditor.h
+++ b/Engine/Source/Editor/MovieSceneTools/Private/TrackEditors/PropertyTrackEditors/BoolPropertyTrackEditor.h
@@ -32,7 +32,7 @@ public:
 
 protected:
 	/** FPropertyTrackEditor Interface */
-	virtual bool TryGenerateKeyFromPropertyChanged( const FPropertyChangedParams& PropertyChangedParams, bool& OutKey ) override;
+	virtual bool TryGenerateKeyFromPropertyChanged( const UMovieSceneTrack* InTrack, const FPropertyChangedParams& PropertyChangedParams, bool& OutKey ) override;
 };
 
 
diff --git a/Engine/Source/Editor/MovieSceneTools/Private/TrackEditors/PropertyTrackEditors/BytePropertyTrackEditor.cpp b/Engine/Source/Editor/MovieSceneTools/Private/TrackEditors/PropertyTrackEditors/BytePropertyTrackEditor.cpp
--- a/Engine/Source/Editor/MovieSceneTools/Private/TrackEditors/PropertyTrackEditors/BytePropertyTrackEditor.cpp
+++ b/Engine/Source/Editor/MovieSceneTools/Private/TrackEditors/PropertyTrackEditors/BytePropertyTrackEditor.cpp
@@ -61,8 +61,19 @@ UMovieSceneTrack* FBytePropertyTrackEditor::AddTrack(UMovieScene* FocusedMovieSc
 	return NewTrack;
 }
 
-bool FBytePropertyTrackEditor::TryGenerateKeyFromPropertyChanged( const FPropertyChangedParams& PropertyChangedParams, uint8& OutKey )
+bool FBytePropertyTrackEditor::TryGenerateKeyFromPropertyChanged( const UMovieSceneTrack* InTrack, const FPropertyChangedParams& PropertyChangedParams, uint8& OutKey )
 {
 	OutKey = *PropertyChangedParams.GetPropertyValue<uint8>();
-	return true;
+
+	if (InTrack)
+	{
+		const UMovieSceneByteTrack* ByteTrack = CastChecked<const UMovieSceneByteTrack>( InTrack );
+		if (ByteTrack)
+		{
+			float KeyTime =	GetTimeForKey(GetMovieSceneSequence());
+			return ByteTrack->CanKeyTrack(KeyTime, OutKey, PropertyChangedParams.KeyParams);
+		}
+	}
+
+	return false;
 }
\ No newline at end of file
diff --git a/Engine/Source/Editor/MovieSceneTools/Private/TrackEditors/PropertyTrackEditors/BytePropertyTrackEditor.h b/Engine/Source/Editor/MovieSceneTools/Private/TrackEditors/PropertyTrackEditors/BytePropertyTrackEditor.h
--- a/Engine/Source/Editor/MovieSceneTools/Private/TrackEditors/PropertyTrackEditors/BytePropertyTrackEditor.h
+++ b/Engine/Source/Editor/MovieSceneTools/Private/TrackEditors/PropertyTrackEditors/BytePropertyTrackEditor.h
@@ -33,7 +33,7 @@ public:
 
 protected:
 	/** FPropertyTrackEditor Interface */
-	virtual bool TryGenerateKeyFromPropertyChanged( const FPropertyChangedParams& PropertyChangedParams, uint8& OutKey ) override;
+	virtual bool TryGenerateKeyFromPropertyChanged( const UMovieSceneTrack* InTrack, const FPropertyChangedParams& PropertyChangedParams, uint8& OutKey ) override;
 };
 
 
diff --git a/Engine/Source/Editor/MovieSceneTools/Private/TrackEditors/PropertyTrackEditors/ColorPropertyTrackEditor.cpp b/Engine/Source/Editor/MovieSceneTools/Private/TrackEditors/PropertyTrackEditors/ColorPropertyTrackEditor.cpp
--- a/Engine/Source/Editor/MovieSceneTools/Private/TrackEditors/PropertyTrackEditors/ColorPropertyTrackEditor.cpp
+++ b/Engine/Source/Editor/MovieSceneTools/Private/TrackEditors/PropertyTrackEditors/ColorPropertyTrackEditor.cpp
@@ -1,6 +1,7 @@
 // Copyright 1998-2015 Epic Games, Inc. All Rights Reserved.
 
 #include "MovieSceneToolsPrivatePCH.h"
+#include "MovieSceneColorTrack.h"
 #include "ColorPropertyTrackEditor.h"
 #include "ColorPropertySection.h"
 
@@ -14,7 +15,7 @@ TSharedRef<ISequencerSection> FColorPropertyTrackEditor::MakeSectionInterface( U
 	return MakeShareable(new FColorPropertySection( SectionObject, Track->GetTrackName(), GetSequencer().Get(), Track ));
 }
 
-bool FColorPropertyTrackEditor::TryGenerateKeyFromPropertyChanged( const FPropertyChangedParams& PropertyChangedParams, FColorKey& OutKey )
+bool FColorPropertyTrackEditor::TryGenerateKeyFromPropertyChanged(  const UMovieSceneTrack* InTrack, const FPropertyChangedParams& PropertyChangedParams, FColorKey& OutKey )
 {
 	bool bIsFColor = false, bIsFLinearColor = false, bIsSlateColor = false;
 	const UStructProperty* StructProp = Cast<const UStructProperty>(PropertyChangedParams.PropertyPath.Last());
@@ -53,5 +54,15 @@ bool FColorPropertyTrackEditor::TryGenerateKeyFromPropertyChanged( const FProper
 	OutKey.CurveName = PropertyChangedParams.StructPropertyNameToKey;
 	OutKey.bIsSlateColor = bIsSlateColor;
 
-	return true;
+	if (InTrack)
+	{
+		const UMovieSceneColorTrack* ColorTrack = CastChecked<const UMovieSceneColorTrack>( InTrack );
+		if (ColorTrack)
+		{
+			float KeyTime =	GetTimeForKey(GetMovieSceneSequence());
+			return ColorTrack->CanKeyTrack(KeyTime, OutKey, PropertyChangedParams.KeyParams);
+		}
+	}
+
+	return false;
 }
diff --git a/Engine/Source/Editor/MovieSceneTools/Private/TrackEditors/PropertyTrackEditors/ColorPropertyTrackEditor.h b/Engine/Source/Editor/MovieSceneTools/Private/TrackEditors/PropertyTrackEditors/ColorPropertyTrackEditor.h
--- a/Engine/Source/Editor/MovieSceneTools/Private/TrackEditors/PropertyTrackEditors/ColorPropertyTrackEditor.h
+++ b/Engine/Source/Editor/MovieSceneTools/Private/TrackEditors/PropertyTrackEditors/ColorPropertyTrackEditor.h
@@ -32,7 +32,7 @@ public:
 
 protected:
 	/** FPropertyTrackEditor Interface */
-	virtual bool TryGenerateKeyFromPropertyChanged( const FPropertyChangedParams& PropertyChangedParams, FColorKey& OutKey ) override;
+	virtual bool TryGenerateKeyFromPropertyChanged( const UMovieSceneTrack* InTrack, const FPropertyChangedParams& PropertyChangedParams, FColorKey& OutKey ) override;
 };
 
 
diff --git a/Engine/Source/Editor/MovieSceneTools/Private/TrackEditors/PropertyTrackEditors/FloatPropertyTrackEditor.cpp b/Engine/Source/Editor/MovieSceneTools/Private/TrackEditors/PropertyTrackEditors/FloatPropertyTrackEditor.cpp
--- a/Engine/Source/Editor/MovieSceneTools/Private/TrackEditors/PropertyTrackEditors/FloatPropertyTrackEditor.cpp
+++ b/Engine/Source/Editor/MovieSceneTools/Private/TrackEditors/PropertyTrackEditors/FloatPropertyTrackEditor.cpp
@@ -14,8 +14,19 @@ TSharedRef<ISequencerSection> FFloatPropertyTrackEditor::MakeSectionInterface( U
 	return MakeShareable(new FFloatPropertySection( SectionObject, Track->GetTrackName() ));
 }
 
-bool FFloatPropertyTrackEditor::TryGenerateKeyFromPropertyChanged( const FPropertyChangedParams& PropertyChangedParams, float& OutKey )
+bool FFloatPropertyTrackEditor::TryGenerateKeyFromPropertyChanged( const UMovieSceneTrack* InTrack, const FPropertyChangedParams& PropertyChangedParams, float& OutKey )
 {
 	OutKey = *PropertyChangedParams.GetPropertyValue<float>();
-	return true;
+
+	if (InTrack)
+	{
+		const UMovieSceneFloatTrack* FloatTrack = CastChecked<const UMovieSceneFloatTrack>( InTrack );
+		if (FloatTrack)
+		{
+			float KeyTime =	GetTimeForKey(GetMovieSceneSequence());
+			return FloatTrack->CanKeyTrack(KeyTime, OutKey, PropertyChangedParams.KeyParams);
+		}
+	}
+
+	return false;
 }
\ No newline at end of file
diff --git a/Engine/Source/Editor/MovieSceneTools/Private/TrackEditors/PropertyTrackEditors/FloatPropertyTrackEditor.h b/Engine/Source/Editor/MovieSceneTools/Private/TrackEditors/PropertyTrackEditors/FloatPropertyTrackEditor.h
--- a/Engine/Source/Editor/MovieSceneTools/Private/TrackEditors/PropertyTrackEditors/FloatPropertyTrackEditor.h
+++ b/Engine/Source/Editor/MovieSceneTools/Private/TrackEditors/PropertyTrackEditors/FloatPropertyTrackEditor.h
@@ -32,7 +32,7 @@ public:
 
 protected:
 	/** FPropertyTrackEditor Interface */
-	virtual bool TryGenerateKeyFromPropertyChanged( const FPropertyChangedParams& PropertyChangedParams, float& OutKey ) override;
+	virtual bool TryGenerateKeyFromPropertyChanged( const UMovieSceneTrack* InTrack, const FPropertyChangedParams& PropertyChangedParams, float& OutKey ) override;
 };
 
 
diff --git a/Engine/Source/Editor/MovieSceneTools/Private/TrackEditors/PropertyTrackEditors/VectorPropertyTrackEditor.cpp b/Engine/Source/Editor/MovieSceneTools/Private/TrackEditors/PropertyTrackEditors/VectorPropertyTrackEditor.cpp
--- a/Engine/Source/Editor/MovieSceneTools/Private/TrackEditors/PropertyTrackEditors/VectorPropertyTrackEditor.cpp
+++ b/Engine/Source/Editor/MovieSceneTools/Private/TrackEditors/PropertyTrackEditors/VectorPropertyTrackEditor.cpp
@@ -1,6 +1,7 @@
 // Copyright 1998-2015 Epic Games, Inc. All Rights Reserved.
 
 #include "MovieSceneToolsPrivatePCH.h"
+#include "MovieSceneVectorTrack.h"
 #include "VectorPropertyTrackEditor.h"
 #include "VectorPropertySection.h"
 
@@ -14,7 +15,7 @@ TSharedRef<ISequencerSection> FVectorPropertyTrackEditor::MakeSectionInterface(
 	return MakeShareable(new FVectorPropertySection( SectionObject, Track->GetTrackName() ));
 }
 
-bool FVectorPropertyTrackEditor::TryGenerateKeyFromPropertyChanged( const FPropertyChangedParams& PropertyChangedParams, FVectorKey& OutKey )
+bool FVectorPropertyTrackEditor::TryGenerateKeyFromPropertyChanged( const UMovieSceneTrack* InTrack, const FPropertyChangedParams& PropertyChangedParams, FVectorKey& OutKey )
 {
 	bool bIsVector2D = false, bIsVector = false, bIsVector4 = false;
 	const UStructProperty* StructProp = Cast<const UStructProperty>( PropertyChangedParams.PropertyPath.Last() );
@@ -44,5 +45,16 @@ bool FVectorPropertyTrackEditor::TryGenerateKeyFromPropertyChanged( const FPrope
 	{
 		OutKey = FVectorKey( *PropertyChangedParams.GetPropertyValue<FVector4>(), PropertyChangedParams.StructPropertyNameToKey );
 	}
-	return true;
+
+	if (InTrack)
+	{
+		const UMovieSceneVectorTrack* VectorTrack = CastChecked<const UMovieSceneVectorTrack>( InTrack );
+		if (VectorTrack)
+		{
+			float KeyTime =	GetTimeForKey(GetMovieSceneSequence());
+			return VectorTrack->CanKeyTrack(KeyTime, OutKey, PropertyChangedParams.KeyParams);
+		}
+	}
+
+	return false;
 }
diff --git a/Engine/Source/Editor/MovieSceneTools/Private/TrackEditors/PropertyTrackEditors/VectorPropertyTrackEditor.h b/Engine/Source/Editor/MovieSceneTools/Private/TrackEditors/PropertyTrackEditors/VectorPropertyTrackEditor.h
--- a/Engine/Source/Editor/MovieSceneTools/Private/TrackEditors/PropertyTrackEditors/VectorPropertyTrackEditor.h
+++ b/Engine/Source/Editor/MovieSceneTools/Private/TrackEditors/PropertyTrackEditors/VectorPropertyTrackEditor.h
@@ -32,7 +32,7 @@ public:
 
 protected:
 	/** FPropertyTrackEditor Interface */
-	virtual bool TryGenerateKeyFromPropertyChanged( const FPropertyChangedParams& PropertyChangedParams, FVectorKey& OutKey ) override;
+	virtual bool TryGenerateKeyFromPropertyChanged( const UMovieSceneTrack* InTrack, const FPropertyChangedParams& PropertyChangedParams, FVectorKey& OutKey ) override;
 };
 
 
diff --git a/Engine/Source/Editor/MovieSceneTools/Private/TrackEditors/PropertyTrackEditors/VisibilityPropertyTrackEditor.cpp b/Engine/Source/Editor/MovieSceneTools/Private/TrackEditors/PropertyTrackEditors/VisibilityPropertyTrackEditor.cpp
--- a/Engine/Source/Editor/MovieSceneTools/Private/TrackEditors/PropertyTrackEditors/VisibilityPropertyTrackEditor.cpp
+++ b/Engine/Source/Editor/MovieSceneTools/Private/TrackEditors/PropertyTrackEditors/VisibilityPropertyTrackEditor.cpp
@@ -1,6 +1,7 @@
 // Copyright 1998-2015 Epic Games, Inc. All Rights Reserved.
 
 #include "MovieSceneToolsPrivatePCH.h"
+#include "MovieSceneVisibilityTrack.h"
 #include "VisibilityPropertyTrackEditor.h"
 #include "VisibilityPropertySection.h"
 
@@ -14,13 +15,22 @@ TSharedRef<ISequencerSection> FVisibilityPropertyTrackEditor::MakeSectionInterfa
 	return MakeShareable(new FVisibilityPropertySection( SectionObject, Track->GetTrackName(), GetSequencer().Get() ));
 }
 
-bool FVisibilityPropertyTrackEditor::TryGenerateKeyFromPropertyChanged( const FPropertyChangedParams& PropertyChangedParams, bool& OutKey )
+bool FVisibilityPropertyTrackEditor::TryGenerateKeyFromPropertyChanged( const UMovieSceneTrack* InTrack, const FPropertyChangedParams& PropertyChangedParams, bool& OutKey )
 {
 	const UBoolProperty* BoolProperty = Cast<const UBoolProperty>(PropertyChangedParams.PropertyPath.Last());
 	if (BoolProperty)
 	{
 		OutKey = BoolProperty->GetPropertyValue(BoolProperty->ContainerPtrToValuePtr<void>(PropertyChangedParams.ObjectsThatChanged.Last()));
-		return true;
+
+		if (InTrack)
+		{
+			const UMovieSceneVisibilityTrack* VisibilityTrack = CastChecked<const UMovieSceneVisibilityTrack>( InTrack );
+			if (VisibilityTrack)
+			{
+				float KeyTime =	GetTimeForKey(GetMovieSceneSequence());
+				return VisibilityTrack->CanKeyTrack(KeyTime, OutKey, PropertyChangedParams.KeyParams);
+			}
+		}
 	}
 
 	return false;
diff --git a/Engine/Source/Editor/MovieSceneTools/Private/TrackEditors/PropertyTrackEditors/VisibilityPropertyTrackEditor.h b/Engine/Source/Editor/MovieSceneTools/Private/TrackEditors/PropertyTrackEditors/VisibilityPropertyTrackEditor.h
--- a/Engine/Source/Editor/MovieSceneTools/Private/TrackEditors/PropertyTrackEditors/VisibilityPropertyTrackEditor.h
+++ b/Engine/Source/Editor/MovieSceneTools/Private/TrackEditors/PropertyTrackEditors/VisibilityPropertyTrackEditor.h
@@ -28,7 +28,7 @@ public:
 
 protected:
 	/** FPropertyTrackEditor Interface */
-	virtual bool TryGenerateKeyFromPropertyChanged( const FPropertyChangedParams& PropertyChangedParams, bool& OutKey ) override;
+	virtual bool TryGenerateKeyFromPropertyChanged( const UMovieSceneTrack* InTrack, const FPropertyChangedParams& PropertyChangedParams, bool& OutKey ) override;
 };
 
 
diff --git a/Engine/Source/Editor/MovieSceneTools/Public/PropertyTrackEditor.h b/Engine/Source/Editor/MovieSceneTools/Public/PropertyTrackEditor.h
--- a/Engine/Source/Editor/MovieSceneTools/Public/PropertyTrackEditor.h
+++ b/Engine/Source/Editor/MovieSceneTools/Public/PropertyTrackEditor.h
@@ -117,12 +117,13 @@ protected:
 	 * Tries to generate a key for the track based on a property change event.  This will be called when a property changes
 	 * which matches a property type name which was supplied on construction.
 	 *
+	 * @param InTrack The track where this key could be added
 	 * @param PropertyChangedParams Parameters associated with the property change.
 	 * @param OutKey If this call is successful this represents the key which was generated.
 	 *
 	 * @return true if generating a key was successful, otherwise false.
 	 */
-	virtual bool TryGenerateKeyFromPropertyChanged( const FPropertyChangedParams& PropertyChangedParams, KeyType& OutKey ) = 0;
+	virtual bool TryGenerateKeyFromPropertyChanged( const UMovieSceneTrack* InTrack, const FPropertyChangedParams& PropertyChangedParams, KeyType& OutKey ) = 0;
 	
 private:
 	/** Adds a callback for property changes for the supplied property type name. */
@@ -173,17 +174,17 @@ private:
 				TSubclassOf<UMovieSceneTrack> SequencerTrackClass = TrackType::StaticClass();
 				GetPropertyAndTrackClass(PropertyChangedParams.PropertyPath.Last(), PropertyName, SequencerTrackClass);
 
-				bool bCreateTrackIfMissing = PropertyChangedParams.KeyParams.bCreateTrackIfMissing;
+				const bool bCreateTrackIfMissing = PropertyChangedParams.KeyParams.bCreateTrackIfMissing;
 
-				if (bCreateTrackIfMissing || GetSequencer()->GetFocusedMovieSceneSequence()->GetMovieScene()->FindTrack(SequencerTrackClass, ObjectHandle, PropertyName))
+				UMovieSceneTrack* Track = GetSequencer()->GetFocusedMovieSceneSequence()->GetMovieScene()->FindTrack(SequencerTrackClass, ObjectHandle, PropertyName);
+
+				KeyType Key;
+				const bool bCanCreateKey = TryGenerateKeyFromPropertyChanged( Track, PropertyChangedParams, Key );
+
+				if (bCreateTrackIfMissing || bCanCreateKey )
 				{
-					// Get the value from the property
-					KeyType Key;
-					if ( TryGenerateKeyFromPropertyChanged( PropertyChangedParams, Key ) )
-					{
-						AnimatablePropertyChanged( TrackType::StaticClass(),
-							FOnKeyProperty::CreateRaw( this, &FPropertyTrackEditor::OnKeyProperty, PropertyChangedParams, Key ) );
-					}
+					AnimatablePropertyChanged( TrackType::StaticClass(),
+						FOnKeyProperty::CreateRaw( this, &FPropertyTrackEditor::OnKeyProperty, PropertyChangedParams, Key ) );
 				}
 			}
 		}
@@ -203,7 +204,7 @@ private:
 				TSubclassOf<UMovieSceneTrack> SequencerTrackClass = TrackType::StaticClass();
 				GetPropertyAndTrackClass(PropertyChangedParams.PropertyPath.Last(), PropertyName, SequencerTrackClass);
 
-				bool bCreateTrackIfMissing = PropertyChangedParams.KeyParams.bCreateTrackIfMissing;
+				const bool bCreateTrackIfMissing = PropertyChangedParams.KeyParams.bCreateTrackIfMissing;
 
 				UMovieSceneTrack* Track = GetTrackForObject( ObjectHandle, SequencerTrackClass, PropertyName, bCreateTrackIfMissing );
 				if ( Track )
diff --git a/Engine/Source/Editor/UMGEditor/Private/Animation/MarginTrackEditor.cpp b/Engine/Source/Editor/UMGEditor/Private/Animation/MarginTrackEditor.cpp
--- a/Engine/Source/Editor/UMGEditor/Private/Animation/MarginTrackEditor.cpp
+++ b/Engine/Source/Editor/UMGEditor/Private/Animation/MarginTrackEditor.cpp
@@ -1,6 +1,7 @@
 // Copyright 1998-2015 Epic Games, Inc. All Rights Reserved.
 
 #include "UMGEditorPrivatePCH.h"
+#include "MovieSceneMarginTrack.h"
 #include "MarginTrackEditor.h"
 #include "MovieSceneMarginSection.h"
 #include "PropertySection.h"
@@ -40,9 +41,20 @@ TSharedRef<ISequencerSection> FMarginTrackEditor::MakeSectionInterface( UMovieSc
 	return NewSection;
 }
 
-bool FMarginTrackEditor::TryGenerateKeyFromPropertyChanged( const FPropertyChangedParams& PropertyChangedParams, FMarginKey& OutKey )
+bool FMarginTrackEditor::TryGenerateKeyFromPropertyChanged( const UMovieSceneTrack* InTrack, const FPropertyChangedParams& PropertyChangedParams, FMarginKey& OutKey )
 {
 	OutKey.CurveName = PropertyChangedParams.StructPropertyNameToKey;
 	OutKey.Value = *PropertyChangedParams.GetPropertyValue<FMargin>();
-	return true;
+
+	if (InTrack)
+	{
+		const UMovieSceneMarginTrack* MarginTrack = CastChecked<const UMovieSceneMarginTrack>( InTrack );
+		if (MarginTrack)
+		{
+			float KeyTime =	GetTimeForKey(GetMovieSceneSequence());
+			return MarginTrack->CanKeyTrack(KeyTime, OutKey, PropertyChangedParams.KeyParams);
+		}
+	}
+
+	return false;
 }
\ No newline at end of file
diff --git a/Engine/Source/Editor/UMGEditor/Private/Animation/MarginTrackEditor.h b/Engine/Source/Editor/UMGEditor/Private/Animation/MarginTrackEditor.h
--- a/Engine/Source/Editor/UMGEditor/Private/Animation/MarginTrackEditor.h
+++ b/Engine/Source/Editor/UMGEditor/Private/Animation/MarginTrackEditor.h
@@ -32,7 +32,7 @@ public:
 
 protected:
 	/** FPropertyTrackEditor Interface */
-	virtual bool TryGenerateKeyFromPropertyChanged( const FPropertyChangedParams& PropertyChangedParams, FMarginKey& OutKey ) override;
+	virtual bool TryGenerateKeyFromPropertyChanged( const UMovieSceneTrack* InTrack, const FPropertyChangedParams& PropertyChangedParams, FMarginKey& OutKey ) override;
 };
 
 
diff --git a/Engine/Source/Editor/UMGEditor/Private/Animation/Sequencer2DTransformTrackEditor.cpp b/Engine/Source/Editor/UMGEditor/Private/Animation/Sequencer2DTransformTrackEditor.cpp
--- a/Engine/Source/Editor/UMGEditor/Private/Animation/Sequencer2DTransformTrackEditor.cpp
+++ b/Engine/Source/Editor/UMGEditor/Private/Animation/Sequencer2DTransformTrackEditor.cpp
@@ -3,6 +3,7 @@
 #include "UMGEditorPrivatePCH.h"
 #include "Sequencer2DTransformTrackEditor.h"
 #include "MovieScene2DTransformSection.h"
+#include "MovieScene2DTransformTrack.h"
 #include "PropertySection.h"
 #include "ISectionLayoutBuilder.h"
 #include "MovieSceneToolHelpers.h"
@@ -55,9 +56,20 @@ TSharedRef<ISequencerSection> F2DTransformTrackEditor::MakeSectionInterface( UMo
 	return NewSection;
 }
 
-bool F2DTransformTrackEditor::TryGenerateKeyFromPropertyChanged( const FPropertyChangedParams& PropertyChangedParams, F2DTransformKey& OutKey )
+bool F2DTransformTrackEditor::TryGenerateKeyFromPropertyChanged( const UMovieSceneTrack* InTrack, const FPropertyChangedParams& PropertyChangedParams, F2DTransformKey& OutKey )
 {
 	OutKey.CurveName = PropertyChangedParams.StructPropertyNameToKey;
 	OutKey.Value = *PropertyChangedParams.GetPropertyValue<FWidgetTransform>();
-	return true;
+
+	if (InTrack)
+	{
+		const UMovieScene2DTransformTrack* TransformTrack = CastChecked<const UMovieScene2DTransformTrack>( InTrack );
+		if (TransformTrack)
+		{
+			float KeyTime =	GetTimeForKey(GetMovieSceneSequence());
+			return TransformTrack->CanKeyTrack(KeyTime, OutKey, PropertyChangedParams.KeyParams);
+		}
+	}
+
+	return false;
 }
\ No newline at end of file
diff --git a/Engine/Source/Editor/UMGEditor/Private/Animation/Sequencer2DTransformTrackEditor.h b/Engine/Source/Editor/UMGEditor/Private/Animation/Sequencer2DTransformTrackEditor.h
--- a/Engine/Source/Editor/UMGEditor/Private/Animation/Sequencer2DTransformTrackEditor.h
+++ b/Engine/Source/Editor/UMGEditor/Private/Animation/Sequencer2DTransformTrackEditor.h
@@ -32,7 +32,7 @@ public:
 
 protected:
 	/** FPropertyTrackEditor Interface */
-	virtual bool TryGenerateKeyFromPropertyChanged( const FPropertyChangedParams& PropertyChangedParams, F2DTransformKey& OutKey ) override;
+	virtual bool TryGenerateKeyFromPropertyChanged( const UMovieSceneTrack* InTrack, const FPropertyChangedParams& PropertyChangedParams, F2DTransformKey& OutKey ) override;
 };
 
 
diff --git a/Engine/Source/Runtime/MovieSceneTracks/Private/Sections/MovieScene3DTransformSection.cpp b/Engine/Source/Runtime/MovieSceneTracks/Private/Sections/MovieScene3DTransformSection.cpp
--- a/Engine/Source/Runtime/MovieSceneTracks/Private/Sections/MovieScene3DTransformSection.cpp
+++ b/Engine/Source/Runtime/MovieSceneTracks/Private/Sections/MovieScene3DTransformSection.cpp
@@ -141,22 +141,17 @@ void UMovieScene3DTransformSection::AddTranslationKeys( const FTransformKey& Tra
 {
 	const float Time = TransformKey.GetKeyTime();
 
-	bool bTxKeyExists = Translation[0].IsKeyHandleValid(Translation[0].FindKey(Time));
-	bool bTyKeyExists = Translation[1].IsKeyHandleValid(Translation[1].FindKey(Time));
-	bool bTzKeyExists = Translation[2].IsKeyHandleValid(Translation[2].FindKey(Time));
-
-	// Key each component. Don't add a keyframe if there are existing keys and auto key is not enabled.
-	if( TransformKey.KeyParams.bAddKeyEvenIfUnchanged || (TransformKey.ShouldKeyTranslation( EAxis::X ) && !(!bTxKeyExists && !TransformKey.KeyParams.bAutoKeying && Translation[0].GetNumKeys() > 0)) )
+	if( TransformKey.KeyParams.bAddKeyEvenIfUnchanged || TransformKey.ShouldKeyTranslation( EAxis::X ) || Translation[0].GetNumKeys() == 0 )
 	{
 		AddKeyToCurve( Translation[0], Time, TransformKey.GetTranslationValue().X, TransformKey.KeyParams );
 	}
 
-	if( TransformKey.KeyParams.bAddKeyEvenIfUnchanged || (TransformKey.ShouldKeyTranslation( EAxis::Y ) && !(!bTyKeyExists && !TransformKey.KeyParams.bAutoKeying && Translation[1].GetNumKeys() > 0)) )
+	if( TransformKey.KeyParams.bAddKeyEvenIfUnchanged || TransformKey.ShouldKeyTranslation( EAxis::Y ) || Translation[1].GetNumKeys() == 0 )
 	{ 
 		AddKeyToCurve( Translation[1], Time, TransformKey.GetTranslationValue().Y, TransformKey.KeyParams );
 	}
 
-	if( TransformKey.KeyParams.bAddKeyEvenIfUnchanged || (TransformKey.ShouldKeyTranslation( EAxis::Z ) && !(!bTzKeyExists && !TransformKey.KeyParams.bAutoKeying && Translation[2].GetNumKeys() > 0)) )
+	if( TransformKey.KeyParams.bAddKeyEvenIfUnchanged || TransformKey.ShouldKeyTranslation( EAxis::Z ) || Translation[2].GetNumKeys() == 0 )
 	{ 
 		AddKeyToCurve( Translation[2], Time, TransformKey.GetTranslationValue().Z, TransformKey.KeyParams );
 	}
@@ -167,22 +162,17 @@ void UMovieScene3DTransformSection::AddRotationKeys( const FTransformKey& Transf
 {
 	const float Time = TransformKey.GetKeyTime();
 
-	bool bRxKeyExists = Rotation[0].IsKeyHandleValid(Rotation[0].FindKey(Time));
-	bool bRyKeyExists = Rotation[1].IsKeyHandleValid(Rotation[1].FindKey(Time));
-	bool bRzKeyExists = Rotation[2].IsKeyHandleValid(Rotation[2].FindKey(Time));
-
-	// Key each component. Don't add a keyframe if there are existing keys and auto key is not enabled.
-	if( TransformKey.KeyParams.bAddKeyEvenIfUnchanged || (TransformKey.ShouldKeyRotation( EAxis::X ) && !(!bRxKeyExists && !TransformKey.KeyParams.bAutoKeying && Rotation[0].GetNumKeys() > 0)) )
+	if( TransformKey.KeyParams.bAddKeyEvenIfUnchanged || TransformKey.ShouldKeyRotation( EAxis::X ) || Rotation[0].GetNumKeys() == 0 )
 	{
 		AddKeyToCurve( Rotation[0], Time, TransformKey.GetRotationValue().Roll, TransformKey.KeyParams, bUnwindRotation );
 	}
 
-	if( TransformKey.KeyParams.bAddKeyEvenIfUnchanged || (TransformKey.ShouldKeyRotation( EAxis::Y ) && !(!bRyKeyExists && !TransformKey.KeyParams.bAutoKeying && Rotation[1].GetNumKeys() > 0)) )
+	if( TransformKey.KeyParams.bAddKeyEvenIfUnchanged || TransformKey.ShouldKeyRotation( EAxis::Y ) || Rotation[1].GetNumKeys() == 0 )
 	{ 
 		AddKeyToCurve( Rotation[1], Time, TransformKey.GetRotationValue().Pitch, TransformKey.KeyParams, bUnwindRotation );
 	}
 
-	if( TransformKey.KeyParams.bAddKeyEvenIfUnchanged || (TransformKey.ShouldKeyRotation( EAxis::Z ) && !(!bRzKeyExists && !TransformKey.KeyParams.bAutoKeying && Rotation[2].GetNumKeys() > 0)) )
+	if( TransformKey.KeyParams.bAddKeyEvenIfUnchanged || TransformKey.ShouldKeyRotation( EAxis::Z ) || Rotation[2].GetNumKeys() == 0 )
 	{ 
 		AddKeyToCurve( Rotation[2], Time, TransformKey.GetRotationValue().Yaw, TransformKey.KeyParams, bUnwindRotation );
 	}
@@ -193,22 +183,17 @@ void UMovieScene3DTransformSection::AddScaleKeys( const FTransformKey& Transform
 {
 	const float Time = TransformKey.GetKeyTime();
 
-	bool bSxKeyExists = Scale[0].IsKeyHandleValid(Scale[0].FindKey(Time));
-	bool bSyKeyExists = Scale[1].IsKeyHandleValid(Scale[1].FindKey(Time));
-	bool bSzKeyExists = Scale[2].IsKeyHandleValid(Scale[2].FindKey(Time));
-
-	// Key each component. Don't add a keyframe if there are existing keys and auto key is not enabled.
-	if( TransformKey.KeyParams.bAddKeyEvenIfUnchanged || (TransformKey.ShouldKeyScale( EAxis::X ) && !(!bSxKeyExists && !TransformKey.KeyParams.bAutoKeying && Scale[0].GetNumKeys() > 0)) )
+	if( TransformKey.KeyParams.bAddKeyEvenIfUnchanged || TransformKey.ShouldKeyScale( EAxis::X ) || Scale[0].GetNumKeys() == 0 )
 	{
 		AddKeyToCurve( Scale[0], Time, TransformKey.GetScaleValue().X, TransformKey.KeyParams );
 	}
 
-	if( TransformKey.KeyParams.bAddKeyEvenIfUnchanged || (TransformKey.ShouldKeyScale( EAxis::Y ) && !(!bSyKeyExists && !TransformKey.KeyParams.bAutoKeying && Scale[1].GetNumKeys() > 0)) )
+	if( TransformKey.KeyParams.bAddKeyEvenIfUnchanged || TransformKey.ShouldKeyScale( EAxis::Y ) || Scale[1].GetNumKeys() == 0 )
 	{ 
 		AddKeyToCurve( Scale[1], Time, TransformKey.GetScaleValue().Y, TransformKey.KeyParams );
 	}
 
-	if( TransformKey.KeyParams.bAddKeyEvenIfUnchanged || (TransformKey.ShouldKeyScale( EAxis::Z ) && !(!bSzKeyExists && !TransformKey.KeyParams.bAutoKeying && Scale[2].GetNumKeys() > 0)) )
+	if( TransformKey.KeyParams.bAddKeyEvenIfUnchanged || TransformKey.ShouldKeyScale( EAxis::Z ) || Scale[2].GetNumKeys() == 0 )
 	{ 
 		AddKeyToCurve( Scale[2], Time, TransformKey.GetScaleValue().Z, TransformKey.KeyParams );
 	}
@@ -217,8 +202,6 @@ void UMovieScene3DTransformSection::AddScaleKeys( const FTransformKey& Transform
 
 bool UMovieScene3DTransformSection::NewKeyIsNewData(const FTransformKey& TransformKey) const
 {
-	const float KeyTime = TransformKey.GetKeyTime();
-
 	bool bHasEmptyKeys = false;
 	for (int32 i = 0; i < 3; ++i)
 	{
@@ -228,33 +211,9 @@ bool UMovieScene3DTransformSection::NewKeyIsNewData(const FTransformKey& Transfo
 			Scale[i].GetNumKeys() == 0;
 	}
 
-	if (bHasEmptyKeys || TransformKey.ShouldKeyAny())
+	if ( bHasEmptyKeys || (TransformKey.KeyParams.bAutoKeying && TransformKey.ShouldKeyAny() ) )
 	{
-		// Don't add a keyframe if there are existing keys and auto key is not enabled.
-		bool bTxKeyExists = Translation[0].IsKeyHandleValid(Translation[0].FindKey(KeyTime));
-		bool bTyKeyExists = Translation[1].IsKeyHandleValid(Translation[1].FindKey(KeyTime));
-		bool bTzKeyExists = Translation[2].IsKeyHandleValid(Translation[2].FindKey(KeyTime));
-
-		bool bRxKeyExists = Rotation[0].IsKeyHandleValid(Rotation[0].FindKey(KeyTime));
-		bool bRyKeyExists = Rotation[1].IsKeyHandleValid(Rotation[1].FindKey(KeyTime));
-		bool bRzKeyExists = Rotation[2].IsKeyHandleValid(Rotation[2].FindKey(KeyTime));
-
-		bool bSxKeyExists = Scale[0].IsKeyHandleValid(Scale[0].FindKey(KeyTime));
-		bool bSyKeyExists = Scale[1].IsKeyHandleValid(Scale[1].FindKey(KeyTime));
-		bool bSzKeyExists = Scale[2].IsKeyHandleValid(Scale[2].FindKey(KeyTime));
-
-		if ( !(!bTxKeyExists && !TransformKey.KeyParams.bAutoKeying && Translation[0].GetNumKeys() > 0) ||
-			 !(!bTyKeyExists && !TransformKey.KeyParams.bAutoKeying && Translation[1].GetNumKeys() > 0) ||
-			 !(!bTzKeyExists && !TransformKey.KeyParams.bAutoKeying && Translation[2].GetNumKeys() > 0) ||
-			 !(!bRxKeyExists && !TransformKey.KeyParams.bAutoKeying && Rotation[0].GetNumKeys() > 0) ||
-			 !(!bRyKeyExists && !TransformKey.KeyParams.bAutoKeying && Rotation[1].GetNumKeys() > 0) ||
-			 !(!bRzKeyExists && !TransformKey.KeyParams.bAutoKeying && Rotation[2].GetNumKeys() > 0) ||
-			 !(!bSxKeyExists && !TransformKey.KeyParams.bAutoKeying && Scale[0].GetNumKeys() > 0) ||
-			 !(!bSyKeyExists && !TransformKey.KeyParams.bAutoKeying && Scale[1].GetNumKeys() > 0) ||
-			 !(!bSzKeyExists && !TransformKey.KeyParams.bAutoKeying && Scale[2].GetNumKeys() > 0) )
-		{
-			return true;
-		}
+		return true;
 	}
 
 	return false;
diff --git a/Engine/Source/Runtime/MovieSceneTracks/Private/Sections/MovieSceneBoolSection.cpp b/Engine/Source/Runtime/MovieSceneTracks/Private/Sections/MovieSceneBoolSection.cpp
--- a/Engine/Source/Runtime/MovieSceneTracks/Private/Sections/MovieSceneBoolSection.cpp
+++ b/Engine/Source/Runtime/MovieSceneTracks/Private/Sections/MovieSceneBoolSection.cpp
@@ -66,15 +66,9 @@ void UMovieSceneBoolSection::AddKey( float Time, bool Value, FKeyParams KeyParam
 
 bool UMovieSceneBoolSection::NewKeyIsNewData(float Time, bool Value, FKeyParams KeyParams) const
 {
-	if (BoolCurve.GetNumKeys() == 0 || Eval(Time) != Value)
+	if ( BoolCurve.GetNumKeys() == 0 || (KeyParams.bAutoKeying && Eval(Time) != Value) )
 	{
-		bool bKeyExists = BoolCurve.IsKeyHandleValid(BoolCurve.FindKey(Time));
-
-		// Don't add a keyframe if there are existing keys and auto key is not enabled.
-		if (!(!bKeyExists && !KeyParams.bAutoKeying && BoolCurve.GetNumKeys() > 0))
-		{
-			return true;
-		}
+		return true;
 	}
 
 	return false;
diff --git a/Engine/Source/Runtime/MovieSceneTracks/Private/Sections/MovieSceneByteSection.cpp b/Engine/Source/Runtime/MovieSceneTracks/Private/Sections/MovieSceneByteSection.cpp
--- a/Engine/Source/Runtime/MovieSceneTracks/Private/Sections/MovieSceneByteSection.cpp
+++ b/Engine/Source/Runtime/MovieSceneTracks/Private/Sections/MovieSceneByteSection.cpp
@@ -59,15 +59,9 @@ void UMovieSceneByteSection::AddKey( float Time, uint8 Value, FKeyParams KeyPara
 
 bool UMovieSceneByteSection::NewKeyIsNewData(float Time, uint8 Value, FKeyParams KeyParams) const
 {
-	if (ByteCurve.GetNumKeys() == 0 || Eval(Time) != Value)
+	if ( ByteCurve.GetNumKeys() == 0 || (KeyParams.bAutoKeying && Eval(Time) != Value) )
 	{
-		bool bKeyExists = ByteCurve.IsKeyHandleValid(ByteCurve.FindKey(Time));
-
-		// Don't add a keyframe if there are existing keys and auto key is not enabled.
-		if (!(!bKeyExists && !KeyParams.bAutoKeying && ByteCurve.GetNumKeys() > 0))
-		{
-			return true;
-		}
+		return true;
 	}
 
 	return false;
diff --git a/Engine/Source/Runtime/MovieSceneTracks/Private/Sections/MovieSceneColorSection.cpp b/Engine/Source/Runtime/MovieSceneTracks/Private/Sections/MovieSceneColorSection.cpp
--- a/Engine/Source/Runtime/MovieSceneTracks/Private/Sections/MovieSceneColorSection.cpp
+++ b/Engine/Source/Runtime/MovieSceneTracks/Private/Sections/MovieSceneColorSection.cpp
@@ -91,31 +91,26 @@ void UMovieSceneColorSection::AddKey( float Time, const FColorKey& Key, FKeyPara
 
 	FName CurveName = Key.CurveName;
 
-	bool bRedKeyExists = RedCurve.IsKeyHandleValid(RedCurve.FindKey(Time));
-	bool bGreenKeyExists = GreenCurve.IsKeyHandleValid(GreenCurve.FindKey(Time));
-	bool bBlueKeyExists = BlueCurve.IsKeyHandleValid(BlueCurve.FindKey(Time));
-	bool bAlphaKeyExists = AlphaCurve.IsKeyHandleValid(AlphaCurve.FindKey(Time));
-
 	if ( (CurveName == NAME_None || CurveName == RedName) &&
-			(KeyParams.bAddKeyEvenIfUnchanged || !(!bRedKeyExists && !KeyParams.bAutoKeying && RedCurve.GetNumKeys() > 0) ) )
+			(KeyParams.bAddKeyEvenIfUnchanged || !FMath::IsNearlyEqual(RedCurve.Eval(Time), Key.Value.R) || RedCurve.GetNumKeys() == 0) )
 	{
 		AddKeyToCurve(RedCurve, Time, Key.Value.R, KeyParams);
 	}
 
 	if ( (CurveName == NAME_None || CurveName == GreenName) &&
-			(KeyParams.bAddKeyEvenIfUnchanged || !(!bGreenKeyExists && !KeyParams.bAutoKeying && GreenCurve.GetNumKeys() > 0) ) )
+			(KeyParams.bAddKeyEvenIfUnchanged || !FMath::IsNearlyEqual(GreenCurve.Eval(Time), Key.Value.G) || GreenCurve.GetNumKeys() == 0) )
 	{
 		AddKeyToCurve(GreenCurve, Time, Key.Value.G, KeyParams);
 	}
 
 	if ( (CurveName == NAME_None || CurveName == BlueName) &&
-			(KeyParams.bAddKeyEvenIfUnchanged || !(!bBlueKeyExists && !KeyParams.bAutoKeying && BlueCurve.GetNumKeys() > 0) ) )
+			(KeyParams.bAddKeyEvenIfUnchanged || !FMath::IsNearlyEqual(BlueCurve.Eval(Time), Key.Value.B) || BlueCurve.GetNumKeys() == 0) )
 	{
 		AddKeyToCurve(BlueCurve, Time, Key.Value.B, KeyParams);
 	}
 
 	if ( (CurveName == NAME_None || CurveName == AlphaName) &&
-			(KeyParams.bAddKeyEvenIfUnchanged || !(!bAlphaKeyExists && !KeyParams.bAutoKeying && AlphaCurve.GetNumKeys() > 0) ) )
+			(KeyParams.bAddKeyEvenIfUnchanged || !FMath::IsNearlyEqual(AlphaCurve.Eval(Time), Key.Value.A) || AlphaCurve.GetNumKeys() == 0) )
 	{
 		AddKeyToCurve(AlphaCurve, Time, Key.Value.A, KeyParams);
 	}
@@ -128,24 +123,11 @@ bool UMovieSceneColorSection::NewKeyIsNewData(float Time, FLinearColor Value, FK
 		RedCurve.GetNumKeys() == 0 ||
 		GreenCurve.GetNumKeys() == 0 ||
 		BlueCurve.GetNumKeys() == 0 ||
-		AlphaCurve.GetNumKeys() == 0 ||
-		!Eval(Time,Value).Equals(Value);
+		AlphaCurve.GetNumKeys() == 0;
 
-	if (bHasEmptyKeys)
+	if ( bHasEmptyKeys || (KeyParams.bAutoKeying && !Eval(Time, Value).Equals(Value) ) )
 	{
-		// Don't add a keyframe if there are existing keys and auto key is not enabled.
-		bool bRedKeyExists = RedCurve.IsKeyHandleValid(RedCurve.FindKey(Time));
-		bool bGreenKeyExists = GreenCurve.IsKeyHandleValid(GreenCurve.FindKey(Time));
-		bool bBlueKeyExists = BlueCurve.IsKeyHandleValid(BlueCurve.FindKey(Time));
-		bool bAlphaKeyExists = AlphaCurve.IsKeyHandleValid(AlphaCurve.FindKey(Time));
-
-		if ( !(!bRedKeyExists && !KeyParams.bAutoKeying && RedCurve.GetNumKeys() > 0) ||
-			 !(!bGreenKeyExists && !KeyParams.bAutoKeying && GreenCurve.GetNumKeys() > 0) ||
-			 !(!bBlueKeyExists && !KeyParams.bAutoKeying && BlueCurve.GetNumKeys() > 0) ||
-			 !(!bAlphaKeyExists && !KeyParams.bAutoKeying && AlphaCurve.GetNumKeys() > 0) )
-		{
-			return true;
-		}
+		return true;
 	}
 
 	return false;
diff --git a/Engine/Source/Runtime/MovieSceneTracks/Private/Sections/MovieSceneFloatSection.cpp b/Engine/Source/Runtime/MovieSceneTracks/Private/Sections/MovieSceneFloatSection.cpp
--- a/Engine/Source/Runtime/MovieSceneTracks/Private/Sections/MovieSceneFloatSection.cpp
+++ b/Engine/Source/Runtime/MovieSceneTracks/Private/Sections/MovieSceneFloatSection.cpp
@@ -62,15 +62,9 @@ void UMovieSceneFloatSection::AddKey( float Time, float Value, FKeyParams KeyPar
 
 bool UMovieSceneFloatSection::NewKeyIsNewData(float Time, float Value, FKeyParams KeyParams) const
 {
-	if (FloatCurve.GetNumKeys() == 0 || Eval(Time) != Value)
+	if ( FloatCurve.GetNumKeys() == 0 || (KeyParams.bAutoKeying && Eval(Time) != Value) )
 	{
-		bool bKeyExists = FloatCurve.IsKeyHandleValid(FloatCurve.FindKey(Time));
-
-		// Don't add a keyframe if there are existing keys and auto key is not enabled.
-		if (!(!bKeyExists && !KeyParams.bAutoKeying && FloatCurve.GetNumKeys() > 0))
-		{
-			return true;
-		}
+		return true;
 	}
 
 	return false;
diff --git a/Engine/Source/Runtime/MovieSceneTracks/Private/Sections/MovieSceneVectorSection.cpp b/Engine/Source/Runtime/MovieSceneTracks/Private/Sections/MovieSceneVectorSection.cpp
--- a/Engine/Source/Runtime/MovieSceneTracks/Private/Sections/MovieSceneVectorSection.cpp
+++ b/Engine/Source/Runtime/MovieSceneTracks/Private/Sections/MovieSceneVectorSection.cpp
@@ -45,11 +45,7 @@ void UMovieSceneVectorSection::AddKey( float Time, FName CurveName, const FVecto
 			}
 			else
 			{
-				bool bKeyExists = Curves[i].IsKeyHandleValid(Curves[i].FindKey(Time));
-				if ( KeyParams.bAddKeyEvenIfUnchanged || !(!bKeyExists && !KeyParams.bAutoKeying && Curves[i].GetNumKeys() > 0) )
-				{
-					Curves[i].UpdateOrAddKey(Time, Value[i]);
-				}
+				Curves[i].UpdateOrAddKey(Time, Value[i]);
 			}
 		}
 	}
@@ -60,33 +56,34 @@ bool UMovieSceneVectorSection::NewKeyIsNewData(float Time, const FVector4& Value
 {
 	check(ChannelsUsed >= 2 && ChannelsUsed <= 4);
 
+	bool bHasEmptyKeys = false;
+	for (int32 i = 0; i < ChannelsUsed; ++i)
+	{
+		if (Curves[i].GetNumKeys() == 0)
+		{
+			bHasEmptyKeys = true;
+			break;
+		}
+	}
+
 	bool bNewData = false;
 	for (int32 i = 0; i < ChannelsUsed; ++i)
 	{
 		float OriginalData = Curves[i].Eval(Time);
 		// don't re-add keys if the data already matches
-		if (Curves[i].GetNumKeys() == 0 || !FMath::IsNearlyEqual(OriginalData, Value[i]))
+		if (!FMath::IsNearlyEqual(OriginalData, Value[i]))
 		{
 			bNewData = true;
-
 			break;
 		}
 	}
 
-	if (bNewData)
+	if ( bHasEmptyKeys || (KeyParams.bAutoKeying && bNewData) )
 	{
-		for (int32 i = 0; i < ChannelsUsed; ++i)
-		{
-			// Don't add a keyframe if there are existing keys and auto key is not enabled.
-			bool bKeyExists = Curves[i].IsKeyHandleValid(Curves[i].FindKey(Time));
-			if ( !(!bKeyExists && !KeyParams.bAutoKeying && Curves[i].GetNumKeys() > 0) )
-			{
-				return true;
-			}
-		}
+		return true;
 	}
 
-	return bNewData;
+	return false;
 }
 
 
diff --git a/Engine/Source/Runtime/MovieSceneTracks/Private/Tracks/MovieSceneBoolTrack.cpp b/Engine/Source/Runtime/MovieSceneTracks/Private/Tracks/MovieSceneBoolTrack.cpp
--- a/Engine/Source/Runtime/MovieSceneTracks/Private/Tracks/MovieSceneBoolTrack.cpp
+++ b/Engine/Source/Runtime/MovieSceneTracks/Private/Tracks/MovieSceneBoolTrack.cpp
@@ -51,3 +51,13 @@ bool UMovieSceneBoolTrack::Eval( float Position, float LastPostion, bool& OutBoo
 
 	return Section != NULL;
 }
+
+bool UMovieSceneBoolTrack::CanKeyTrack(float Time, bool Value, FKeyParams KeyParams) const
+{
+	const UMovieSceneSection* NearestSection = MovieSceneHelpers::FindNearestSectionAtTime( Sections, Time );
+	if (!NearestSection || CastChecked<UMovieSceneBoolSection>(NearestSection)->NewKeyIsNewData(Time, Value, KeyParams))
+	{
+		return true;
+	}
+	return false;
+}
diff --git a/Engine/Source/Runtime/MovieSceneTracks/Private/Tracks/MovieSceneByteTrack.cpp b/Engine/Source/Runtime/MovieSceneTracks/Private/Tracks/MovieSceneByteTrack.cpp
--- a/Engine/Source/Runtime/MovieSceneTracks/Private/Tracks/MovieSceneByteTrack.cpp
+++ b/Engine/Source/Runtime/MovieSceneTracks/Private/Tracks/MovieSceneByteTrack.cpp
@@ -55,6 +55,15 @@ bool UMovieSceneByteTrack::Eval( float Position, float LastPostion, uint8& OutBy
 	return Section != nullptr;
 }
 
+bool UMovieSceneByteTrack::CanKeyTrack(float Time, uint8 Value, FKeyParams KeyParams) const
+{
+	const UMovieSceneSection* NearestSection = MovieSceneHelpers::FindNearestSectionAtTime( Sections, Time );
+	if (!NearestSection || CastChecked<UMovieSceneByteSection>(NearestSection)->NewKeyIsNewData(Time, Value, KeyParams))
+	{
+		return true;
+	}
+	return false;
+}
 
 void UMovieSceneByteTrack::SetEnum(UEnum* InEnum)
 {
diff --git a/Engine/Source/Runtime/MovieSceneTracks/Private/Tracks/MovieSceneColorTrack.cpp b/Engine/Source/Runtime/MovieSceneTracks/Private/Tracks/MovieSceneColorTrack.cpp
--- a/Engine/Source/Runtime/MovieSceneTracks/Private/Tracks/MovieSceneColorTrack.cpp
+++ b/Engine/Source/Runtime/MovieSceneTracks/Private/Tracks/MovieSceneColorTrack.cpp
@@ -52,4 +52,15 @@ bool UMovieSceneColorTrack::Eval( float Position, float LastPosition, FLinearCol
 	}
 
 	return Section != nullptr;
-}
\ No newline at end of file
+}
+
+bool UMovieSceneColorTrack::CanKeyTrack(float Time, const FColorKey& Key, FKeyParams KeyParams) const
+{
+	const UMovieSceneSection* NearestSection = MovieSceneHelpers::FindNearestSectionAtTime( Sections, Time );
+	if (!NearestSection || CastChecked<UMovieSceneColorSection>(NearestSection)->NewKeyIsNewData(Time, Key.Value, KeyParams))
+	{
+		return true;
+	}
+	return false;
+}
+
diff --git a/Engine/Source/Runtime/MovieSceneTracks/Private/Tracks/MovieSceneFloatTrack.cpp b/Engine/Source/Runtime/MovieSceneTracks/Private/Tracks/MovieSceneFloatTrack.cpp
--- a/Engine/Source/Runtime/MovieSceneTracks/Private/Tracks/MovieSceneFloatTrack.cpp
+++ b/Engine/Source/Runtime/MovieSceneTracks/Private/Tracks/MovieSceneFloatTrack.cpp
@@ -54,3 +54,14 @@ bool UMovieSceneFloatTrack::Eval( float Position, float LastPosition, float& Out
 
 	return Section != nullptr;
 }
+
+bool UMovieSceneFloatTrack::CanKeyTrack(float Time, float Value, FKeyParams KeyParams) const
+{
+	const UMovieSceneSection* NearestSection = MovieSceneHelpers::FindNearestSectionAtTime( Sections, Time );
+	if (!NearestSection || CastChecked<UMovieSceneFloatSection>(NearestSection)->NewKeyIsNewData(Time, Value, KeyParams))
+	{
+		return true;
+	}
+	return false;
+}
+
diff --git a/Engine/Source/Runtime/MovieSceneTracks/Private/Tracks/MovieSceneVectorTrack.cpp b/Engine/Source/Runtime/MovieSceneTracks/Private/Tracks/MovieSceneVectorTrack.cpp
--- a/Engine/Source/Runtime/MovieSceneTracks/Private/Tracks/MovieSceneVectorTrack.cpp
+++ b/Engine/Source/Runtime/MovieSceneTracks/Private/Tracks/MovieSceneVectorTrack.cpp
@@ -98,3 +98,14 @@ bool UMovieSceneVectorTrack::Eval( float Position, float LastPosition, FVector4&
 
 	return Section != nullptr;
 }
+
+bool UMovieSceneVectorTrack::CanKeyTrack(float Time, const FVectorKey& Key, FKeyParams KeyParams) const
+{
+	const UMovieSceneSection* NearestSection = MovieSceneHelpers::FindNearestSectionAtTime( Sections, Time );
+	if (!NearestSection || CastChecked<UMovieSceneVectorSection>(NearestSection)->NewKeyIsNewData(Time, Key.Value, KeyParams))
+	{
+		return true;
+	}
+	return false;
+}
+
diff --git a/Engine/Source/Runtime/MovieSceneTracks/Public/Tracks/MovieSceneBoolTrack.h b/Engine/Source/Runtime/MovieSceneTracks/Public/Tracks/MovieSceneBoolTrack.h
--- a/Engine/Source/Runtime/MovieSceneTracks/Public/Tracks/MovieSceneBoolTrack.h
+++ b/Engine/Source/Runtime/MovieSceneTracks/Public/Tracks/MovieSceneBoolTrack.h
@@ -39,4 +39,14 @@ public:
 	 * @return true if anything was evaluated. Note: if false is returned OutBool remains unchanged
 	 */
 	virtual bool Eval( float Position, float LastPostion, bool& OutBool ) const;
+
+	/*
+	 * Get whether the track can be keyed at a particular time.
+	 *
+	 * @param Time				The time relative to the owning movie scene where the section should be
+	 * @param Value				The value of the key
+	 * @param KeyParams         The keying parameters
+	 * @return Whether the track can be keyed
+	 */
+	virtual bool CanKeyTrack( float Time, bool Value, FKeyParams KeyParams) const;
 };
diff --git a/Engine/Source/Runtime/MovieSceneTracks/Public/Tracks/MovieSceneByteTrack.h b/Engine/Source/Runtime/MovieSceneTracks/Public/Tracks/MovieSceneByteTrack.h
--- a/Engine/Source/Runtime/MovieSceneTracks/Public/Tracks/MovieSceneByteTrack.h
+++ b/Engine/Source/Runtime/MovieSceneTracks/Public/Tracks/MovieSceneByteTrack.h
@@ -40,6 +40,16 @@ public:
 	 */
 	virtual bool Eval( float Position, float LastPostion, uint8& OutByte ) const;
 
+	/**
+	 * Get whether the track can be keyed at a particular time.
+	 *
+	 * @param Time				The time relative to the owning movie scene where the section should be
+	 * @param Value				The value of the key
+	 * @param KeyParams         The keying parameters
+	 * @return Whether the track can be keyed
+	 */
+	virtual bool CanKeyTrack( float Time, uint8 Value, FKeyParams KeyParams ) const;
+
 	void SetEnum(UEnum* Enum);
 
 	class UEnum* GetEnum() const;
diff --git a/Engine/Source/Runtime/MovieSceneTracks/Public/Tracks/MovieSceneColorTrack.h b/Engine/Source/Runtime/MovieSceneTracks/Public/Tracks/MovieSceneColorTrack.h
--- a/Engine/Source/Runtime/MovieSceneTracks/Public/Tracks/MovieSceneColorTrack.h
+++ b/Engine/Source/Runtime/MovieSceneTracks/Public/Tracks/MovieSceneColorTrack.h
@@ -46,6 +46,17 @@ public:
 	 * @return true if anything was evaluated. Note: if false is returned OutColor remains unchanged
 	 */
 	virtual bool Eval( float Position, float LastPostion, FLinearColor& OutColor ) const;
+
+	/**
+	 * Get whether the track can be keyed at a particular time.
+	 *
+	 * @param Time				The time relative to the owning movie scene where the section should be
+	 * @param Value				The value of the key
+	 * @param KeyParams         The keying parameters
+	 * @return Whether the track can be keyed
+	 */
+	virtual bool CanKeyTrack( float Time, const FColorKey& Key, FKeyParams KeyParams ) const;
+
 private:
 	UPROPERTY()
 	bool bIsSlateColor_DEPRECATED;
diff --git a/Engine/Source/Runtime/MovieSceneTracks/Public/Tracks/MovieSceneFloatTrack.h b/Engine/Source/Runtime/MovieSceneTracks/Public/Tracks/MovieSceneFloatTrack.h
--- a/Engine/Source/Runtime/MovieSceneTracks/Public/Tracks/MovieSceneFloatTrack.h
+++ b/Engine/Source/Runtime/MovieSceneTracks/Public/Tracks/MovieSceneFloatTrack.h
@@ -40,4 +40,13 @@ public:
 	 */
 	virtual bool Eval( float Position, float LastPosition, float& OutFloat ) const;
 
+	/**
+	 * Get whether the track can be keyed at a particular time.
+	 *
+	 * @param Time				The time relative to the owning movie scene where the section should be
+	 * @param Value				The value of the key
+	 * @param KeyParams         The keying parameters
+	 * @return Whether the track can be keyed
+	 */
+	virtual bool CanKeyTrack( float Time, float Value, FKeyParams KeyParams ) const;
 };
diff --git a/Engine/Source/Runtime/MovieSceneTracks/Public/Tracks/MovieSceneVectorTrack.h b/Engine/Source/Runtime/MovieSceneTracks/Public/Tracks/MovieSceneVectorTrack.h
--- a/Engine/Source/Runtime/MovieSceneTracks/Public/Tracks/MovieSceneVectorTrack.h
+++ b/Engine/Source/Runtime/MovieSceneTracks/Public/Tracks/MovieSceneVectorTrack.h
@@ -53,6 +53,16 @@ public:
 	 */
 	virtual bool Eval( float Position, float LastPostion, FVector4& InOutVector ) const;
 
+	/**
+	 * Get whether the track can be keyed at a particular time.
+	 *
+	 * @param Time				The time relative to the owning movie scene where the section should be
+	 * @param Value				The value of the key
+	 * @param KeyParams         The keying parameters
+	 * @return Whether the track can be keyed
+	 */
+	virtual bool CanKeyTrack( float Time, const FVectorKey& Key, FKeyParams KeyParams ) const;
+
 	/** @return Get the number of channels used by the vector */
 	int32 GetNumChannelsUsed() const { return NumChannelsUsed; }
 
diff --git a/Engine/Source/Runtime/UMG/Private/Animation/MovieScene2DTransformTrack.cpp b/Engine/Source/Runtime/UMG/Private/Animation/MovieScene2DTransformTrack.cpp
--- a/Engine/Source/Runtime/UMG/Private/Animation/MovieScene2DTransformTrack.cpp
+++ b/Engine/Source/Runtime/UMG/Private/Animation/MovieScene2DTransformTrack.cpp
@@ -53,4 +53,14 @@ bool UMovieScene2DTransformTrack::AddKeyToSection(float Time, const F2DTransform
 		return true;
 	}
 	return false;
-}
\ No newline at end of file
+}
+
+bool UMovieScene2DTransformTrack::CanKeyTrack(float Time, const F2DTransformKey& TransformKey, FKeyParams KeyParams) const
+{
+	const UMovieSceneSection* NearestSection = MovieSceneHelpers::FindNearestSectionAtTime( Sections, Time );
+	if (!NearestSection || CastChecked<UMovieScene2DTransformSection>(NearestSection)->NewKeyIsNewData(Time, TransformKey.Value, KeyParams))
+	{
+		return true;
+	}
+	return false;
+}
diff --git a/Engine/Source/Runtime/UMG/Private/Animation/MovieSceneMarginTrack.cpp b/Engine/Source/Runtime/UMG/Private/Animation/MovieSceneMarginTrack.cpp
--- a/Engine/Source/Runtime/UMG/Private/Animation/MovieSceneMarginTrack.cpp
+++ b/Engine/Source/Runtime/UMG/Private/Animation/MovieSceneMarginTrack.cpp
@@ -54,4 +54,14 @@ bool UMovieSceneMarginTrack::Eval( float Position, float LastPosition, FMargin&
 	}
 
 	return Section != NULL;
-}
\ No newline at end of file
+}
+
+bool UMovieSceneMarginTrack::CanKeyTrack(float Time, const FMarginKey& MarginKey, FKeyParams KeyParams) const
+{
+	const UMovieSceneSection* NearestSection = MovieSceneHelpers::FindNearestSectionAtTime( Sections, Time );
+	if (!NearestSection || CastChecked<UMovieSceneMarginSection>(NearestSection)->NewKeyIsNewData(Time, MarginKey.Value, KeyParams))
+	{
+		return true;
+	}
+	return false;
+}
diff --git a/Engine/Source/Runtime/UMG/Public/Animation/MovieScene2DTransformTrack.h b/Engine/Source/Runtime/UMG/Public/Animation/MovieScene2DTransformTrack.h
--- a/Engine/Source/Runtime/UMG/Public/Animation/MovieScene2DTransformTrack.h
+++ b/Engine/Source/Runtime/UMG/Public/Animation/MovieScene2DTransformTrack.h
@@ -47,4 +47,13 @@ public:
 	 */
 	UMG_API bool AddKeyToSection( float Time, const F2DTransformKey& TransformKey, FKeyParams KeyParams );
 
+	/**
+	 * Get whether the track can be keyed at a particular time.
+	 *
+	 * @param Time				The time relative to the owning movie scene where the section should be
+	 * @param Value				The value of the key
+	 * @param KeyParams         The keying parameters 
+	 * @return True if the key was successfully added.
+	 */
+	UMG_API bool CanKeyTrack( float Time, const F2DTransformKey& TransformKey, FKeyParams KeyParams ) const;
 };
\ No newline at end of file
diff --git a/Engine/Source/Runtime/UMG/Public/Animation/MovieSceneMarginTrack.h b/Engine/Source/Runtime/UMG/Public/Animation/MovieSceneMarginTrack.h
--- a/Engine/Source/Runtime/UMG/Public/Animation/MovieSceneMarginTrack.h
+++ b/Engine/Source/Runtime/UMG/Public/Animation/MovieSceneMarginTrack.h
@@ -45,4 +45,13 @@ public:
 	 */
 	bool Eval( float Position, float LastPostion, FMargin& InOutMargin ) const;
 
+	/**
+	 * Get whether the track can be keyed at a particular time.
+	 *
+	 * @param Time				The time relative to the owning movie scene where the section should be
+	 * @param Value				The value of the key
+	 * @param KeyParams         The keying parameters
+	 * @return True if the key was successfully added.
+	 */
+	UMG_API bool CanKeyTrack( float Time, const FMarginKey& MarginKey, FKeyParams KeyParams ) const;
 };
