Sequencer: Fix skeletal mesh animation timing.

- Reverse dilation so that 0.5 means slower and 2x means faster.
- Protect against 0 dilation value.

diff --git a/Engine/Source/Editor/MovieSceneTools/Private/TrackEditors/SkeletalAnimationTrackEditor.cpp b/Engine/Source/Editor/MovieSceneTools/Private/TrackEditors/SkeletalAnimationTrackEditor.cpp
--- a/Engine/Source/Editor/MovieSceneTools/Private/TrackEditors/SkeletalAnimationTrackEditor.cpp
+++ b/Engine/Source/Editor/MovieSceneTools/Private/TrackEditors/SkeletalAnimationTrackEditor.cpp
@@ -100,7 +100,7 @@ int32 FSkeletalAnimationSection::OnPaintSection( const FGeometry& AllottedGeomet
 
 	// Add lines where the animation starts and ends/loops
 	float CurrentTime = AnimSection->GetAnimationStartTime();
-	while (CurrentTime < AnimSection->GetEndTime())
+	while (CurrentTime < AnimSection->GetEndTime() && !FMath::IsNearlyZero(AnimSection->GetAnimationDuration()))
 	{
 		if (CurrentTime > AnimSection->GetStartTime())
 		{
diff --git a/Engine/Source/Runtime/MovieSceneTracks/Private/Sections/MovieSceneSkeletalAnimationSection.cpp b/Engine/Source/Runtime/MovieSceneTracks/Private/Sections/MovieSceneSkeletalAnimationSection.cpp
--- a/Engine/Source/Runtime/MovieSceneTracks/Private/Sections/MovieSceneSkeletalAnimationSection.cpp
+++ b/Engine/Source/Runtime/MovieSceneTracks/Private/Sections/MovieSceneSkeletalAnimationSection.cpp
@@ -41,7 +41,7 @@ void UMovieSceneSkeletalAnimationSection::GetSnapTimes(TArray<float>& OutSnapTim
 	Super::GetSnapTimes(OutSnapTimes, bGetSectionBorders);
 
 	float CurrentTime = GetAnimationStartTime();
-	while (CurrentTime <= GetEndTime())
+	while (CurrentTime <= GetEndTime() && !FMath::IsNearlyZero(GetAnimationDuration()))
 	{
 		if (CurrentTime >= GetStartTime())
 		{
diff --git a/Engine/Source/Runtime/MovieSceneTracks/Private/TrackInstances/MovieSceneSkeletalAnimationTrackInstance.cpp b/Engine/Source/Runtime/MovieSceneTracks/Private/TrackInstances/MovieSceneSkeletalAnimationTrackInstance.cpp
--- a/Engine/Source/Runtime/MovieSceneTracks/Private/TrackInstances/MovieSceneSkeletalAnimationTrackInstance.cpp
+++ b/Engine/Source/Runtime/MovieSceneTracks/Private/TrackInstances/MovieSceneSkeletalAnimationTrackInstance.cpp
@@ -35,7 +35,7 @@ void FMovieSceneSkeletalAnimationTrackInstance::Update( float Position, float La
 				int32 ChannelIndex = 0;
 				FName SlotName = FName("AnimationSlot");
 				UAnimSequence* AnimSequence = AnimSection->GetAnimSequence();
-				float AnimPosition = FMath::Fmod(Position - AnimSection->GetAnimationStartTime(), AnimSection->GetAnimationDuration()) / AnimSection->GetAnimationDilationFactor();
+				float AnimPosition = FMath::Fmod(Position - AnimSection->GetAnimationStartTime(), AnimSection->GetAnimationDuration()) * AnimSection->GetAnimationDilationFactor();
 
 				AnimInterface->PreviewBeginAnimControl(nullptr);
 				AnimInterface->PreviewSetAnimPosition(SlotName, ChannelIndex, AnimSequence, AnimPosition, true, false, 0.f);
diff --git a/Engine/Source/Runtime/MovieSceneTracks/Public/Sections/MovieSceneSkeletalAnimationSection.h b/Engine/Source/Runtime/MovieSceneTracks/Public/Sections/MovieSceneSkeletalAnimationSection.h
--- a/Engine/Source/Runtime/MovieSceneTracks/Public/Sections/MovieSceneSkeletalAnimationSection.h
+++ b/Engine/Source/Runtime/MovieSceneTracks/Public/Sections/MovieSceneSkeletalAnimationSection.h
@@ -26,7 +26,7 @@ public:
 	float GetAnimationStartTime() const {return AnimationStartTime;}
 	
 	/** Gets the animation duration, modified by dilation */
-	float GetAnimationDuration() const {return AnimSequence->SequenceLength * AnimationDilationFactor;}
+	float GetAnimationDuration() const {return FMath::IsNearlyZero(AnimationDilationFactor) ? 0.f : AnimSequence->SequenceLength / AnimationDilationFactor;}
 
 	/**
 	 * @return The dilation multiplier of the animation
