MovieSceneTools: Continued to implement event tracks

diff --git a/Engine/Source/Editor/MovieSceneTools/Private/TrackEditors/EventTrackEditor.cpp b/Engine/Source/Editor/MovieSceneTools/Private/TrackEditors/EventTrackEditor.cpp
--- a/Engine/Source/Editor/MovieSceneTools/Private/TrackEditors/EventTrackEditor.cpp
+++ b/Engine/Source/Editor/MovieSceneTools/Private/TrackEditors/EventTrackEditor.cpp
@@ -31,6 +31,17 @@ void FEventTrackEditor::AddKey(const FGuid& ObjectGuid, UObject* AdditionalAsset
 }
 
 
+void FEventTrackEditor::BuildObjectBindingTrackMenu(FMenuBuilder& MenuBuilder, const FGuid& ObjectBinding, const UClass* ObjectClass)
+{
+	if (!ObjectClass->IsChildOf(AActor::StaticClass()))
+	{
+		return;
+	}
+
+
+}
+
+
 TSharedRef<ISequencerSection> FEventTrackEditor::MakeSectionInterface(UMovieSceneSection& SectionObject, UMovieSceneTrack* Track)
 {
 	return MakeShareable(new FEventTrackSection(SectionObject/*, Track->GetTrackName()*/, GetSequencer()));
diff --git a/Engine/Source/Editor/MovieSceneTools/Private/TrackEditors/EventTrackEditor.h b/Engine/Source/Editor/MovieSceneTools/Private/TrackEditors/EventTrackEditor.h
--- a/Engine/Source/Editor/MovieSceneTools/Private/TrackEditors/EventTrackEditor.h
+++ b/Engine/Source/Editor/MovieSceneTools/Private/TrackEditors/EventTrackEditor.h
@@ -39,6 +39,7 @@ public:
 	// FMovieSceneTrackEditor interface
 
 	virtual void AddKey(const FGuid& ObjectGuid, UObject* AdditionalAsset = nullptr) override;
+	virtual void BuildObjectBindingTrackMenu(FMenuBuilder& MenuBuilder, const FGuid& ObjectBinding, const UClass* ObjectClass) override;
 	virtual TSharedRef<ISequencerSection> MakeSectionInterface(UMovieSceneSection& SectionObject, UMovieSceneTrack* Track) override;
 	virtual bool SupportsType( TSubclassOf<UMovieSceneTrack> Type ) const override;
 };
diff --git a/Engine/Source/Runtime/MovieSceneTracks/Private/Sections/MovieSceneEventSection.cpp b/Engine/Source/Runtime/MovieSceneTracks/Private/Sections/MovieSceneEventSection.cpp
--- a/Engine/Source/Runtime/MovieSceneTracks/Private/Sections/MovieSceneEventSection.cpp
+++ b/Engine/Source/Runtime/MovieSceneTracks/Private/Sections/MovieSceneEventSection.cpp
@@ -9,6 +9,16 @@
 
 void UMovieSceneEventSection::AddKey(float Time, const FName& EventName, FKeyParams KeyParams)
 {
+	for (auto& Key : Keys)
+	{
+		if (Key.Time == Time)
+		{
+			Key.EventNames.AddUnique(EventName);
+
+			return;
+		}
+	}
+
 	Keys.HeapPush(FMovieSceneEventSectionKey(EventName, Time));
 }
 
diff --git a/Engine/Source/Runtime/MovieSceneTracks/Public/Sections/MovieSceneEventSection.h b/Engine/Source/Runtime/MovieSceneTracks/Public/Sections/MovieSceneEventSection.h
--- a/Engine/Source/Runtime/MovieSceneTracks/Public/Sections/MovieSceneEventSection.h
+++ b/Engine/Source/Runtime/MovieSceneTracks/Public/Sections/MovieSceneEventSection.h
@@ -13,9 +13,9 @@ struct FMovieSceneEventSectionKey
 {
 	GENERATED_USTRUCT_BODY()
 
-	/** The name of the event to be triggered. */
+	/** The names of the events to be triggered. */
 	UPROPERTY(EditAnywhere, Category=EventTrackKey)
-	FName EventName;
+	TArray<FName> EventNames;
 
 	/** The time at which the event should be triggered. */
 	UPROPERTY()
@@ -28,9 +28,10 @@ struct FMovieSceneEventSectionKey
 
 	/** Creates and initializes a new instance. */
 	FMovieSceneEventSectionKey(const FName& InEventName, float InTime)
-		: EventName(InEventName)
-		, Time(InTime)
-	{ }
+		: Time(InTime)
+	{
+		EventNames.Add(InEventName);
+	}
 
 	/** Operator less, used to sort the heap based on time until execution. */
 	bool operator<(const FMovieSceneEventSectionKey& Other) const
@@ -72,6 +73,6 @@ public:
 private:
 
 	/** The section's keys. */
-	UPROPERTY()
+	UPROPERTY(EditAnywhere, Category="Audio")
 	TArray<FMovieSceneEventSectionKey> Keys;
 };
