Sequencer: Fix skeletal animation position - position wasn't computed properly when the section was cropped in or dilated.

diff --git a/Engine/Source/Runtime/MovieSceneTracks/Private/TrackInstances/MovieSceneSkeletalAnimationTrackInstance.cpp b/Engine/Source/Runtime/MovieSceneTracks/Private/TrackInstances/MovieSceneSkeletalAnimationTrackInstance.cpp
--- a/Engine/Source/Runtime/MovieSceneTracks/Private/TrackInstances/MovieSceneSkeletalAnimationTrackInstance.cpp
+++ b/Engine/Source/Runtime/MovieSceneTracks/Private/TrackInstances/MovieSceneSkeletalAnimationTrackInstance.cpp
@@ -35,10 +35,33 @@ void FMovieSceneSkeletalAnimationTrackInstance::Update( float Position, float La
 				int32 ChannelIndex = 0;
 				FName SlotName = FName("AnimationSlot");
 				UAnimSequence* AnimSequence = AnimSection->GetAnimSequence();
-				float AnimPosition = FMath::Fmod(Position - AnimSection->GetAnimationStartTime(), AnimSection->GetAnimationDuration()) * AnimSection->GetAnimationDilationFactor();
 
-				AnimInterface->PreviewBeginAnimControl(nullptr);
-				AnimInterface->PreviewSetAnimPosition(SlotName, ChannelIndex, AnimSequence, AnimPosition, true, false, 0.f);
+				float AnimStartOffset = AnimSection->GetStartTime() - AnimSection->GetAnimationStartTime();
+				float AnimEndOffset = AnimSection->GetEndTime() - AnimSection->GetAnimationStartTime();
+
+				float AnimPosition = (Position - AnimSection->GetAnimationStartTime()) * AnimSection->GetAnimationDilationFactor();
+
+				// Looping if greater than first section
+				const bool bLooping = Position > (AnimSection->GetAnimationStartTime() + (AnimSection->GetAnimationSequenceLength() / AnimSection->GetAnimationDilationFactor()));
+
+				if (bLooping)
+				{
+					AnimPosition = FMath::Fmod(Position - AnimSection->GetAnimationStartTime(), AnimSection->GetAnimationSequenceLength() / AnimSection->GetAnimationDilationFactor());
+					AnimPosition *= AnimSection->GetAnimationDilationFactor();
+				}
+
+				AnimPosition = FMath::Clamp(AnimPosition, AnimPosition, AnimSection->GetEndTime());
+
+				if (GIsEditor && !GWorld->HasBegunPlay())
+				{
+					AnimInterface->PreviewBeginAnimControl(nullptr);
+					AnimInterface->PreviewSetAnimPosition(SlotName, ChannelIndex, AnimSequence, AnimPosition, true, false, 0.f);
+				}
+				else
+				{
+					AnimInterface->BeginAnimControl(nullptr);
+					AnimInterface->SetAnimPosition(SlotName, ChannelIndex, AnimSequence, AnimPosition, true, false);
+				}
 			}
 		}
 	}
diff --git a/Engine/Source/Runtime/MovieSceneTracks/Public/Sections/MovieSceneSkeletalAnimationSection.h b/Engine/Source/Runtime/MovieSceneTracks/Public/Sections/MovieSceneSkeletalAnimationSection.h
--- a/Engine/Source/Runtime/MovieSceneTracks/Public/Sections/MovieSceneSkeletalAnimationSection.h
+++ b/Engine/Source/Runtime/MovieSceneTracks/Public/Sections/MovieSceneSkeletalAnimationSection.h
@@ -28,6 +28,9 @@ public:
 	/** Gets the animation duration, modified by dilation */
 	float GetAnimationDuration() const {return FMath::IsNearlyZero(AnimationDilationFactor) ? 0.f : AnimSequence->SequenceLength / AnimationDilationFactor;}
 
+	/** Gets the animation sequence length, not modified by dilation */
+	float GetAnimationSequenceLength() const {return AnimSequence->SequenceLength;}
+
 	/**
 	 * @return The dilation multiplier of the animation
 	 */
