Sequencer: Protect against 0 anim dilation.

diff --git a/Engine/Source/Runtime/MovieSceneTracks/Private/TrackInstances/MovieSceneSkeletalAnimationTrackInstance.cpp b/Engine/Source/Runtime/MovieSceneTracks/Private/TrackInstances/MovieSceneSkeletalAnimationTrackInstance.cpp
--- a/Engine/Source/Runtime/MovieSceneTracks/Private/TrackInstances/MovieSceneSkeletalAnimationTrackInstance.cpp
+++ b/Engine/Source/Runtime/MovieSceneTracks/Private/TrackInstances/MovieSceneSkeletalAnimationTrackInstance.cpp
@@ -36,20 +36,20 @@ void FMovieSceneSkeletalAnimationTrackInstance::Update( float Position, float La
 				FName SlotName = FName("AnimationSlot");
 				UAnimSequence* AnimSequence = AnimSection->GetAnimSequence();
 
-				float AnimStartOffset = AnimSection->GetStartTime() - AnimSection->GetAnimationStartTime();
-				float AnimEndOffset = AnimSection->GetEndTime() - AnimSection->GetAnimationStartTime();
+				float AnimDilationFactor = FMath::IsNearlyZero(AnimSection->GetAnimationDilationFactor()) ? 1.0f : AnimSection->GetAnimationDilationFactor();
 
-				float AnimPosition = (Position - AnimSection->GetAnimationStartTime()) * AnimSection->GetAnimationDilationFactor();
+				float AnimPosition = (Position - AnimSection->GetAnimationStartTime()) * AnimDilationFactor;
 
 				// Looping if greater than first section
-				const bool bLooping = Position > (AnimSection->GetAnimationStartTime() + (AnimSection->GetAnimationSequenceLength() / AnimSection->GetAnimationDilationFactor()));
+				const bool bLooping = Position > (AnimSection->GetAnimationStartTime() + (AnimSection->GetAnimationSequenceLength() / AnimDilationFactor));
 
 				if (bLooping)
 				{
-					AnimPosition = FMath::Fmod(Position - AnimSection->GetAnimationStartTime(), AnimSection->GetAnimationSequenceLength() / AnimSection->GetAnimationDilationFactor());
-					AnimPosition *= AnimSection->GetAnimationDilationFactor();
+					AnimPosition = FMath::Fmod(Position - AnimSection->GetAnimationStartTime(), AnimSection->GetAnimationSequenceLength() / AnimDilationFactor);
+					AnimPosition *= AnimDilationFactor;
 				}
 
+				// Clamp to end of section
 				AnimPosition = FMath::Clamp(AnimPosition, AnimPosition, AnimSection->GetEndTime());
 
 				if (GIsEditor && !GWorld->HasBegunPlay())
