Sequencer: Continued to implement slomo track

diff --git a/Engine/Source/Editor/MovieSceneTools/Private/TrackEditors/SlomoTrackEditor.cpp b/Engine/Source/Editor/MovieSceneTools/Private/TrackEditors/SlomoTrackEditor.cpp
--- a/Engine/Source/Editor/MovieSceneTools/Private/TrackEditors/SlomoTrackEditor.cpp
+++ b/Engine/Source/Editor/MovieSceneTools/Private/TrackEditors/SlomoTrackEditor.cpp
@@ -1,6 +1,8 @@
 // Copyright 1998-2015 Epic Games, Inc. All Rights Reserved.
 
 #include "MovieSceneToolsPrivatePCH.h"
+#include "MovieSceneSlomoTrack.h"
+#include "MovieSceneTrack.h"
 #include "SlomoTrackEditor.h"
 
 
@@ -51,7 +53,18 @@ bool FSlomoTrackEditor::SupportsType(TSubclassOf<UMovieSceneTrack> Type) const
 
 void FSlomoTrackEditor::HandleAddSlomoTrackMenuEntryExecute()
 {
-
+	if (SlomoTrack.IsValid())
+	{
+		return;
+	}
+
+	UMovieSceneSequence* FocusedSequence = GetSequencer()->GetFocusedMovieSceneSequence();
+	UMovieScene* MovieScene = FocusedSequence->GetMovieScene();
+
+	if (MovieScene != nullptr)
+	{
+		SlomoTrack = MovieScene->AddMasterTrack(UMovieSceneSlomoTrack::StaticClass());
+	}
 }
 
 
diff --git a/Engine/Source/Editor/MovieSceneTools/Private/TrackEditors/SlomoTrackEditor.h b/Engine/Source/Editor/MovieSceneTools/Private/TrackEditors/SlomoTrackEditor.h
--- a/Engine/Source/Editor/MovieSceneTools/Private/TrackEditors/SlomoTrackEditor.h
+++ b/Engine/Source/Editor/MovieSceneTools/Private/TrackEditors/SlomoTrackEditor.h
@@ -2,11 +2,11 @@
 
 #pragma once
 
-#include "MovieSceneSlomoTrack.h"
 #include "FloatPropertyTrackEditor.h"
 
 
 class ISequencer;
+class UMovieSceneTrack;
 
 
 /**
@@ -48,4 +48,9 @@ private:
 
 	/** Callback for executing the "Add Event Track" menu entry. */
 	void HandleAddSlomoTrackMenuEntryExecute();
+
+private:
+
+	/** Pointer to the sequence's slomo track, if it exists. */
+	TWeakObjectPtr<UMovieSceneTrack> SlomoTrack;
 };
diff --git a/Engine/Source/Runtime/MovieSceneTracks/Private/TrackInstances/MovieSceneSlomoTrackInstance.cpp b/Engine/Source/Runtime/MovieSceneTracks/Private/TrackInstances/MovieSceneSlomoTrackInstance.cpp
--- a/Engine/Source/Runtime/MovieSceneTracks/Private/TrackInstances/MovieSceneSlomoTrackInstance.cpp
+++ b/Engine/Source/Runtime/MovieSceneTracks/Private/TrackInstances/MovieSceneSlomoTrackInstance.cpp
@@ -3,3 +3,49 @@
 #include "MovieSceneTracksPrivatePCH.h"
 #include "MovieSceneSlomoTrack.h"
 #include "MovieSceneSlomoTrackInstance.h"
+
+
+/* IMovieSceneTrackInstance interface
+ *****************************************************************************/
+
+void FMovieSceneSlomoTrackInstance::Update(float Position, float LastPosition, const TArray<UObject*>& RuntimeObjects, IMovieScenePlayer& Player)
+{
+	if (!ShouldBeApplied())
+	{
+		return;
+	}
+
+	AWorldSettings* WorldSettings = GWorld->GetWorldSettings();
+
+	if (WorldSettings == nullptr)
+	{
+		return;
+	}
+
+	float FloatValue = 0.0f;
+
+	if (SlomoTrack->Eval(Position, LastPosition, FloatValue))
+	{
+		WorldSettings->MatineeTimeDilation = FloatValue;
+		WorldSettings->ForceNetUpdate();
+	}
+}
+
+
+/* IMovieSceneTrackInstance implementation
+ *****************************************************************************/
+
+bool FMovieSceneSlomoTrackInstance::ShouldBeApplied() const
+{
+	if (GIsEditor)
+	{
+		return true;
+	}
+
+	if (GWorld->GetNetMode() == NM_Client)
+	{
+		return false;
+	}
+
+	return (GEngine != nullptr);
+}
diff --git a/Engine/Source/Runtime/MovieSceneTracks/Public/TrackInstances/MovieSceneSlomoTrackInstance.h b/Engine/Source/Runtime/MovieSceneTracks/Public/TrackInstances/MovieSceneSlomoTrackInstance.h
--- a/Engine/Source/Runtime/MovieSceneTracks/Public/TrackInstances/MovieSceneSlomoTrackInstance.h
+++ b/Engine/Source/Runtime/MovieSceneTracks/Public/TrackInstances/MovieSceneSlomoTrackInstance.h
@@ -34,7 +34,15 @@ public:
 	virtual void RefreshInstance(const TArray<UObject*>& RuntimeObjects, IMovieScenePlayer& Player) override { }
 	virtual void RestoreState(const TArray<UObject*>& RuntimeObjects) override { }
 	virtual void SaveState(const TArray<UObject*>& RuntimeObjects) override { }
-//	virtual void Update(float Position, float LastPosition, const TArray<UObject*>& RuntimeObjects, IMovieScenePlayer& Player) override;
+	virtual void Update(float Position, float LastPosition, const TArray<UObject*>& RuntimeObjects, IMovieScenePlayer& Player) override;
+
+protected:
+
+	/** Get the world settings object. */
+	AWorldSettings* GetWorld() const;
+
+	/** Check whether the slow motion should be applied. */
+	bool ShouldBeApplied() const;
 
 private:
 
