Sequencer: Overriding slomo track visibility

diff --git a/Engine/Source/Runtime/MovieSceneTracks/Private/Tracks/MovieSceneSlomoTrack.cpp b/Engine/Source/Runtime/MovieSceneTracks/Private/Tracks/MovieSceneSlomoTrack.cpp
--- a/Engine/Source/Runtime/MovieSceneTracks/Private/Tracks/MovieSceneSlomoTrack.cpp
+++ b/Engine/Source/Runtime/MovieSceneTracks/Private/Tracks/MovieSceneSlomoTrack.cpp
@@ -5,3 +5,11 @@
 #include "MovieSceneSlomoTrack.h"
 #include "MovieSceneSlomoTrackInstance.h"
 
+
+/* UMovieSceneFloatTrack overrides
+ *****************************************************************************/
+
+bool UMovieSceneSlomoTrack::HasShowableData() const
+{
+	return true;
+}
diff --git a/Engine/Source/Runtime/MovieSceneTracks/Public/Tracks/MovieSceneSlomoTrack.h b/Engine/Source/Runtime/MovieSceneTracks/Public/Tracks/MovieSceneSlomoTrack.h
--- a/Engine/Source/Runtime/MovieSceneTracks/Public/Tracks/MovieSceneSlomoTrack.h
+++ b/Engine/Source/Runtime/MovieSceneTracks/Public/Tracks/MovieSceneSlomoTrack.h
@@ -14,4 +14,10 @@ class UMovieSceneSlomoTrack
 	: public UMovieSceneFloatTrack
 {
 	GENERATED_BODY()
+
+public:
+
+	// UMovieSceneFloatTrack overrides
+
+	virtual bool HasShowableData() const override;
 };
