Sequencer: using IsVisibleWhenEmpty instead

diff --git a/Engine/Source/Runtime/MovieSceneTracks/Private/Tracks/MovieSceneSlomoTrack.cpp b/Engine/Source/Runtime/MovieSceneTracks/Private/Tracks/MovieSceneSlomoTrack.cpp
--- a/Engine/Source/Runtime/MovieSceneTracks/Private/Tracks/MovieSceneSlomoTrack.cpp
+++ b/Engine/Source/Runtime/MovieSceneTracks/Private/Tracks/MovieSceneSlomoTrack.cpp
@@ -9,7 +9,7 @@
 /* UMovieSceneFloatTrack overrides
  *****************************************************************************/
 
-bool UMovieSceneSlomoTrack::HasShowableData() const
+bool UMovieSceneSlomoTrack::IsVisibleWhenEmpty() const
 {
 	return true;
 }
diff --git a/Engine/Source/Runtime/MovieSceneTracks/Public/Tracks/MovieSceneSlomoTrack.h b/Engine/Source/Runtime/MovieSceneTracks/Public/Tracks/MovieSceneSlomoTrack.h
--- a/Engine/Source/Runtime/MovieSceneTracks/Public/Tracks/MovieSceneSlomoTrack.h
+++ b/Engine/Source/Runtime/MovieSceneTracks/Public/Tracks/MovieSceneSlomoTrack.h
@@ -19,5 +19,5 @@ public:
 
 	// UMovieSceneFloatTrack overrides
 
-	virtual bool HasShowableData() const override;
+	virtual bool IsVisibleWhenEmpty() const override;
 };
