Sequencer: Added default sections for event and slomo tracks

diff --git a/Engine/Source/Editor/MovieSceneTools/Private/Sections/EventTrackSection.cpp b/Engine/Source/Editor/MovieSceneTools/Private/Sections/EventTrackSection.cpp
--- a/Engine/Source/Editor/MovieSceneTools/Private/Sections/EventTrackSection.cpp
+++ b/Engine/Source/Editor/MovieSceneTools/Private/Sections/EventTrackSection.cpp
@@ -12,7 +12,7 @@
  *****************************************************************************/
 
 FEventTrackSection::FEventTrackSection(UMovieSceneSection& InSection, TSharedPtr<ISequencer> InSequencer)
-	: Section(nullptr/*CastChecked<UMovieSceneEventSection>(&InSection)*/)
+	: Section(&InSection)
 	, Sequencer(InSequencer)
 { }
 
diff --git a/Engine/Source/Editor/MovieSceneTools/Private/Sections/EventTrackSection.h b/Engine/Source/Editor/MovieSceneTools/Private/Sections/EventTrackSection.h
--- a/Engine/Source/Editor/MovieSceneTools/Private/Sections/EventTrackSection.h
+++ b/Engine/Source/Editor/MovieSceneTools/Private/Sections/EventTrackSection.h
@@ -3,7 +3,7 @@
 #pragma once
 
 
-class UMovieSceneEventSection;
+class UMovieSceneSection;
 
 
 /**
@@ -44,7 +44,7 @@ public:
 private:
 
 	/** The section being visualized. */
-	UMovieSceneEventSection* Section;
+	UMovieSceneSection* Section;
 
 	/** The parent sequencer we are a part of. */
 	TWeakPtr<ISequencer> Sequencer;
diff --git a/Engine/Source/Editor/MovieSceneTools/Private/TrackEditors/EventTrackEditor.cpp b/Engine/Source/Editor/MovieSceneTools/Private/TrackEditors/EventTrackEditor.cpp
--- a/Engine/Source/Editor/MovieSceneTools/Private/TrackEditors/EventTrackEditor.cpp
+++ b/Engine/Source/Editor/MovieSceneTools/Private/TrackEditors/EventTrackEditor.cpp
@@ -87,7 +87,9 @@ void FEventTrackEditor::HandleAddEventTrackMenuEntryExecute()
 
 	if (MovieScene != nullptr)
 	{
-		MovieScene->AddMasterTrack(UMovieSceneEventTrack::StaticClass());
+		UMovieSceneTrack* NewTrack = MovieScene->AddMasterTrack(UMovieSceneEventTrack::StaticClass());
+		NewTrack->AddSection(NewTrack->CreateNewSection());
+
 		GetSequencer()->NotifyMovieSceneDataChanged();
 	}
 }
diff --git a/Engine/Source/Editor/MovieSceneTools/Private/TrackEditors/SlomoTrackEditor.cpp b/Engine/Source/Editor/MovieSceneTools/Private/TrackEditors/SlomoTrackEditor.cpp
--- a/Engine/Source/Editor/MovieSceneTools/Private/TrackEditors/SlomoTrackEditor.cpp
+++ b/Engine/Source/Editor/MovieSceneTools/Private/TrackEditors/SlomoTrackEditor.cpp
@@ -71,6 +71,8 @@ void FSlomoTrackEditor::HandleAddSlomoTrackMenuEntryExecute()
 	if (MovieScene != nullptr)
 	{
 		SlomoTrack = MovieScene->AddMasterTrack(UMovieSceneSlomoTrack::StaticClass());
+		SlomoTrack->AddSection(SlomoTrack->CreateNewSection());
+
 		GetSequencer()->NotifyMovieSceneDataChanged();
 	}
 }
diff --git a/Engine/Source/Runtime/MovieSceneTracks/Private/Sections/MovieSceneEventSection.cpp b/Engine/Source/Runtime/MovieSceneTracks/Private/Sections/MovieSceneEventSection.cpp
--- a/Engine/Source/Runtime/MovieSceneTracks/Private/Sections/MovieSceneEventSection.cpp
+++ b/Engine/Source/Runtime/MovieSceneTracks/Private/Sections/MovieSceneEventSection.cpp
@@ -4,6 +4,15 @@
 #include "MovieSceneEventSection.h"
 
 
+/* UMovieSceneSection structors
+ *****************************************************************************/
+
+UMovieSceneEventSection::UMovieSceneEventSection()
+{
+	SetIsInfinite(true);
+}
+
+
 /* UMovieSceneSection overrides
  *****************************************************************************/
 
diff --git a/Engine/Source/Runtime/MovieSceneTracks/Private/Sections/MovieSceneSlomoSection.cpp b/Engine/Source/Runtime/MovieSceneTracks/Private/Sections/MovieSceneSlomoSection.cpp
--- a/Engine/Source/Runtime/MovieSceneTracks/Private/Sections/MovieSceneSlomoSection.cpp
+++ b/Engine/Source/Runtime/MovieSceneTracks/Private/Sections/MovieSceneSlomoSection.cpp
@@ -2,3 +2,13 @@
 
 #include "MovieSceneTracksPrivatePCH.h"
 #include "MovieSceneSlomoSection.h"
+
+
+/* UMovieSceneSlomoSection structors
+ *****************************************************************************/
+
+UMovieSceneSlomoSection::UMovieSceneSlomoSection()
+	: UMovieSceneFloatSection()
+{
+	SetIsInfinite(true);
+}
diff --git a/Engine/Source/Runtime/MovieSceneTracks/Private/Tracks/MovieSceneEventTrack.cpp b/Engine/Source/Runtime/MovieSceneTracks/Private/Tracks/MovieSceneEventTrack.cpp
--- a/Engine/Source/Runtime/MovieSceneTracks/Private/Tracks/MovieSceneEventTrack.cpp
+++ b/Engine/Source/Runtime/MovieSceneTracks/Private/Tracks/MovieSceneEventTrack.cpp
@@ -82,6 +82,12 @@ void UMovieSceneEventTrack::TriggerEvents(TRange<float> TimeRange, bool Backward
 /* UMovieSceneTrack interface
  *****************************************************************************/
 
+void UMovieSceneEventTrack::AddSection(UMovieSceneSection* Section)
+{
+	Sections.Add(Section);
+}
+
+
 TSharedPtr<IMovieSceneTrackInstance> UMovieSceneEventTrack::CreateInstance()
 {
 	return MakeShareable(new FMovieSceneEventTrackInstance(*this));
diff --git a/Engine/Source/Runtime/MovieSceneTracks/Private/Tracks/MovieSceneSlomoTrack.cpp b/Engine/Source/Runtime/MovieSceneTracks/Private/Tracks/MovieSceneSlomoTrack.cpp
--- a/Engine/Source/Runtime/MovieSceneTracks/Private/Tracks/MovieSceneSlomoTrack.cpp
+++ b/Engine/Source/Runtime/MovieSceneTracks/Private/Tracks/MovieSceneSlomoTrack.cpp
@@ -9,6 +9,12 @@
 /* UMovieSceneEventTrack interface
  *****************************************************************************/
 
+UMovieSceneSection* UMovieSceneSlomoTrack::CreateNewSection()
+{
+	return NewObject<UMovieSceneSection>(this, UMovieSceneSlomoSection::StaticClass(), NAME_None, RF_Transactional);
+}
+
+
 FName UMovieSceneSlomoTrack::GetTrackName() const
 {
 	return FName("Play Rate");
diff --git a/Engine/Source/Runtime/MovieSceneTracks/Public/Sections/MovieSceneEventSection.h b/Engine/Source/Runtime/MovieSceneTracks/Public/Sections/MovieSceneEventSection.h
--- a/Engine/Source/Runtime/MovieSceneTracks/Public/Sections/MovieSceneEventSection.h
+++ b/Engine/Source/Runtime/MovieSceneTracks/Public/Sections/MovieSceneEventSection.h
@@ -50,6 +50,9 @@ class UMovieSceneEventSection
 {
 	GENERATED_BODY()
 
+	/** Default constructor. */
+	UMovieSceneEventSection();
+
 public:
 
 	/**
diff --git a/Engine/Source/Runtime/MovieSceneTracks/Public/Sections/MovieSceneSlomoSection.h b/Engine/Source/Runtime/MovieSceneTracks/Public/Sections/MovieSceneSlomoSection.h
--- a/Engine/Source/Runtime/MovieSceneTracks/Public/Sections/MovieSceneSlomoSection.h
+++ b/Engine/Source/Runtime/MovieSceneTracks/Public/Sections/MovieSceneSlomoSection.h
@@ -15,4 +15,6 @@ class UMovieSceneSlomoSection
 {
 	GENERATED_BODY()
 
+	/** Default constructor. */
+	UMovieSceneSlomoSection();
 };
diff --git a/Engine/Source/Runtime/MovieSceneTracks/Public/Tracks/MovieSceneEventTrack.h b/Engine/Source/Runtime/MovieSceneTracks/Public/Tracks/MovieSceneEventTrack.h
--- a/Engine/Source/Runtime/MovieSceneTracks/Public/Tracks/MovieSceneEventTrack.h
+++ b/Engine/Source/Runtime/MovieSceneTracks/Public/Tracks/MovieSceneEventTrack.h
@@ -47,6 +47,7 @@ public:
 
 	// UMovieSceneTrack interface
 
+	virtual void AddSection(UMovieSceneSection* Section) override;
 	virtual TSharedPtr<IMovieSceneTrackInstance> CreateInstance() override;
 	virtual UMovieSceneSection* CreateNewSection() override;
 	virtual const TArray<UMovieSceneSection*>& GetAllSections() const override;
diff --git a/Engine/Source/Runtime/MovieSceneTracks/Public/Tracks/MovieSceneSlomoTrack.h b/Engine/Source/Runtime/MovieSceneTracks/Public/Tracks/MovieSceneSlomoTrack.h
--- a/Engine/Source/Runtime/MovieSceneTracks/Public/Tracks/MovieSceneSlomoTrack.h
+++ b/Engine/Source/Runtime/MovieSceneTracks/Public/Tracks/MovieSceneSlomoTrack.h
@@ -19,5 +19,6 @@ public:
 
 	// UMovieSceneTrack interface
 
+	virtual UMovieSceneSection* CreateNewSection() override;
 	virtual FName GetTrackName() const override;
 };
