Sequencer: Continued to implement slomo track

diff --git a/Engine/Source/Runtime/MovieSceneTracks/Private/TrackInstances/MovieSceneSlomoTrackInstance.cpp b/Engine/Source/Runtime/MovieSceneTracks/Private/TrackInstances/MovieSceneSlomoTrackInstance.cpp
--- a/Engine/Source/Runtime/MovieSceneTracks/Private/TrackInstances/MovieSceneSlomoTrackInstance.cpp
+++ b/Engine/Source/Runtime/MovieSceneTracks/Private/TrackInstances/MovieSceneSlomoTrackInstance.cpp
@@ -5,6 +5,14 @@
 #include "MovieSceneSlomoTrackInstance.h"
 
 
+/* FMovieSceneSlomoTrackInstance structors
+ *****************************************************************************/
+
+FMovieSceneSlomoTrackInstance::FMovieSceneSlomoTrackInstance(UMovieSceneSlomoTrack& InSlomoTrack)
+	: SlomoTrack(&InSlomoTrack)
+{ }
+
+
 /* IMovieSceneTrackInstance interface
  *****************************************************************************/
 
diff --git a/Engine/Source/Runtime/MovieSceneTracks/Private/Tracks/MovieSceneSlomoTrack.cpp b/Engine/Source/Runtime/MovieSceneTracks/Private/Tracks/MovieSceneSlomoTrack.cpp
--- a/Engine/Source/Runtime/MovieSceneTracks/Private/Tracks/MovieSceneSlomoTrack.cpp
+++ b/Engine/Source/Runtime/MovieSceneTracks/Private/Tracks/MovieSceneSlomoTrack.cpp
@@ -9,6 +9,12 @@
 /* UMovieSceneEventTrack interface
  *****************************************************************************/
 
+TSharedPtr<IMovieSceneTrackInstance> UMovieSceneSlomoTrack::CreateInstance()
+{
+	return MakeShareable(new FMovieSceneSlomoTrackInstance(*this)); 
+}
+
+
 UMovieSceneSection* UMovieSceneSlomoTrack::CreateNewSection()
 {
 	return NewObject<UMovieSceneSection>(this, UMovieSceneSlomoSection::StaticClass(), NAME_None, RF_Transactional);
diff --git a/Engine/Source/Runtime/MovieSceneTracks/Public/Tracks/MovieSceneSlomoTrack.h b/Engine/Source/Runtime/MovieSceneTracks/Public/Tracks/MovieSceneSlomoTrack.h
--- a/Engine/Source/Runtime/MovieSceneTracks/Public/Tracks/MovieSceneSlomoTrack.h
+++ b/Engine/Source/Runtime/MovieSceneTracks/Public/Tracks/MovieSceneSlomoTrack.h
@@ -19,6 +19,7 @@ public:
 
 	// UMovieSceneTrack interface
 
+	virtual TSharedPtr<IMovieSceneTrackInstance> CreateInstance() override;
 	virtual UMovieSceneSection* CreateNewSection() override;
 	virtual FName GetTrackName() const override;
 };
