Sequencer: Exposed bool section curve keys

diff --git a/Engine/Source/Runtime/MovieSceneTracks/Public/Sections/MovieSceneBoolSection.h b/Engine/Source/Runtime/MovieSceneTracks/Public/Sections/MovieSceneBoolSection.h
--- a/Engine/Source/Runtime/MovieSceneTracks/Public/Sections/MovieSceneBoolSection.h
+++ b/Engine/Source/Runtime/MovieSceneTracks/Public/Sections/MovieSceneBoolSection.h
@@ -52,6 +52,6 @@ private:
 	/** Ordered curve data */
 	// @todo Sequencer This could be optimized by packing the bools separately
 	// but that may not be worth the effort
-	UPROPERTY()
+	UPROPERTY(EditAnywhere, Category="Curve")
 	FIntegralCurve BoolCurve;
 };
