Sequencer: Continued to implement event tracks

diff --git a/Engine/Source/Runtime/Engine/Classes/Curves/CurveBase.h b/Engine/Source/Runtime/Engine/Classes/Curves/CurveBase.h
--- a/Engine/Source/Runtime/Engine/Classes/Curves/CurveBase.h
+++ b/Engine/Source/Runtime/Engine/Classes/Curves/CurveBase.h
@@ -260,6 +260,16 @@ public:
 	FNameCurveKey& GetKey(FKeyHandle KeyHandle);
 	FNameCurveKey GetKey(FKeyHandle KeyHandle) const;
 
+	/**
+	 * Read-only access to the key collection.
+	 *
+	 * @return Collection of keys.
+	 */
+	const TArray<FNameCurveKey>& GetKeys() const
+	{
+		return Keys;
+	}
+
 	/**
 	 * Get the time for the Key with the specified index.
 	 *
diff --git a/Engine/Source/Runtime/MovieSceneTracks/Private/Sections/MovieSceneEventSection.cpp b/Engine/Source/Runtime/MovieSceneTracks/Private/Sections/MovieSceneEventSection.cpp
--- a/Engine/Source/Runtime/MovieSceneTracks/Private/Sections/MovieSceneEventSection.cpp
+++ b/Engine/Source/Runtime/MovieSceneTracks/Private/Sections/MovieSceneEventSection.cpp
@@ -23,6 +23,37 @@ void UMovieSceneEventSection::AddKey(float Time, const FName& EventName, FKeyPar
 }
 
 
+void UMovieSceneEventSection::TriggerEvents(ALevelScriptActor* LevelScriptActor, float Position, float LastPosition)
+{
+	const TArray<FNameCurveKey>& Keys = Events.GetKeys();
+
+	if (Position >= LastPosition)
+	{
+		for (const auto& Key : Keys)
+		{
+			if ((Key.Time > LastPosition) && (Key.Time <= Position))
+			{
+				TriggerEvent(Key.Value, LevelScriptActor);
+			}
+		}
+	}
+	else
+	{
+		for (int32 KeyIndex = Keys.Num() - 1; KeyIndex >= 0; --KeyIndex)
+		{
+			const auto& Key = Keys[KeyIndex];
+
+			if ((Key.Time >= Position) && (Key.Time < LastPosition))
+			{
+				TriggerEvent(Key.Value, LevelScriptActor);
+			}
+
+			
+		}
+	}
+}
+
+
 /* UMovieSceneSection overrides
  *****************************************************************************/
 
@@ -53,3 +84,27 @@ void UMovieSceneEventSection::MoveSection(float DeltaPosition, TSet<FKeyHandle>&
 
 	Events.ShiftCurve(DeltaPosition, KeyHandles);
 }
+
+
+/* UMovieSceneSection overrides
+ *****************************************************************************/
+
+void UMovieSceneEventSection::TriggerEvent(const FName& Event, ALevelScriptActor* LevelScriptActor)
+{
+	UFunction* EventFunction = LevelScriptActor->FindFunction(Event);
+
+	if (EventFunction == nullptr)
+	{
+		// @todo sequencer: gmp: add external log category for MovieScene
+		//UE_LOG(LogMovieScene, Log, TEXT("UMovieSceneEventSection::TriggerEvent: Unable to find function '%s'"), *Event.ToString());
+	}
+	else if (EventFunction->NumParms != 0)
+	{
+		// @todo sequencer: gmp: add external log category for MovieScene
+		//UE_LOG(LogMovieScene, Log, TEXT("UMovieSceneEventSection::TriggerEvent: Function '%s' does not have zero parameters."), *Event.ToString());
+	}
+	else
+	{
+		LevelScriptActor->ProcessEvent(EventFunction, nullptr);
+	}
+}
diff --git a/Engine/Source/Runtime/MovieSceneTracks/Private/TrackInstances/MovieSceneEventTrackInstance.cpp b/Engine/Source/Runtime/MovieSceneTracks/Private/TrackInstances/MovieSceneEventTrackInstance.cpp
--- a/Engine/Source/Runtime/MovieSceneTracks/Private/TrackInstances/MovieSceneEventTrackInstance.cpp
+++ b/Engine/Source/Runtime/MovieSceneTracks/Private/TrackInstances/MovieSceneEventTrackInstance.cpp
@@ -19,12 +19,5 @@ FMovieSceneEventTrackInstance::FMovieSceneEventTrackInstance(UMovieSceneEventTra
 
 void FMovieSceneEventTrackInstance::Update(float Position, float LastPosition, const TArray<UObject*>& RuntimeObjects, IMovieScenePlayer& Player)
 {
-	if (Position > LastPosition)
-	{
-		EventTrack->TriggerEvents(TRange<float>(LastPosition, Position), false);	
-	}
-	else if (Position < LastPosition)
-	{
-		EventTrack->TriggerEvents(TRange<float>(Position, LastPosition), true);
-	}
+	EventTrack->TriggerEvents(Position, LastPosition);
 }
diff --git a/Engine/Source/Runtime/MovieSceneTracks/Private/Tracks/MovieSceneEventTrack.cpp b/Engine/Source/Runtime/MovieSceneTracks/Private/Tracks/MovieSceneEventTrack.cpp
--- a/Engine/Source/Runtime/MovieSceneTracks/Private/Tracks/MovieSceneEventTrack.cpp
+++ b/Engine/Source/Runtime/MovieSceneTracks/Private/Tracks/MovieSceneEventTrack.cpp
@@ -35,47 +35,64 @@ bool UMovieSceneEventTrack::AddKeyToSection(float Time, FName EventName, FKeyPar
 }
 
 
-void UMovieSceneEventTrack::TriggerEvents(TRange<float> TimeRange, bool Backwards)
+void UMovieSceneEventTrack::TriggerEvents(float Position, float LastPosition)
 {
+	if ((Sections.Num() == 0) || (Position == LastPosition))
+	{
+		return;
+	}
+
+	bool Backwards = Position < LastPosition;
+
 	if ((!Backwards && !bFireEventsWhenForwards) ||
 		(Backwards && !bFireEventsWhenBackwards))
 	{
 		return;
 	}
 
-	// @todo sequencer: gmp: implement event track triggering
-	/*
-	ULevel* Level = GetLevel();
-	ALevelScriptActor* LevelScriptActor = Level->LevelScriptActor;
+	// get the level script actor
+	if ((GEngine == nullptr) || (GEngine->GameViewport == nullptr))
+	{
+		return;
+	}
+
+	UWorld* World = GEngine->GameViewport->GetWorld();
+
+	if (World == nullptr)
+	{
+		return;
+	}
+
+	ALevelScriptActor* LevelScriptActor = World->GetLevelScriptActor();
 
 	if (LevelScriptActor == nullptr)
 	{
 		return;
 	}
 
-	// for each key...
-	for ()
+	// get the traversed sections
+	TArray<UMovieSceneSection*> TraversedSections = MovieSceneHelpers::GetTraversedSections(Sections, Position, LastPosition);
+	int32 EndIndex = TraversedSections.Num() - 1;
+
+	if (EndIndex < 0)
+	{
+		return;
+	}
+
+	int32 StartIndex = 0;
+	int32 Increment = 1;
+
+	if (Backwards)
+	{
+		Swap(StartIndex, EndIndex);
+		Increment = -1;
+	}
+
+	// traverse sections
+	for (int32 SectionIndex = StartIndex; SectionIndex <= EndIndex; SectionIndex += Increment)
 	{
-		UFunction* EventFunction = LevelScriptActor->FindFunction(EventName);
-
-		if (EventFunction != nullptr)
-		{
-			if(EventFunction->NumParms == 0)
-			{
-				LevelScriptActor->ProcessEvent(EventFunction, nullptr);
-			}
-			else
-			{
-				// @todo sequencer: gmp: add external log category for MovieScene
-				//UE_LOG(LogMovieScene, Log, TEXT("NotifyEventTriggered: Function '%s' does not have zero parameters."), *EventFuncName.ToString());
-			}
-		}
-		else
-		{
-			// @todo sequencer: gmp: add external log category for MovieScene
-			//UE_LOG(LogMovieScene, Log, TEXT("NotifyEventTriggered: Unable to find function '%s'"), *EventFuncName.ToString());
-		}
-	}*/
+		CastChecked<UMovieSceneEventSection>(TraversedSections[SectionIndex])->TriggerEvents(LevelScriptActor, Position, LastPosition);
+	}
 }
 
 
diff --git a/Engine/Source/Runtime/MovieSceneTracks/Public/Sections/MovieSceneEventSection.h b/Engine/Source/Runtime/MovieSceneTracks/Public/Sections/MovieSceneEventSection.h
--- a/Engine/Source/Runtime/MovieSceneTracks/Public/Sections/MovieSceneEventSection.h
+++ b/Engine/Source/Runtime/MovieSceneTracks/Public/Sections/MovieSceneEventSection.h
@@ -5,6 +5,9 @@
 #include "MovieSceneEventSection.generated.h"
 
 
+class ALevelScriptActor;
+
+
 /**
  * Structure for event section keys.
  */
@@ -72,6 +75,15 @@ public:
 		return Events;
 	}
 
+	/**
+	 * Trigger the events that fall into the given time range.
+	 *
+	 * @param LevelScriptActor The script actor to trigger the events on.
+	 * @param Position The current position in time.
+	 * @param LastPosition The time at the last update.
+	 */
+	void TriggerEvents(ALevelScriptActor* LevelScriptActor, float Position, float LastPosition);
+
 public:
 
 	// UMovieSceneSection interface
@@ -80,6 +92,16 @@ public:
 	virtual void GetKeyHandles(TSet<FKeyHandle>& KeyHandles) const override;
 	virtual void MoveSection(float DeltaPosition, TSet<FKeyHandle>& KeyHandles) override;
 
+protected:
+
+	/**
+	 * Trigger event for the specified key.
+	 *
+	 * @param Key The key to trigger.
+	 * @param LevelScriptActor The script actor to trigger the events on.
+	 */
+	void TriggerEvent(const FName& Event, ALevelScriptActor* LevelScriptActor);
+
 private:
 
 	/** The section's keys. */
diff --git a/Engine/Source/Runtime/MovieSceneTracks/Public/Tracks/MovieSceneEventTrack.h b/Engine/Source/Runtime/MovieSceneTracks/Public/Tracks/MovieSceneEventTrack.h
--- a/Engine/Source/Runtime/MovieSceneTracks/Public/Tracks/MovieSceneEventTrack.h
+++ b/Engine/Source/Runtime/MovieSceneTracks/Public/Tracks/MovieSceneEventTrack.h
@@ -19,7 +19,9 @@ public:
 
 	/** Default constructor. */
 	UMovieSceneEventTrack()
-		: TrackName("Events")
+		: bFireEventsWhenForwards(true)
+		, bFireEventsWhenBackwards(false)
+		, TrackName("Events")
 	{ }
 
 public:
@@ -38,10 +40,10 @@ public:
 	/**
 	 * Trigger the events that fall into the given time range.
 	 *
-	 * @param TimeRange The time range to trigger events in.
-	 * @param Backwards Whether the events should be triggered in reverse order.
+	 * @param Position The current position in time.
+	 * @param LastPosition The time at the last update.
 	 */
-	void TriggerEvents(TRange<float> TimeRange, bool Backwards);
+	void TriggerEvents(float Position, float LastPosition);
 
 public:
 
@@ -68,10 +70,6 @@ private:
 	UPROPERTY(EditAnywhere, Category=TrackEvent)
 	uint32 bFireEventsWhenBackwards:1;
 
-	/** If checked each key's event name is the exact name of the custom event function in level script that will be called */
-	UPROPERTY(EditAnywhere, Category=TrackEvent)
-	uint32 bUseCustomEventName:1;
-
 	/** The track's sections. */
 	UPROPERTY()
 	TArray<UMovieSceneSection*> Sections;
