Sequencer: Set default value of the slomo section to 1.0f.

diff --git a/Engine/Source/Runtime/MovieSceneTracks/Private/Sections/MovieSceneSlomoSection.cpp b/Engine/Source/Runtime/MovieSceneTracks/Private/Sections/MovieSceneSlomoSection.cpp
--- a/Engine/Source/Runtime/MovieSceneTracks/Private/Sections/MovieSceneSlomoSection.cpp
+++ b/Engine/Source/Runtime/MovieSceneTracks/Private/Sections/MovieSceneSlomoSection.cpp
@@ -11,4 +11,5 @@ UMovieSceneSlomoSection::UMovieSceneSlomoSection()
 	: UMovieSceneFloatSection()
 {
 	SetIsInfinite(true);
+	GetFloatCurve().SetDefaultValue(1.0f);
 }
