Sequencer: Fix sound duration for audio cues.

diff --git a/Engine/Source/Runtime/MovieSceneTracks/Private/Tracks/MovieSceneAudioTrack.cpp b/Engine/Source/Runtime/MovieSceneTracks/Private/Tracks/MovieSceneAudioTrack.cpp
--- a/Engine/Source/Runtime/MovieSceneTracks/Private/Tracks/MovieSceneAudioTrack.cpp
+++ b/Engine/Source/Runtime/MovieSceneTracks/Private/Tracks/MovieSceneAudioTrack.cpp
@@ -5,6 +5,7 @@
 #include "MovieSceneAudioTrack.h"
 #include "IMovieScenePlayer.h"
 #include "SoundDefinitions.h"
+#include "Sound/SoundNodeWavePlayer.h"
 #include "Runtime/Engine/Classes/Kismet/GameplayStatics.h"
 #include "Runtime/Engine/Public/AudioDecompress.h"
 #include "MovieSceneAudioTrackInstance.h"
@@ -76,6 +77,33 @@ TRange<float> UMovieSceneAudioTrack::GetSectionBoundaries() const
 	return TRange<float>::Hull(Bounds);
 }
 
+float GetSoundDuration(USoundBase* Sound)
+{
+	USoundWave* SoundWave = nullptr;
+
+	if (Sound->IsA<USoundWave>())
+	{
+		SoundWave = Cast<USoundWave>(Sound);
+	}
+	else if (Sound->IsA<USoundCue>())
+	{
+		USoundCue* SoundCue = Cast<USoundCue>(Sound);
+
+		// @todo Sequencer - Right now for sound cues, we just use the first sound wave in the cue
+		// In the future, it would be better to properly generate the sound cue's data after forcing determinism
+		const TArray<USoundNode*>& AllNodes = SoundCue->AllNodes;
+		for (int32 i = 0; i < AllNodes.Num() && SoundWave == nullptr; ++i)
+		{
+			if (AllNodes[i]->IsA<USoundNodeWavePlayer>())
+			{
+				SoundWave = Cast<USoundNodeWavePlayer>(AllNodes[i])->GetSoundWave();
+			}
+		}
+	}
+
+	const float Duration = (SoundWave ? SoundWave->GetDuration() : 0.f);
+	return Duration == INDEFINITELY_LOOPING_DURATION ? SoundWave->Duration : Duration;
+}
 
 void UMovieSceneAudioTrack::AddNewSound(USoundBase* Sound, float Time)
 {
@@ -85,7 +113,7 @@ void UMovieSceneAudioTrack::AddNewSound(USoundBase* Sound, float Time)
 	// @todo Once we have infinite sections, we can remove this
 	float DurationToUse = 1.f; // if all else fails, use 1 second duration
 
-	float SoundDuration = Sound->GetDuration();
+	float SoundDuration = GetSoundDuration(Sound);
 	if (SoundDuration != INDEFINITELY_LOOPING_DURATION)
 	{
 		DurationToUse = SoundDuration;
