Sequencer: Compile fix, include SoundCue.h

diff --git a/Engine/Source/Runtime/MovieSceneTracks/Private/Tracks/MovieSceneAudioTrack.cpp b/Engine/Source/Runtime/MovieSceneTracks/Private/Tracks/MovieSceneAudioTrack.cpp
--- a/Engine/Source/Runtime/MovieSceneTracks/Private/Tracks/MovieSceneAudioTrack.cpp
+++ b/Engine/Source/Runtime/MovieSceneTracks/Private/Tracks/MovieSceneAudioTrack.cpp
@@ -6,6 +6,7 @@
 #include "IMovieScenePlayer.h"
 #include "SoundDefinitions.h"
 #include "Sound/SoundNodeWavePlayer.h"
+#include "Sound/SoundCue.h"
 #include "Runtime/Engine/Classes/Kismet/GameplayStatics.h"
 #include "Runtime/Engine/Public/AudioDecompress.h"
 #include "MovieSceneAudioTrackInstance.h"
