Sequencer: Compile fix - use WITH_EDITORONLY_DATA for SoundCue AllNodes.

diff --git a/Engine/Source/Runtime/MovieSceneTracks/Private/Tracks/MovieSceneAudioTrack.cpp b/Engine/Source/Runtime/MovieSceneTracks/Private/Tracks/MovieSceneAudioTrack.cpp
--- a/Engine/Source/Runtime/MovieSceneTracks/Private/Tracks/MovieSceneAudioTrack.cpp
+++ b/Engine/Source/Runtime/MovieSceneTracks/Private/Tracks/MovieSceneAudioTrack.cpp
@@ -88,6 +88,7 @@ float GetSoundDuration(USoundBase* Sound)
 	}
 	else if (Sound->IsA<USoundCue>())
 	{
+#if WITH_EDITORONLY_DATA
 		USoundCue* SoundCue = Cast<USoundCue>(Sound);
 
 		// @todo Sequencer - Right now for sound cues, we just use the first sound wave in the cue
@@ -100,6 +101,7 @@ float GetSoundDuration(USoundBase* Sound)
 				SoundWave = Cast<USoundNodeWavePlayer>(AllNodes[i])->GetSoundWave();
 			}
 		}
+#endif
 	}
 
 	const float Duration = (SoundWave ? SoundWave->GetDuration() : 0.f);
