Sequencer: Forgot file.

diff --git a/Engine/Source/Runtime/MovieSceneTracks/Public/Tracks/MovieSceneFadeTrack.h b/Engine/Source/Runtime/MovieSceneTracks/Public/Tracks/MovieSceneFadeTrack.h
new file mode 100644
--- /dev/null
+++ b/Engine/Source/Runtime/MovieSceneTracks/Public/Tracks/MovieSceneFadeTrack.h
@@ -0,0 +1,25 @@
+// Copyright 1998-2015 Epic Games, Inc. All Rights Reserved.
+
+#pragma once
+
+#include "MovieSceneFloatTrack.h"
+#include "MovieSceneFadeTrack.generated.h"
+
+
+/**
+ * Implements a movie scene track that controls a fade.
+ */
+UCLASS(MinimalAPI)
+class UMovieSceneFadeTrack
+	: public UMovieSceneFloatTrack
+{
+	GENERATED_BODY()
+
+public:
+
+	// UMovieSceneTrack interface
+
+	virtual TSharedPtr<IMovieSceneTrackInstance> CreateInstance() override;
+	virtual UMovieSceneSection* CreateNewSection() override;
+	virtual FName GetTrackName() const override;
+};
