Sequencer: Fix infinite section eval for pre and post infinity. Don't clamp time.

diff --git a/Engine/Source/Runtime/MovieSceneTracks/Private/Tracks/MovieScene3DTransformTrack.cpp b/Engine/Source/Runtime/MovieSceneTracks/Private/Tracks/MovieScene3DTransformTrack.cpp
--- a/Engine/Source/Runtime/MovieSceneTracks/Private/Tracks/MovieScene3DTransformTrack.cpp
+++ b/Engine/Source/Runtime/MovieSceneTracks/Private/Tracks/MovieScene3DTransformTrack.cpp
@@ -114,7 +114,10 @@ bool UMovieScene3DTransformTrack::Eval( float Position, float LastPosition, FVec
 	{
 		const UMovieScene3DTransformSection* TransformSection = CastChecked<UMovieScene3DTransformSection>( Section );
 
-		Position = FMath::Clamp(Position, Section->GetStartTime(), Section->GetEndTime());
+		if (!Section->IsInfinite())
+		{
+			Position = FMath::Clamp(Position, Section->GetStartTime(), Section->GetEndTime());
+		}
 
 		// Evaluate translation,rotation, and scale curves.  If no keys were found on one of these, that component of the transform will remain unchained
 		TransformSection->EvalTranslation( Position, OutTranslation );
diff --git a/Engine/Source/Runtime/MovieSceneTracks/Private/Tracks/MovieSceneBoolTrack.cpp b/Engine/Source/Runtime/MovieSceneTracks/Private/Tracks/MovieSceneBoolTrack.cpp
--- a/Engine/Source/Runtime/MovieSceneTracks/Private/Tracks/MovieSceneBoolTrack.cpp
+++ b/Engine/Source/Runtime/MovieSceneTracks/Private/Tracks/MovieSceneBoolTrack.cpp
@@ -44,7 +44,10 @@ bool UMovieSceneBoolTrack::Eval( float Position, float LastPostion, bool& OutBoo
 
 	if( Section )
 	{
-		Position = FMath::Clamp(Position, Section->GetStartTime(), Section->GetEndTime());
+		if (!Section->IsInfinite())
+		{
+			Position = FMath::Clamp(Position, Section->GetStartTime(), Section->GetEndTime());
+		}
 
 		OutBool = CastChecked<UMovieSceneBoolSection>( Section )->Eval( Position );
 	}
diff --git a/Engine/Source/Runtime/MovieSceneTracks/Private/Tracks/MovieSceneByteTrack.cpp b/Engine/Source/Runtime/MovieSceneTracks/Private/Tracks/MovieSceneByteTrack.cpp
--- a/Engine/Source/Runtime/MovieSceneTracks/Private/Tracks/MovieSceneByteTrack.cpp
+++ b/Engine/Source/Runtime/MovieSceneTracks/Private/Tracks/MovieSceneByteTrack.cpp
@@ -47,7 +47,10 @@ bool UMovieSceneByteTrack::Eval( float Position, float LastPostion, uint8& OutBy
 
 	if( Section )
 	{
-		Position = FMath::Clamp(Position, Section->GetStartTime(), Section->GetEndTime());
+		if (!Section->IsInfinite())
+		{
+			Position = FMath::Clamp(Position, Section->GetStartTime(), Section->GetEndTime());
+		}
 
 		OutByte = CastChecked<UMovieSceneByteSection>( Section )->Eval( Position );
 	}
diff --git a/Engine/Source/Runtime/MovieSceneTracks/Private/Tracks/MovieSceneColorTrack.cpp b/Engine/Source/Runtime/MovieSceneTracks/Private/Tracks/MovieSceneColorTrack.cpp
--- a/Engine/Source/Runtime/MovieSceneTracks/Private/Tracks/MovieSceneColorTrack.cpp
+++ b/Engine/Source/Runtime/MovieSceneTracks/Private/Tracks/MovieSceneColorTrack.cpp
@@ -47,7 +47,11 @@ bool UMovieSceneColorTrack::Eval( float Position, float LastPosition, FLinearCol
 
 	if( Section )
 	{
-		Position = FMath::Clamp(Position, Section->GetStartTime(), Section->GetEndTime());
+		if (!Section->IsInfinite())
+		{
+			Position = FMath::Clamp(Position, Section->GetStartTime(), Section->GetEndTime());
+		}
+
 		OutColor = CastChecked<UMovieSceneColorSection>( Section )->Eval( Position, OutColor );
 	}
 
diff --git a/Engine/Source/Runtime/MovieSceneTracks/Private/Tracks/MovieSceneFloatTrack.cpp b/Engine/Source/Runtime/MovieSceneTracks/Private/Tracks/MovieSceneFloatTrack.cpp
--- a/Engine/Source/Runtime/MovieSceneTracks/Private/Tracks/MovieSceneFloatTrack.cpp
+++ b/Engine/Source/Runtime/MovieSceneTracks/Private/Tracks/MovieSceneFloatTrack.cpp
@@ -47,7 +47,10 @@ bool UMovieSceneFloatTrack::Eval( float Position, float LastPosition, float& Out
 
 	if( Section )
 	{
-		Position = FMath::Clamp(Position, Section->GetStartTime(), Section->GetEndTime());
+		if (!Section->IsInfinite())
+		{
+			Position = FMath::Clamp(Position, Section->GetStartTime(), Section->GetEndTime());
+		}
 
 		OutFloat = CastChecked<UMovieSceneFloatSection>( Section )->Eval( Position );
 	}
diff --git a/Engine/Source/Runtime/MovieSceneTracks/Private/Tracks/MovieSceneVectorTrack.cpp b/Engine/Source/Runtime/MovieSceneTracks/Private/Tracks/MovieSceneVectorTrack.cpp
--- a/Engine/Source/Runtime/MovieSceneTracks/Private/Tracks/MovieSceneVectorTrack.cpp
+++ b/Engine/Source/Runtime/MovieSceneTracks/Private/Tracks/MovieSceneVectorTrack.cpp
@@ -91,7 +91,10 @@ bool UMovieSceneVectorTrack::Eval( float Position, float LastPosition, FVector4&
 
 	if( Section )
 	{
-		Position = FMath::Clamp(Position, Section->GetStartTime(), Section->GetEndTime());
+		if (!Section->IsInfinite())
+		{
+			Position = FMath::Clamp(Position, Section->GetStartTime(), Section->GetEndTime());
+		}
 
 		InOutVector = CastChecked<UMovieSceneVectorSection>( Section )->Eval( Position, InOutVector );
 	}
diff --git a/Engine/Source/Runtime/UMG/Private/Animation/MovieScene2DTransformTrack.cpp b/Engine/Source/Runtime/UMG/Private/Animation/MovieScene2DTransformTrack.cpp
--- a/Engine/Source/Runtime/UMG/Private/Animation/MovieScene2DTransformTrack.cpp
+++ b/Engine/Source/Runtime/UMG/Private/Animation/MovieScene2DTransformTrack.cpp
@@ -30,7 +30,10 @@ bool UMovieScene2DTransformTrack::Eval(float Position, float LastPosition, FWidg
 	{
 		const UMovieScene2DTransformSection* TransformSection = CastChecked<UMovieScene2DTransformSection>(Section);
 
-		Position = FMath::Clamp(Position, Section->GetStartTime(), Section->GetEndTime());
+		if (!Section->IsInfinite())
+		{
+			Position = FMath::Clamp(Position, Section->GetStartTime(), Section->GetEndTime());
+		}
 
 		InOutTransform = TransformSection->Eval(Position, InOutTransform);
 	}
diff --git a/Engine/Source/Runtime/UMG/Private/Animation/MovieSceneMarginTrack.cpp b/Engine/Source/Runtime/UMG/Private/Animation/MovieSceneMarginTrack.cpp
--- a/Engine/Source/Runtime/UMG/Private/Animation/MovieSceneMarginTrack.cpp
+++ b/Engine/Source/Runtime/UMG/Private/Animation/MovieSceneMarginTrack.cpp
@@ -48,7 +48,10 @@ bool UMovieSceneMarginTrack::Eval( float Position, float LastPosition, FMargin&
 	{
 		const UMovieSceneMarginSection* MarginSection = CastChecked<UMovieSceneMarginSection>( Section );
 
-		Position = FMath::Clamp(Position, Section->GetStartTime(), Section->GetEndTime());
+		if (!Section->IsInfinite())
+		{
+			Position = FMath::Clamp(Position, Section->GetStartTime(), Section->GetEndTime());
+		}
 
 		InOutMargin = MarginSection->Eval( Position, InOutMargin );
 	}
