Mergine changes from //depot/UE4 to //UE4/Main

diff --git a/Engine/Source/Runtime/MovieScene/Public/MovieSceneSequenceInstance.h b/Engine/Source/Runtime/MovieScene/Public/MovieSceneSequenceInstance.h
--- a/Engine/Source/Runtime/MovieScene/Public/MovieSceneSequenceInstance.h
+++ b/Engine/Source/Runtime/MovieScene/Public/MovieSceneSequenceInstance.h
@@ -2,6 +2,7 @@
 
 #pragma once
 
+#include "IMovieSceneTrackInstance.h"
 
 class IMovieSceneTrackInstance;
 class UMovieSceneTrack;
@@ -52,16 +53,18 @@ public:
 	 * Spawn an object relating to the specified object ID
 	 * 
 	 * @param ObjectId The ID of the object to spawn
+	 * @param Player The movie scene player responsible for this instance
 	 * @return The newly spawned or previously-spawned object, or nullptr
 	 */
-	MOVIESCENE_API void SpawnObject(const FGuid& ObjectId);
+	MOVIESCENE_API void SpawnObject(const FGuid& ObjectId, IMovieScenePlayer& Player);
 
 	/**
 	 * Destroy a previously spawned object relating to the specified object ID
 	 * 
 	 * @param ObjectId The ID of the object to destroy
+	 * @param Player The movie scene player responsible for this instance
 	 */
-	MOVIESCENE_API void DestroySpawnedObject(const FGuid& ObjectId);
+	MOVIESCENE_API void DestroySpawnedObject(const FGuid& ObjectId, IMovieScenePlayer& Player);
 
 	/**
 	 * Get the sequence associated with this instance.
@@ -90,6 +93,11 @@ protected:
 
 	void RefreshInstanceMap(const TArray<UMovieSceneTrack*>& Tracks, const TArray<UObject*>& RuntimeObjects, FMovieSceneInstanceMap& TrackInstances, class IMovieScenePlayer& Player);
 
+	void UpdateInternal(float Position, float LastPosition, class IMovieScenePlayer& Player, EMovieSceneUpdatePass UpdatePass);
+
+	/** Update the object binding instance for the specified object */
+	void UpdateObjectBinding(const FGuid& ObjectId, IMovieScenePlayer& Player);
+
 private:
 
 	/** A paring of a runtime object guid to a runtime data and track instances animating the runtime objects */
diff --git a/Engine/Source/Runtime/MovieScene/Public/MovieSceneSpawnable.h b/Engine/Source/Runtime/MovieScene/Public/MovieSceneSpawnable.h
--- a/Engine/Source/Runtime/MovieScene/Public/MovieSceneSpawnable.h
+++ b/Engine/Source/Runtime/MovieScene/Public/MovieSceneSpawnable.h
@@ -18,7 +18,7 @@ public:
 	/** FMovieSceneSpawnable default constructor */
 	FMovieSceneSpawnable() { }
 
-	/** FMovieSceneSpawnable initialization constructor */
+/** FMovieSceneSpawnable initialization constructor */
 	FMovieSceneSpawnable(const FString& InitName, UClass* InitClass)
 		: Guid(FGuid::NewGuid())
 		, Name(InitName)
@@ -60,6 +60,34 @@ public:
 		return Name;
 	}
 
+	/**
+	 * Report the specified GUID as being an inner possessable dependency for this spawnable
+	 *
+	 * @param PossessableGuid The guid pertaining to the inner possessable
+	 */
+	void AddChildPossessable(const FGuid& PossessableGuid)
+	{
+		ChildPossessables.AddUnique(PossessableGuid);
+	}
+
+	/**
+	 * Remove the specified GUID from this spawnables list of dependent possessables
+	 *
+	 * @param PossessableGuid The guid pertaining to the inner possessable
+	 */
+	void RemoveChildPossessable(const FGuid& PossessableGuid)
+	{
+		ChildPossessables.Remove(PossessableGuid);
+	}
+
+	/**
+	 * @return const access to the child possessables set
+	 */
+	const TArray<FGuid>& GetChildPossessables() const
+	{
+		return ChildPossessables;
+	}
+
 private:
 
 	/** Unique identifier of the spawnable object. */
@@ -77,4 +105,9 @@ private:
 	// @todo sequencer: Could be weak object ptr here, IF blueprints that are inners are housekept properly without references
 	UPROPERTY()
 	UClass* GeneratedClass;
+
+	/** Set of GUIDs to possessable object bindings that are bound to an object inside this spawnable */
+	// @todo sequencer: This should be a TSet, but they don't duplicate correctly atm
+	UPROPERTY()
+	TArray<FGuid> ChildPossessables;
 };
diff --git a/Engine/Source/Runtime/MovieScene/Public/MovieSceneTrack.h b/Engine/Source/Runtime/MovieScene/Public/MovieSceneTrack.h
--- a/Engine/Source/Runtime/MovieScene/Public/MovieSceneTrack.h
+++ b/Engine/Source/Runtime/MovieScene/Public/MovieSceneTrack.h
@@ -8,29 +8,18 @@
 class UMovieSceneSection;
 class FMovieSceneSequenceInstance;
 
+
 /**
  * Base class for a track in a Movie Scene
  */
-UCLASS( abstract, MinimalAPI )
+UCLASS(abstract, MinimalAPI)
 class UMovieSceneTrack
 	: public UObject
 {
-	GENERATED_UCLASS_BODY()
+	GENERATED_BODY()
 
 public:
 
-	/**
-	 * @return The name that makes this track unique from other track of the same class.
-	 */
-	virtual FName GetTrackName() const PURE_VIRTUAL (UMovieSceneTrack::GetTrackName, return NAME_None;)
-
-	/**
-	 * Called when all the sections of the track need to be retrieved.
-	 * 
-	 * @return List of all the sections in the track.
-	 */
-	virtual const TArray<UMovieSceneSection*>& GetAllSections() const PURE_VIRTUAL(UMovieSceneTrack::GetAllSections, static TArray<UMovieSceneSection*> Empty; return Empty;);
-
 	/**
 	 * Creates a runtime instance of this class.
 	 * 
@@ -40,19 +29,36 @@ public:
 	virtual TSharedPtr<class IMovieSceneTrackInstance> CreateInstance() PURE_VIRTUAL(UMovieSceneTrack::CreateInstance, check(false); return nullptr;);
 
 	/**
-	 * Removes animation data.
+	 * @return The name that makes this track unique from other track of the same class.
 	 */
-	virtual void RemoveAllAnimationData() {}
+	virtual FName GetTrackName() const { return NAME_None; }
 
 	/**
 	 * @return Whether or not this track has any data in it.
 	 */
 	virtual bool IsEmpty() const PURE_VIRTUAL(UMovieSceneTrack::IsEmpty, return true;);
 
+	/**
+	 * Removes animation data.
+	 */
+	virtual void RemoveAllAnimationData() { }
+
 	/**
 	 * @return Whether or not this track supports multiple row indices.
 	 */
-	virtual bool SupportsMultipleRows() const { return false; }
+	virtual bool SupportsMultipleRows() const
+	{
+		return false;
+	}
+
+public:
+
+	/**
+	 * Add a section to this track.
+	 *
+	 * @param Section The section to add.
+	 */
+	virtual void AddSection(UMovieSceneSection& Section) PURE_VIRTUAL(UMovieSceneSection::AddSection,);
 
 	/**
 	 * Generates a new section suitable for use with this track.
@@ -61,6 +67,20 @@ public:
 	 */
 	virtual class UMovieSceneSection* CreateNewSection() PURE_VIRTUAL(UMovieSceneTrack::CreateNewSection, return nullptr;);
 	
+	/**
+	 * Called when all the sections of the track need to be retrieved.
+	 * 
+	 * @return List of all the sections in the track.
+	 */
+	virtual const TArray<UMovieSceneSection*>& GetAllSections() const PURE_VIRTUAL(UMovieSceneTrack::GetAllSections, static TArray<UMovieSceneSection*> Empty; return Empty;);
+
+	/**
+	 * Gets the section boundaries of this track.
+	 * 
+	 * @return The range of time boundaries.
+	 */
+	virtual TRange<float> GetSectionBoundaries() const PURE_VIRTUAL(UMovieSceneTrack::GetSectionBoundaries, return TRange<float>::Empty(););
+
 	/**
 	 * Checks to see if the section is in this track.
 	 *
@@ -69,13 +89,6 @@ public:
 	 */
 	virtual bool HasSection(const UMovieSceneSection& Section) const PURE_VIRTUAL(UMovieSceneSection::HasSection, return false;);
 
-	/**
-	 * Add a section to this track.
-	 *
-	 * @param Section The section to add.
-	 */
-	virtual void AddSection(UMovieSceneSection& Section) PURE_VIRTUAL(UMovieSceneSection::AddSection,);
-
 	/**
 	 * Removes a section from this track.
 	 *
@@ -83,12 +96,14 @@ public:
 	 */
 	virtual void RemoveSection(UMovieSceneSection& Section) PURE_VIRTUAL(UMovieSceneSection::RemoveSection,);
 
+#if WITH_EDITORONLY_DATA
 	/**
-	 * Gets the section boundaries of this track.
-	 * 
-	 * @return The range of time boundaries.
+	 * Get the track's display name.
+	 *
+	 * @return Display name text.
 	 */
-	virtual TRange<float> GetSectionBoundaries() const PURE_VIRTUAL(UMovieSceneTrack::GetSectionBoundaries, return TRange<float>::Empty(););
+	virtual FText GetDisplayName() const PURE_VIRTUAL(UMovieSceneTrack::GetDisplayName, return FText::FromString(TEXT("Unnamed Track")););
+#endif
 
 #if WITH_EDITOR
 	/**
diff --git a/Engine/Source/Runtime/MovieSceneCapture/Private/AutomatedLevelSequenceCapture.cpp b/Engine/Source/Runtime/MovieSceneCapture/Private/AutomatedLevelSequenceCapture.cpp
--- a/Engine/Source/Runtime/MovieSceneCapture/Private/AutomatedLevelSequenceCapture.cpp
+++ b/Engine/Source/Runtime/MovieSceneCapture/Private/AutomatedLevelSequenceCapture.cpp
@@ -1,6 +1,7 @@
 // Copyright 1998-2015 Epic Games, Inc. All Rights Reserved.
 
 #include "MovieSceneCapturePCH.h"
+
 #include "LevelSequencePlayer.h"
 #include "AutomatedLevelSequenceCapture.h"
 #include "ErrorCodes.h"
@@ -9,10 +10,18 @@
 UAutomatedLevelSequenceCapture::UAutomatedLevelSequenceCapture(const FObjectInitializer& Init)
 	: Super(Init)
 {
+#if WITH_EDITORONLY_DATA == 0
+	if (!HasAnyFlags(RF_ClassDefaultObject))
+	{
+		checkf(false, TEXT("Automated level sequence captures can only be used in editor builds."));
+	}
+#else
 	bStageSequence = false;
-	bWasPlaying = false;
+#endif
 }
 
+#if WITH_EDITORONLY_DATA
+
 void UAutomatedLevelSequenceCapture::SetLevelSequenceAsset(FString AssetPath)
 {
 	LevelSequenceAsset = MoveTemp(AssetPath);
@@ -85,33 +94,46 @@ void UAutomatedLevelSequenceCapture::Tick(float DeltaSeconds)
 			Value += DeltaSeconds;
 			if (Value >= PrerollAmount)
 			{
-				Actor->SequencePlayer->Play();
+				if (!OnPlayerUpdatedBinding.IsValid())
+				{
+					// Bind to the event so we know when to capture a frame
+					OnPlayerUpdatedBinding = Actor->SequencePlayer->OnSequenceUpdated().AddUObject(this, &UAutomatedLevelSequenceCapture::SequenceUpdated);
+				}
+
 				StartCapture();
-				PrepareForScreenshot();
+				Actor->SequencePlayer->Play();
 
 				PrerollTime.Reset();
 			}
 		}
 
-		// If the sequence is playing, we need to capture the frame. Ensure we get the last frame if the sequence was stopped this frame.
-		else if (Actor->SequencePlayer->IsPlaying() || bWasPlaying)
-		{
-			PrepareForScreenshot();
-		}
-
 		// Otherwise we can quit if we've processed all the frames
-		else if (!bHasOutstandingFrame)
+		else if (OutstandingFrameCount == 0)
 		{
-			Close();
-			FPlatformMisc::RequestExit(0);
+			StopCapture();
 		}
+	}
+}
+
+void UAutomatedLevelSequenceCapture::OnCaptureStopped()
+{
+	Close();
+	FPlatformMisc::RequestExit(0);
+}
 
-		bWasPlaying = Actor->SequencePlayer->IsPlaying();
+void UAutomatedLevelSequenceCapture::SequenceUpdated(const ULevelSequencePlayer& Player, float CurrentTime, float PreviousTime)
+{
+	// Grab this frame
+	PrepareForScreenshot();
+	if (!Player.IsPlaying())
+	{
+		Player.OnSequenceUpdated().Remove(OnPlayerUpdatedBinding);
 	}
 }
-#if WITH_EDITOR
+
 void UAutomatedLevelSequenceCapture::PostEditChangeProperty(struct FPropertyChangedEvent& PropertyChangedEvent)
 {
 	LevelSequenceActorId = LevelSequenceActor.GetUniqueID().GetGuid();
 }
+
 #endif
\ No newline at end of file
diff --git a/Engine/Source/Runtime/MovieSceneCapture/Private/MovieSceneCapture.cpp b/Engine/Source/Runtime/MovieSceneCapture/Private/MovieSceneCapture.cpp
--- a/Engine/Source/Runtime/MovieSceneCapture/Private/MovieSceneCapture.cpp
+++ b/Engine/Source/Runtime/MovieSceneCapture/Private/MovieSceneCapture.cpp
@@ -28,6 +28,7 @@ FMovieSceneCaptureSettings::FMovieSceneCaptureSettings()
 	OutputDirectory.Path = FPaths::VideoCaptureDir();
 	OutputFormat = NSLOCTEXT("MovieCapture", "DefaultFormat", "{world}_{width}x{height}").ToString();
 	FrameRate = 24;
+	FrameCount = 0;
 	CaptureType = EMovieCaptureType::AVI;
 	bUseCompression = true;
 	CompressionQuality = 1.f;
@@ -61,7 +62,7 @@ UMovieSceneCapture::UMovieSceneCapture(const FObjectInitializer& Initializer)
 
 	Handle = FUniqueMovieSceneCaptureHandle();
 
-	bHasOutstandingFrame = false;
+	OutstandingFrameCount = 0;
 	LastFrameDelta = 0.f;
 	bCapturing = false;
 }
@@ -99,7 +100,7 @@ void UMovieSceneCapture::PrepareForScreenshot()
 		return;
 	}
 
-	bHasOutstandingFrame = true;
+	++OutstandingFrameCount;
 
 	FWidgetPath WidgetPath;
 	FSlateApplication::Get().GeneratePathToWidgetChecked(ViewportWidget.ToSharedRef(), WidgetPath);
@@ -125,10 +126,10 @@ void UMovieSceneCapture::PrepareForScreenshot()
 
 void UMovieSceneCapture::Tick(float DeltaSeconds)
 {
-	if (bHasOutstandingFrame)
+	if (OutstandingFrameCount > 0)
 	{
 		CaptureFrame(LastFrameDelta);
-		bHasOutstandingFrame = false;
+		--OutstandingFrameCount;
 	}
 	LastFrameDelta = DeltaSeconds;
 }
@@ -215,6 +216,10 @@ void UMovieSceneCapture::CaptureFrame(float DeltaSeconds)
 			SaveFrameToFile(MoveTemp(ThisFrameBuffer));
 		}
 #endif
+		if (Settings.FrameCount != 0 && CachedMetrics.Frame >= Settings.FrameCount)
+		{
+			StopCapture();
+		}
 	}
 }
 
@@ -224,6 +229,7 @@ void UMovieSceneCapture::StopCapture()
 
 	if (bCapturing)
 	{
+		OutstandingFrameCount = 0;
 		bCapturing = false;
 
 		CaptureStrategy->OnStop();
@@ -234,6 +240,8 @@ void UMovieSceneCapture::StopCapture()
 			AVIWriter->StopCapture();
 			AVIWriter.Reset();
 		}
+
+		OnCaptureStopped();
 	}
 }
 
diff --git a/Engine/Source/Runtime/MovieSceneCapture/Public/AutomatedLevelSequenceCapture.h b/Engine/Source/Runtime/MovieSceneCapture/Public/AutomatedLevelSequenceCapture.h
--- a/Engine/Source/Runtime/MovieSceneCapture/Public/AutomatedLevelSequenceCapture.h
+++ b/Engine/Source/Runtime/MovieSceneCapture/Public/AutomatedLevelSequenceCapture.h
@@ -17,6 +17,7 @@ public:
 
 	GENERATED_BODY()
 
+#if WITH_EDITORONLY_DATA
 	/** Set the level sequence asset that we are to record. We will spawn a new actor at runtime for this asset for playback. */
 	void SetLevelSequenceAsset(FString InAssetPath);
 
@@ -34,10 +35,19 @@ public:
 public:
 	// UMovieSceneCapture interface
 	virtual void Initialize(TWeakPtr<FSceneViewport> InViewport) override;
-#if WITH_EDITOR
 	virtual void PostEditChangeProperty(FPropertyChangedEvent& PropertyChangedEvent) override;
-#endif
-	
+
+private:
+
+	/** Called when this capture has stopped */
+	virtual void OnCaptureStopped() override;
+
+	/** Called when the level sequence has updated the world */
+	void SequenceUpdated(const ULevelSequencePlayer& Player, float CurrentTime, float PreviousTime);
+
+	/** Delegate binding for the above callback */
+	FDelegateHandle OnPlayerUpdatedBinding;
+
 private:
 
 	virtual void Tick(float DeltaSeconds) override;
@@ -57,6 +67,6 @@ private:
 	/** Transient amount of time that has passed since the level was started. When >= Settings.Preroll, playback will start  */
 	TOptional<float> PrerollTime;
 
-	/** Flag specifying whether the sequence was playing last frame or not */
-	bool bWasPlaying;
-};
\ No newline at end of file
+#endif
+};
+
diff --git a/Engine/Source/Runtime/MovieSceneCapture/Public/MovieSceneCapture.h b/Engine/Source/Runtime/MovieSceneCapture/Public/MovieSceneCapture.h
--- a/Engine/Source/Runtime/MovieSceneCapture/Public/MovieSceneCapture.h
+++ b/Engine/Source/Runtime/MovieSceneCapture/Public/MovieSceneCapture.h
@@ -38,7 +38,7 @@ struct FCachedMetrics
 };
 
 /** Class responsible for capturing scene data */
-UCLASS(config=EditorSettings)
+UCLASS(config=EditorPerProjectUserSettings)
 class MOVIESCENECAPTURE_API UMovieSceneCapture : public UObject, public IMovieSceneCaptureInterface
 {
 public:
@@ -103,8 +103,8 @@ protected:
 
 private:
 
-	/** Internal function that calls the virtual */
-	void OnPreWorldTickInternal(ELevelTick TickType, float DeltaSeconds) {}
+	/** Called when this movie scene capture has stopped */
+	virtual void OnCaptureStopped() {}
 
 protected:
 
@@ -151,8 +151,8 @@ protected:
 	TSharedPtr<ICaptureStrategy> CaptureStrategy;
 	/** Scratch space for per-frame screenshots */
 	TArray<FColor> ScratchBuffer;
-	/** Whether we need to capture this frame or not, set to true during a frame to capture it, then processed on the proceeding frame */
-	bool bHasOutstandingFrame;
+	/** Running count of how many frames we're waiting to capture from slate */
+	int32 OutstandingFrameCount;
 	/** The delta of the last frame */
 	float LastFrameDelta;
 	/** Whether we have started capturing or not */
diff --git a/Engine/Source/Runtime/MovieSceneCapture/Public/MovieSceneCaptureSettings.h b/Engine/Source/Runtime/MovieSceneCapture/Public/MovieSceneCaptureSettings.h
--- a/Engine/Source/Runtime/MovieSceneCapture/Public/MovieSceneCaptureSettings.h
+++ b/Engine/Source/Runtime/MovieSceneCapture/Public/MovieSceneCaptureSettings.h
@@ -59,6 +59,10 @@ struct MOVIESCENECAPTURE_API FMovieSceneCaptureSettings
 	UPROPERTY(config, EditAnywhere, Category=CaptureSettings)
 	int32 FrameRate;
 
+	/** An explicit number of frames to capture, or 0 if unspecified. Not exposed on UI - calculated internally */
+	UPROPERTY()
+	int32 FrameCount;
+
 	/** The resolution at which to capture */
 	UPROPERTY(config, EditAnywhere, Category=CaptureSettings, meta=(ShowOnlyInnerProperties))
 	FCaptureResolution Resolution;
diff --git a/Engine/Source/Runtime/MovieSceneTracks/Private/Sections/MovieScene3DAttachSection.cpp b/Engine/Source/Runtime/MovieSceneTracks/Private/Sections/MovieScene3DAttachSection.cpp
--- a/Engine/Source/Runtime/MovieSceneTracks/Private/Sections/MovieScene3DAttachSection.cpp
+++ b/Engine/Source/Runtime/MovieSceneTracks/Private/Sections/MovieScene3DAttachSection.cpp
@@ -8,6 +8,7 @@
 UMovieScene3DAttachSection::UMovieScene3DAttachSection( const FObjectInitializer& ObjectInitializer )
 	: Super( ObjectInitializer )
 {
+	AttachSocketName = NAME_None;
 	bConstrainTx = true;
 	bConstrainTy = true;
 	bConstrainTz = true;
@@ -19,10 +20,17 @@ UMovieScene3DAttachSection::UMovieScene3DAttachSection( const FObjectInitializer
 
 void UMovieScene3DAttachSection::Eval( USceneComponent* SceneComponent, float Position, AActor* Actor, FVector& OutTranslation, FRotator& OutRotation ) const
 {
-	OutTranslation = Actor->GetActorLocation();
-	OutRotation = Actor->GetActorRotation();
-
-	//@todo eval with component/bone
+	if (Actor->GetRootComponent()->DoesSocketExist(AttachSocketName))
+	{
+		FTransform SocketTransform = Actor->GetRootComponent()->GetSocketTransform(AttachSocketName);
+		OutTranslation = SocketTransform.GetLocation();
+		OutRotation = SocketTransform.GetRotation().Rotator();
+	}
+	else
+	{
+		OutTranslation = Actor->GetActorLocation();
+		OutRotation = Actor->GetActorRotation();
+	}
 
 	if (!bConstrainTx)
 	{
diff --git a/Engine/Source/Runtime/MovieSceneTracks/Private/TrackInstances/MovieScene3DConstraintTrackInstance.cpp b/Engine/Source/Runtime/MovieSceneTracks/Private/TrackInstances/MovieScene3DConstraintTrackInstance.cpp
--- a/Engine/Source/Runtime/MovieSceneTracks/Private/TrackInstances/MovieScene3DConstraintTrackInstance.cpp
+++ b/Engine/Source/Runtime/MovieSceneTracks/Private/TrackInstances/MovieScene3DConstraintTrackInstance.cpp
@@ -53,7 +53,7 @@ void FMovieScene3DConstraintTrackInstance::RestoreState(const TArray<UObject*>&
 }
 
 
-void FMovieScene3DConstraintTrackInstance::Update( float Position, float LastPosition, const TArray<UObject*>& RuntimeObjects, class IMovieScenePlayer& Player, FMovieSceneSequenceInstance& SequenceInstance ) 
+void FMovieScene3DConstraintTrackInstance::Update( float Position, float LastPosition, const TArray<UObject*>& RuntimeObjects, class IMovieScenePlayer& Player, FMovieSceneSequenceInstance& SequenceInstance, EMovieSceneUpdatePass UpdatePass ) 
 {
 	UMovieScene3DConstraintSection* FirstConstraintSection = nullptr;
 
diff --git a/Engine/Source/Runtime/MovieSceneTracks/Private/TrackInstances/MovieScene3DTransformTrackInstance.cpp b/Engine/Source/Runtime/MovieSceneTracks/Private/TrackInstances/MovieScene3DTransformTrackInstance.cpp
--- a/Engine/Source/Runtime/MovieSceneTracks/Private/TrackInstances/MovieScene3DTransformTrackInstance.cpp
+++ b/Engine/Source/Runtime/MovieSceneTracks/Private/TrackInstances/MovieScene3DTransformTrackInstance.cpp
@@ -50,7 +50,7 @@ void FMovieScene3DTransformTrackInstance::RestoreState(const TArray<UObject*>& R
 }
 
 
-void FMovieScene3DTransformTrackInstance::Update( float Position, float LastPosition, const TArray<UObject*>& RuntimeObjects, class IMovieScenePlayer& Player, FMovieSceneSequenceInstance& SequenceInstance ) 
+void FMovieScene3DTransformTrackInstance::Update( float Position, float LastPosition, const TArray<UObject*>& RuntimeObjects, class IMovieScenePlayer& Player, FMovieSceneSequenceInstance& SequenceInstance, EMovieSceneUpdatePass UpdatePass ) 
 {
 	FVector Translation;
 	FRotator Rotation;
@@ -64,9 +64,16 @@ void FMovieScene3DTransformTrackInstance::Update( float Position, float LastPosi
 
 			if (SceneComponent != nullptr)
 			{
-				SceneComponent->SetRelativeLocationAndRotation(Translation, Rotation);
-
-				SceneComponent->SetRelativeScale3D(Scale);
+				if (UpdatePass == MSUP_PreUpdate)
+				{
+					SceneComponent->ResetRelativeTransform();
+				}
+				else if (UpdatePass == MSUP_Update)
+				{
+					SceneComponent->AddRelativeLocation(Translation);
+					SceneComponent->AddRelativeRotation(Rotation);
+					SceneComponent->SetRelativeScale3D(Scale);
+				}
 			}
 		}
 	}
diff --git a/Engine/Source/Runtime/MovieSceneTracks/Private/TrackInstances/MovieSceneAudioTrackInstance.cpp b/Engine/Source/Runtime/MovieSceneTracks/Private/TrackInstances/MovieSceneAudioTrackInstance.cpp
--- a/Engine/Source/Runtime/MovieSceneTracks/Private/TrackInstances/MovieSceneAudioTrackInstance.cpp
+++ b/Engine/Source/Runtime/MovieSceneTracks/Private/TrackInstances/MovieSceneAudioTrackInstance.cpp
@@ -16,7 +16,7 @@ FMovieSceneAudioTrackInstance::FMovieSceneAudioTrackInstance( UMovieSceneAudioTr
 }
 
 
-void FMovieSceneAudioTrackInstance::Update( float Position, float LastPosition, const TArray<UObject*>& RuntimeObjects, class IMovieScenePlayer& Player, FMovieSceneSequenceInstance& SequenceInstance ) 
+void FMovieSceneAudioTrackInstance::Update( float Position, float LastPosition, const TArray<UObject*>& RuntimeObjects, class IMovieScenePlayer& Player, FMovieSceneSequenceInstance& SequenceInstance, EMovieSceneUpdatePass UpdatePass ) 
 {
 	const TArray<UMovieSceneSection*>& AudioSections = AudioTrack->GetAudioSections();
 
diff --git a/Engine/Source/Runtime/MovieSceneTracks/Private/TrackInstances/MovieSceneBoolTrackInstance.cpp b/Engine/Source/Runtime/MovieSceneTracks/Private/TrackInstances/MovieSceneBoolTrackInstance.cpp
--- a/Engine/Source/Runtime/MovieSceneTracks/Private/TrackInstances/MovieSceneBoolTrackInstance.cpp
+++ b/Engine/Source/Runtime/MovieSceneTracks/Private/TrackInstances/MovieSceneBoolTrackInstance.cpp
@@ -47,7 +47,7 @@ void FMovieSceneBoolTrackInstance::RestoreState(const TArray<UObject*>& RuntimeO
 }
 
 
-void FMovieSceneBoolTrackInstance::Update( float Position, float LastPosition, const TArray<UObject*>& RuntimeObjects, class IMovieScenePlayer& Player, FMovieSceneSequenceInstance& SequenceInstance ) 
+void FMovieSceneBoolTrackInstance::Update( float Position, float LastPosition, const TArray<UObject*>& RuntimeObjects, class IMovieScenePlayer& Player, FMovieSceneSequenceInstance& SequenceInstance, EMovieSceneUpdatePass UpdatePass ) 
 {
 	bool BoolValue = false;
 
diff --git a/Engine/Source/Runtime/MovieSceneTracks/Private/TrackInstances/MovieSceneByteTrackInstance.cpp b/Engine/Source/Runtime/MovieSceneTracks/Private/TrackInstances/MovieSceneByteTrackInstance.cpp
--- a/Engine/Source/Runtime/MovieSceneTracks/Private/TrackInstances/MovieSceneByteTrackInstance.cpp
+++ b/Engine/Source/Runtime/MovieSceneTracks/Private/TrackInstances/MovieSceneByteTrackInstance.cpp
@@ -46,7 +46,7 @@ void FMovieSceneByteTrackInstance::RestoreState(const TArray<UObject*>& RuntimeO
 }
 
 
-void FMovieSceneByteTrackInstance::Update( float Position, float LastPosition, const TArray<UObject*>& RuntimeObjects, class IMovieScenePlayer& Player, FMovieSceneSequenceInstance& SequenceInstance ) 
+void FMovieSceneByteTrackInstance::Update( float Position, float LastPosition, const TArray<UObject*>& RuntimeObjects, class IMovieScenePlayer& Player, FMovieSceneSequenceInstance& SequenceInstance, EMovieSceneUpdatePass UpdatePass ) 
 {
 	uint8 ByteValue = 0;
 	if( ByteTrack->Eval( Position, LastPosition, ByteValue ) )
diff --git a/Engine/Source/Runtime/MovieSceneTracks/Private/TrackInstances/MovieSceneColorTrackInstance.cpp b/Engine/Source/Runtime/MovieSceneTracks/Private/TrackInstances/MovieSceneColorTrackInstance.cpp
--- a/Engine/Source/Runtime/MovieSceneTracks/Private/TrackInstances/MovieSceneColorTrackInstance.cpp
+++ b/Engine/Source/Runtime/MovieSceneTracks/Private/TrackInstances/MovieSceneColorTrackInstance.cpp
@@ -91,7 +91,7 @@ void FMovieSceneColorTrackInstance::RestoreState(const TArray<UObject*>& Runtime
 }
 
 
-void FMovieSceneColorTrackInstance::Update( float Position, float LastPosition, const TArray<UObject*>& RuntimeObjects, class IMovieScenePlayer& Player, FMovieSceneSequenceInstance& SequenceInstance ) 
+void FMovieSceneColorTrackInstance::Update( float Position, float LastPosition, const TArray<UObject*>& RuntimeObjects, class IMovieScenePlayer& Player, FMovieSceneSequenceInstance& SequenceInstance, EMovieSceneUpdatePass UpdatePass ) 
 {
 	for(UObject* Object : RuntimeObjects)
 	{
diff --git a/Engine/Source/Runtime/MovieSceneTracks/Private/TrackInstances/MovieSceneEventTrackInstance.cpp b/Engine/Source/Runtime/MovieSceneTracks/Private/TrackInstances/MovieSceneEventTrackInstance.cpp
--- a/Engine/Source/Runtime/MovieSceneTracks/Private/TrackInstances/MovieSceneEventTrackInstance.cpp
+++ b/Engine/Source/Runtime/MovieSceneTracks/Private/TrackInstances/MovieSceneEventTrackInstance.cpp
@@ -17,7 +17,7 @@ FMovieSceneEventTrackInstance::FMovieSceneEventTrackInstance(UMovieSceneEventTra
 /* IMovieSceneTrackInstance interface
  *****************************************************************************/
 
-void FMovieSceneEventTrackInstance::Update(float Position, float LastPosition, const TArray<UObject*>& RuntimeObjects, IMovieScenePlayer& Player, FMovieSceneSequenceInstance& SequenceInstance)
+void FMovieSceneEventTrackInstance::Update(float Position, float LastPosition, const TArray<UObject*>& RuntimeObjects, IMovieScenePlayer& Player, FMovieSceneSequenceInstance& SequenceInstance, EMovieSceneUpdatePass UpdatePass)
 {
 	EventTrack->TriggerEvents(Position, LastPosition);
 }
diff --git a/Engine/Source/Runtime/MovieSceneTracks/Private/TrackInstances/MovieSceneFadeTrackInstance.cpp b/Engine/Source/Runtime/MovieSceneTracks/Private/TrackInstances/MovieSceneFadeTrackInstance.cpp
--- a/Engine/Source/Runtime/MovieSceneTracks/Private/TrackInstances/MovieSceneFadeTrackInstance.cpp
+++ b/Engine/Source/Runtime/MovieSceneTracks/Private/TrackInstances/MovieSceneFadeTrackInstance.cpp
@@ -34,7 +34,7 @@ void FMovieSceneFadeTrackInstance::RestoreState(const TArray<UObject*>& RuntimeO
 	Player.SetViewportSettings(ViewportParamsMap);
 }
 
-void FMovieSceneFadeTrackInstance::Update(float Position, float LastPosition, const TArray<UObject*>& RuntimeObjects, IMovieScenePlayer& Player, FMovieSceneSequenceInstance& SequenceInstance)
+void FMovieSceneFadeTrackInstance::Update(float Position, float LastPosition, const TArray<UObject*>& RuntimeObjects, IMovieScenePlayer& Player, FMovieSceneSequenceInstance& SequenceInstance, EMovieSceneUpdatePass UpdatePass)
 {
 	FLinearColor FadeColor = FLinearColor::Black;
 	bool bFadeAudio = false;
diff --git a/Engine/Source/Runtime/MovieSceneTracks/Private/TrackInstances/MovieSceneFloatTrackInstance.cpp b/Engine/Source/Runtime/MovieSceneTracks/Private/TrackInstances/MovieSceneFloatTrackInstance.cpp
--- a/Engine/Source/Runtime/MovieSceneTracks/Private/TrackInstances/MovieSceneFloatTrackInstance.cpp
+++ b/Engine/Source/Runtime/MovieSceneTracks/Private/TrackInstances/MovieSceneFloatTrackInstance.cpp
@@ -46,7 +46,7 @@ void FMovieSceneFloatTrackInstance::RestoreState(const TArray<UObject*>& Runtime
 }
 
 
-void FMovieSceneFloatTrackInstance::Update( float Position, float LastPosition, const TArray<UObject*>& RuntimeObjects, class IMovieScenePlayer& Player, FMovieSceneSequenceInstance& SequenceInstance ) 
+void FMovieSceneFloatTrackInstance::Update( float Position, float LastPosition, const TArray<UObject*>& RuntimeObjects, class IMovieScenePlayer& Player, FMovieSceneSequenceInstance& SequenceInstance, EMovieSceneUpdatePass UpdatePass ) 
 {
 	float FloatValue = 0.0f;
 	if( FloatTrack->Eval( Position, LastPosition, FloatValue ) )
diff --git a/Engine/Source/Runtime/MovieSceneTracks/Private/TrackInstances/MovieSceneMaterialTrackInstance.cpp b/Engine/Source/Runtime/MovieSceneTracks/Private/TrackInstances/MovieSceneMaterialTrackInstance.cpp
--- a/Engine/Source/Runtime/MovieSceneTracks/Private/TrackInstances/MovieSceneMaterialTrackInstance.cpp
+++ b/Engine/Source/Runtime/MovieSceneTracks/Private/TrackInstances/MovieSceneMaterialTrackInstance.cpp
@@ -53,7 +53,7 @@ void FMovieSceneMaterialTrackInstance::RestoreState(const TArray<UObject*>& Runt
 	}
 }
 
-void FMovieSceneMaterialTrackInstance::Update( float Position, float LastPosition, const TArray<UObject*>& RuntimeObjects, class IMovieScenePlayer& Player, FMovieSceneSequenceInstance& SequenceInstance ) 
+void FMovieSceneMaterialTrackInstance::Update( float Position, float LastPosition, const TArray<UObject*>& RuntimeObjects, class IMovieScenePlayer& Player, FMovieSceneSequenceInstance& SequenceInstance, EMovieSceneUpdatePass UpdatePass ) 
 {
 	TArray<FScalarParameterNameAndValue> ScalarValues;
 	TArray<FVectorParameterNameAndValue> VectorValues;
diff --git a/Engine/Source/Runtime/MovieSceneTracks/Private/TrackInstances/MovieSceneParticleTrackInstance.cpp b/Engine/Source/Runtime/MovieSceneTracks/Private/TrackInstances/MovieSceneParticleTrackInstance.cpp
--- a/Engine/Source/Runtime/MovieSceneTracks/Private/TrackInstances/MovieSceneParticleTrackInstance.cpp
+++ b/Engine/Source/Runtime/MovieSceneTracks/Private/TrackInstances/MovieSceneParticleTrackInstance.cpp
@@ -13,7 +13,7 @@ FMovieSceneParticleTrackInstance::~FMovieSceneParticleTrackInstance()
 }
 
 
-void FMovieSceneParticleTrackInstance::Update( float Position, float LastPosition, const TArray<UObject*>& RuntimeObjects, class IMovieScenePlayer& Player, FMovieSceneSequenceInstance& SequenceInstance ) 
+void FMovieSceneParticleTrackInstance::Update( float Position, float LastPosition, const TArray<UObject*>& RuntimeObjects, class IMovieScenePlayer& Player, FMovieSceneSequenceInstance& SequenceInstance, EMovieSceneUpdatePass UpdatePass ) 
 {
 	// @todo Sequencer We need something analagous to Matinee 1's particle replay tracks
 	// What we have here is simple toggling/triggering
diff --git a/Engine/Source/Runtime/MovieSceneTracks/Private/TrackInstances/MovieSceneShotTrackInstance.cpp b/Engine/Source/Runtime/MovieSceneTracks/Private/TrackInstances/MovieSceneShotTrackInstance.cpp
--- a/Engine/Source/Runtime/MovieSceneTracks/Private/TrackInstances/MovieSceneShotTrackInstance.cpp
+++ b/Engine/Source/Runtime/MovieSceneTracks/Private/TrackInstances/MovieSceneShotTrackInstance.cpp
@@ -1,6 +1,8 @@
 // Copyright 1998-2015 Epic Games, Inc. All Rights Reserved.
 
 #include "MovieSceneTracksPrivatePCH.h"
+#include "MovieSceneSequence.h"
+#include "MovieSceneSequenceInstance.h"
 #include "MovieSceneShotTrack.h"
 #include "MovieSceneShotTrackInstance.h"
 #include "MovieSceneShotSection.h"
@@ -20,7 +22,7 @@ FMovieSceneShotTrackInstance::FMovieSceneShotTrackInstance(UMovieSceneShotTrack&
 
 void FMovieSceneShotTrackInstance::ClearInstance(IMovieScenePlayer& Player, FMovieSceneSequenceInstance& SequenceInstance)
 {
-	Player.UpdateCameraCut(nullptr, false);
+	Player.UpdateCameraCut(nullptr, nullptr);
 }
 
 
@@ -28,48 +30,56 @@ void FMovieSceneShotTrackInstance::RefreshInstance(const TArray<UObject*>& Runti
 {
 	const TArray<UMovieSceneSection*>& ShotSections = ShotTrack->GetAllSections();
 
-	RuntimeCameraObjects.Empty(ShotSections.Num());
+	CachedCameraObjects.SetNumZeroed(ShotSections.Num());
 
-	for (UMovieSceneSection* Section : ShotSections)
+	for (int32 ShotIndex = 0; ShotIndex < ShotSections.Num(); ++ShotIndex)
 	{
-		// @todo Sequencer - Sub-moviescenes: Get the cameras from the root movie scene instance.  We should support adding cameras for sub-moviescenes as shots
-		TArray<UObject*> CameraObjects;
-		Player.GetRuntimeObjects(Player.GetRootMovieSceneSequenceInstance(), CastChecked<UMovieSceneShotSection>(Section)->GetCameraGuid(), CameraObjects);
-
-		if (CameraObjects.Num() == 1)
-		{
-			RuntimeCameraObjects.Add(CameraObjects[0]);
-		}
-		else
-		{
-			// No valid camera object was found, take up space.  There should be one entry per section
-			RuntimeCameraObjects.Add(nullptr);
-		}
+		UMovieSceneShotSection* ShotSection = CastChecked<UMovieSceneShotSection>(ShotSections[ShotIndex]);
+		AcquireCameraForShot(ShotIndex, ShotSection->GetCameraGuid(), SequenceInstance);
 	}
 }
 
 
-void FMovieSceneShotTrackInstance::Update(float Position, float LastPosition, const TArray<UObject*>& RuntimeObjects, class IMovieScenePlayer& Player, FMovieSceneSequenceInstance& SequenceInstance) 
+void FMovieSceneShotTrackInstance::Update(float Position, float LastPosition, const TArray<UObject*>& RuntimeObjects, class IMovieScenePlayer& Player, FMovieSceneSequenceInstance& SequenceInstance, EMovieSceneUpdatePass UpdatePass) 
 {
+	UObject* CameraObject = nullptr;
 	const TArray<UMovieSceneSection*>& ShotSections = ShotTrack->GetAllSections();
 
 	for (int32 ShotIndex = 0; ShotIndex < ShotSections.Num(); ++ShotIndex)
 	{
 		UMovieSceneShotSection* ShotSection = CastChecked<UMovieSceneShotSection>(ShotSections[ShotIndex]);
 
-		// Note shot times are not inclusive
-		// @todo Sequencer: This could be faster with a binary search
-		if (ShotSection->GetStartTime() <= Position && ShotSection->GetEndTime() > Position)
+		if (ShotSection->GetRange().Contains(Position))
 		{
-			UObject* Camera = RuntimeCameraObjects[ShotIndex].Get();
-			
-			const bool bNewCameraCut = CurrentCameraObject != Camera;
-			Player.UpdateCameraCut(Camera, bNewCameraCut);
-
-			CurrentCameraObject = Camera;
+			CameraObject = AcquireCameraForShot(ShotIndex, ShotSection->GetCameraGuid(), SequenceInstance);
 
-			// no need to process any more shots.  Only one shot can be active at a time
 			break;
 		}
 	}
+
+	if (CameraObject != LastCameraObject)
+	{
+		Player.UpdateCameraCut(CameraObject, LastCameraObject.Get());
+		LastCameraObject = CameraObject;
+	}
+	else if (CameraObject != nullptr)
+	{
+		Player.UpdateCameraCut(CameraObject, nullptr);
+	}
+}
+
+
+/* FMovieSceneShotTrackInstance implementation
+ *****************************************************************************/
+
+UObject* FMovieSceneShotTrackInstance::AcquireCameraForShot(int32 ShotIndex, const FGuid& CameraGuid, FMovieSceneSequenceInstance& SequenceInstance)
+{
+	auto& CachedCamera = CachedCameraObjects[ShotIndex];
+
+	if (!CachedCamera.IsValid())
+	{
+		CachedCamera = SequenceInstance.GetSequence()->FindObject(CameraGuid);
+	}
+
+	return CachedCamera.Get();
 }
diff --git a/Engine/Source/Runtime/MovieSceneTracks/Private/TrackInstances/MovieSceneSkeletalAnimationTrackInstance.cpp b/Engine/Source/Runtime/MovieSceneTracks/Private/TrackInstances/MovieSceneSkeletalAnimationTrackInstance.cpp
--- a/Engine/Source/Runtime/MovieSceneTracks/Private/TrackInstances/MovieSceneSkeletalAnimationTrackInstance.cpp
+++ b/Engine/Source/Runtime/MovieSceneTracks/Private/TrackInstances/MovieSceneSkeletalAnimationTrackInstance.cpp
@@ -21,7 +21,7 @@ FMovieSceneSkeletalAnimationTrackInstance::~FMovieSceneSkeletalAnimationTrackIns
 }
 
 
-void FMovieSceneSkeletalAnimationTrackInstance::Update( float Position, float LastPosition, const TArray<UObject*>& RuntimeObjects, class IMovieScenePlayer& Player, FMovieSceneSequenceInstance& SequenceInstance ) 
+void FMovieSceneSkeletalAnimationTrackInstance::Update( float Position, float LastPosition, const TArray<UObject*>& RuntimeObjects, class IMovieScenePlayer& Player, FMovieSceneSequenceInstance& SequenceInstance, EMovieSceneUpdatePass UpdatePass ) 
 {
 	// @todo Sequencer gameplay update has a different code path than editor update for animation
 
diff --git a/Engine/Source/Runtime/MovieSceneTracks/Private/TrackInstances/MovieSceneSlomoTrackInstance.cpp b/Engine/Source/Runtime/MovieSceneTracks/Private/TrackInstances/MovieSceneSlomoTrackInstance.cpp
--- a/Engine/Source/Runtime/MovieSceneTracks/Private/TrackInstances/MovieSceneSlomoTrackInstance.cpp
+++ b/Engine/Source/Runtime/MovieSceneTracks/Private/TrackInstances/MovieSceneSlomoTrackInstance.cpp
@@ -16,7 +16,7 @@ FMovieSceneSlomoTrackInstance::FMovieSceneSlomoTrackInstance(UMovieSceneSlomoTra
 /* IMovieSceneTrackInstance interface
  *****************************************************************************/
 
-void FMovieSceneSlomoTrackInstance::Update(float Position, float LastPosition, const TArray<UObject*>& RuntimeObjects, IMovieScenePlayer& Player, FMovieSceneSequenceInstance& SequenceInstance)
+void FMovieSceneSlomoTrackInstance::Update(float Position, float LastPosition, const TArray<UObject*>& RuntimeObjects, IMovieScenePlayer& Player, FMovieSceneSequenceInstance& SequenceInstance, EMovieSceneUpdatePass UpdatePass)
 {
 	if (!ShouldBeApplied())
 	{
diff --git a/Engine/Source/Runtime/MovieSceneTracks/Private/TrackInstances/MovieSceneSpawnTrackInstance.cpp b/Engine/Source/Runtime/MovieSceneTracks/Private/TrackInstances/MovieSceneSpawnTrackInstance.cpp
--- a/Engine/Source/Runtime/MovieSceneTracks/Private/TrackInstances/MovieSceneSpawnTrackInstance.cpp
+++ b/Engine/Source/Runtime/MovieSceneTracks/Private/TrackInstances/MovieSceneSpawnTrackInstance.cpp
@@ -12,11 +12,11 @@ FMovieSceneSpawnTrackInstance::FMovieSceneSpawnTrackInstance(UMovieSceneSpawnTra
 	Track = &InTrack;
 }
 
-void FMovieSceneSpawnTrackInstance::Update(float Position, float LastPosition, const TArray<UObject*>& RuntimeObjects, IMovieScenePlayer& Player, FMovieSceneSequenceInstance& SequenceInstance)
+void FMovieSceneSpawnTrackInstance::Update(float Position, float LastPosition, const TArray<UObject*>& RuntimeObjects, IMovieScenePlayer& Player, FMovieSceneSequenceInstance& SequenceInstance, EMovieSceneUpdatePass UpdatePass)
 {
 	if (!SequenceInstance.GetTimeRange().Contains(Position))
 	{
-		SequenceInstance.DestroySpawnedObject(Track->GetObject());
+		SequenceInstance.DestroySpawnedObject(Track->GetObject(), Player);
 		return;
 	}
 
@@ -26,13 +26,13 @@ void FMovieSceneSpawnTrackInstance::Update(float Position, float LastPosition, c
 		// Spawn the object if needed
 		if (bIsSpawned && RuntimeObjects.Num() == 0)
 		{
-			SequenceInstance.SpawnObject(Track->GetObject());
+			SequenceInstance.SpawnObject(Track->GetObject(), Player);
 		}
 
 		// Destroy the object if needed
 		if (!bIsSpawned && RuntimeObjects.Num() != 0)
 		{
-			SequenceInstance.DestroySpawnedObject(Track->GetObject());
+			SequenceInstance.DestroySpawnedObject(Track->GetObject(), Player);
 		}
 	}
 }
@@ -41,11 +41,11 @@ void FMovieSceneSpawnTrackInstance::RestoreState(const TArray<UObject*>& Runtime
 {
 	if (RuntimeObjects.Num())
 	{
-		SequenceInstance.DestroySpawnedObject(Track->GetObject());
+		SequenceInstance.DestroySpawnedObject(Track->GetObject(), Player);
 	}
 }
 
 void FMovieSceneSpawnTrackInstance::ClearInstance(IMovieScenePlayer& Player, FMovieSceneSequenceInstance& SequenceInstance)
 {
-	SequenceInstance.DestroySpawnedObject(Track->GetObject());
+	SequenceInstance.DestroySpawnedObject(Track->GetObject(), Player);
 }
\ No newline at end of file
diff --git a/Engine/Source/Runtime/MovieSceneTracks/Private/TrackInstances/MovieSceneSubTrackInstance.cpp b/Engine/Source/Runtime/MovieSceneTracks/Private/TrackInstances/MovieSceneSubTrackInstance.cpp
--- a/Engine/Source/Runtime/MovieSceneTracks/Private/TrackInstances/MovieSceneSubTrackInstance.cpp
+++ b/Engine/Source/Runtime/MovieSceneTracks/Private/TrackInstances/MovieSceneSubTrackInstance.cpp
@@ -104,24 +104,13 @@ void FMovieSceneSubTrackInstance::SaveState(const TArray<UObject*>& RuntimeObjec
 }
 
 
-void FMovieSceneSubTrackInstance::Update(float Position, float LastPosition, const TArray<UObject*>& RuntimeObjects, class IMovieScenePlayer& Player, FMovieSceneSequenceInstance& SequenceInstance) 
+void FMovieSceneSubTrackInstance::Update(float Position, float LastPosition, const TArray<UObject*>& RuntimeObjects, class IMovieScenePlayer& Player, FMovieSceneSequenceInstance& SequenceInstance, EMovieSceneUpdatePass UpdatePass) 
 {
 	const TArray<UMovieSceneSection*>& AllSections = SubTrack->GetAllSections();
 	TArray<UMovieSceneSection*> TraversedSections = MovieSceneHelpers::GetTraversedSections(AllSections, Position, LastPosition);
 
 	for (const auto Section : TraversedSections)
 	{
-		const TRange<float>& SectionRange = Section->GetRange();
-
-		// skip inactive sections
-		if (!SectionRange.Overlaps((LastPosition <= Position)
-				? TRange<float>(LastPosition, Position)
-				: TRange<float>(Position, LastPosition)
-			))
-		{
-			continue;
-		}
-
 		// skip sections with invalid time scale
 		UMovieSceneSubSection* SubSection = CastChecked<UMovieSceneSubSection>(Section);
 
diff --git a/Engine/Source/Runtime/MovieSceneTracks/Private/TrackInstances/MovieSceneVectorTrackInstance.cpp b/Engine/Source/Runtime/MovieSceneTracks/Private/TrackInstances/MovieSceneVectorTrackInstance.cpp
--- a/Engine/Source/Runtime/MovieSceneTracks/Private/TrackInstances/MovieSceneVectorTrackInstance.cpp
+++ b/Engine/Source/Runtime/MovieSceneTracks/Private/TrackInstances/MovieSceneVectorTrackInstance.cpp
@@ -97,7 +97,7 @@ void FMovieSceneVectorTrackInstance::RestoreState(const TArray<UObject*>& Runtim
 }
 
 
-void FMovieSceneVectorTrackInstance::Update( float Position, float LastPosition, const TArray<UObject*>& RuntimeObjects, class IMovieScenePlayer& Player, FMovieSceneSequenceInstance& SequenceInstance ) 
+void FMovieSceneVectorTrackInstance::Update( float Position, float LastPosition, const TArray<UObject*>& RuntimeObjects, class IMovieScenePlayer& Player, FMovieSceneSequenceInstance& SequenceInstance, EMovieSceneUpdatePass UpdatePass ) 
 {
 	FVector4 Vector;
 	if( VectorTrack->Eval( Position, LastPosition, Vector ) )
diff --git a/Engine/Source/Runtime/MovieSceneTracks/Private/TrackInstances/MovieSceneVisibilityTrackInstance.cpp b/Engine/Source/Runtime/MovieSceneTracks/Private/TrackInstances/MovieSceneVisibilityTrackInstance.cpp
--- a/Engine/Source/Runtime/MovieSceneTracks/Private/TrackInstances/MovieSceneVisibilityTrackInstance.cpp
+++ b/Engine/Source/Runtime/MovieSceneTracks/Private/TrackInstances/MovieSceneVisibilityTrackInstance.cpp
@@ -98,7 +98,7 @@ void FMovieSceneVisibilityTrackInstance::RestoreState(const TArray<UObject*>& Ru
 }
 
 
-void FMovieSceneVisibilityTrackInstance::Update( float Position, float LastPosition, const TArray<UObject*>& RuntimeObjects, class IMovieScenePlayer& Player, FMovieSceneSequenceInstance& SequenceInstance ) 
+void FMovieSceneVisibilityTrackInstance::Update( float Position, float LastPosition, const TArray<UObject*>& RuntimeObjects, class IMovieScenePlayer& Player, FMovieSceneSequenceInstance& SequenceInstance, EMovieSceneUpdatePass UpdatePass ) 
 {
  	bool Visible = false;
  	if( VisibilityTrack->Eval( Position, LastPosition, Visible ) )
diff --git a/Engine/Source/Runtime/MovieSceneTracks/Private/Tracks/MovieScene3DAttachTrack.cpp b/Engine/Source/Runtime/MovieSceneTracks/Private/Tracks/MovieScene3DAttachTrack.cpp
--- a/Engine/Source/Runtime/MovieSceneTracks/Private/Tracks/MovieScene3DAttachTrack.cpp
+++ b/Engine/Source/Runtime/MovieSceneTracks/Private/Tracks/MovieScene3DAttachTrack.cpp
@@ -7,29 +7,38 @@
 #include "MovieScene3DAttachTrackInstance.h"
 
 
+#define LOCTEXT_NAMESPACE "MovieScene3DAttachTrack"
+
+
 UMovieScene3DAttachTrack::UMovieScene3DAttachTrack( const FObjectInitializer& ObjectInitializer )
 	: Super( ObjectInitializer )
 { }
 
 
-FName UMovieScene3DAttachTrack::GetTrackName() const
-{
-	return FName("Attach");
-}
-
-
 TSharedPtr<IMovieSceneTrackInstance> UMovieScene3DAttachTrack::CreateInstance()
 {
 	return MakeShareable( new FMovieScene3DAttachTrackInstance( *this ) ); 
 }
 
 
-void UMovieScene3DAttachTrack::AddConstraint(float KeyTime, float ConstraintEndTime, const FGuid& ConstraintId)
+void UMovieScene3DAttachTrack::AddConstraint(float KeyTime, float ConstraintEndTime, const FName SocketName, const FGuid& ConstraintId)
 {
 	// add the section
 	UMovieScene3DAttachSection* NewSection = NewObject<UMovieScene3DAttachSection>(this);
 	NewSection->AddAttach(KeyTime, ConstraintEndTime, ConstraintId);
 	NewSection->InitialPlacement( ConstraintSections, KeyTime, ConstraintEndTime, SupportsMultipleRows() );
+	NewSection->AttachSocketName = SocketName;
 
 	ConstraintSections.Add(NewSection);
 }
+
+
+#if WITH_EDITORONLY_DATA
+FText UMovieScene3DAttachTrack::GetDisplayName() const
+{
+	return LOCTEXT("TrackName", "Attach");
+}
+#endif
+
+
+#undef LOCTEXT_NAMESPACE
diff --git a/Engine/Source/Runtime/MovieSceneTracks/Private/Tracks/MovieScene3DConstraintTrack.cpp b/Engine/Source/Runtime/MovieSceneTracks/Private/Tracks/MovieScene3DConstraintTrack.cpp
--- a/Engine/Source/Runtime/MovieSceneTracks/Private/Tracks/MovieScene3DConstraintTrack.cpp
+++ b/Engine/Source/Runtime/MovieSceneTracks/Private/Tracks/MovieScene3DConstraintTrack.cpp
@@ -7,6 +7,9 @@
 #include "MovieScene3DConstraintTrackInstance.h"
 
 
+#define LOCTEXT_NAMESPACE "MovieScene3DConstraintTrack"
+
+
 UMovieScene3DConstraintTrack::UMovieScene3DConstraintTrack( const FObjectInitializer& ObjectInitializer )
 	: Super( ObjectInitializer )
 { }
@@ -59,3 +62,14 @@ TRange<float> UMovieScene3DConstraintTrack::GetSectionBoundaries() const
 
 	return TRange<float>::Hull(Bounds);
 }
+
+
+#if WITH_EDITORONLY_DATA
+FText UMovieScene3DConstraintTrack::GetDisplayName() const
+{
+	return LOCTEXT("TrackName", "Constraint");
+}
+#endif
+
+
+#undef LOCTEXT_NAMESPACE
diff --git a/Engine/Source/Runtime/MovieSceneTracks/Private/Tracks/MovieScene3DPathTrack.cpp b/Engine/Source/Runtime/MovieSceneTracks/Private/Tracks/MovieScene3DPathTrack.cpp
--- a/Engine/Source/Runtime/MovieSceneTracks/Private/Tracks/MovieScene3DPathTrack.cpp
+++ b/Engine/Source/Runtime/MovieSceneTracks/Private/Tracks/MovieScene3DPathTrack.cpp
@@ -7,24 +7,21 @@
 #include "MovieScene3DPathTrackInstance.h"
 
 
+#define LOCTEXT_NAMESPACE "MovieScene3DPathTrack"
+
+
 UMovieScene3DPathTrack::UMovieScene3DPathTrack( const FObjectInitializer& ObjectInitializer )
 	: Super( ObjectInitializer )
 { }
 
 
-FName UMovieScene3DPathTrack::GetTrackName() const
-{
-	return FName("Path");
-}
-
-
 TSharedPtr<IMovieSceneTrackInstance> UMovieScene3DPathTrack::CreateInstance()
 {
 	return MakeShareable( new FMovieScene3DPathTrackInstance( *this ) ); 
 }
 
 
-void UMovieScene3DPathTrack::AddConstraint(float KeyTime, float ConstraintEndTime, const FGuid& ConstraintId)
+void UMovieScene3DPathTrack::AddConstraint(float KeyTime, float ConstraintEndTime, const FName SocketName, const FGuid& ConstraintId)
 {
 	UMovieScene3DPathSection* NewSection = NewObject<UMovieScene3DPathSection>(this);
 	{
@@ -34,3 +31,14 @@ void UMovieScene3DPathTrack::AddConstraint(float KeyTime, float ConstraintEndTim
 
 	ConstraintSections.Add(NewSection);
 }
+
+
+#if WITH_EDITORONLY_DATA
+FText UMovieScene3DPathTrack::GetDisplayName() const
+{
+	return LOCTEXT("TrackName", "Path");
+}
+#endif
+
+
+#undef LOCTEXT_NAMESPACE
diff --git a/Engine/Source/Runtime/MovieSceneTracks/Private/Tracks/MovieSceneAudioTrack.cpp b/Engine/Source/Runtime/MovieSceneTracks/Private/Tracks/MovieSceneAudioTrack.cpp
--- a/Engine/Source/Runtime/MovieSceneTracks/Private/Tracks/MovieSceneAudioTrack.cpp
+++ b/Engine/Source/Runtime/MovieSceneTracks/Private/Tracks/MovieSceneAudioTrack.cpp
@@ -20,12 +20,6 @@ UMovieSceneAudioTrack::UMovieSceneAudioTrack( const FObjectInitializer& ObjectIn
 { }
 
 
-FName UMovieSceneAudioTrack::GetTrackName() const
-{
-	return AudioTrackConstants::UniqueTrackName;
-}
-
-
 TSharedPtr<IMovieSceneTrackInstance> UMovieSceneAudioTrack::CreateInstance()
 {
 	return MakeShareable( new FMovieSceneAudioTrackInstance( *this ) ); 
@@ -138,4 +132,10 @@ bool UMovieSceneAudioTrack::IsAMasterTrack() const
 }
 
 
+bool UMovieSceneAudioTrack::SupportsMultipleRows() const
+{
+	return true;
+}
+
+
 #undef LOCTEXT_NAMESPACE
diff --git a/Engine/Source/Runtime/MovieSceneTracks/Private/Tracks/MovieSceneEventTrack.cpp b/Engine/Source/Runtime/MovieSceneTracks/Private/Tracks/MovieSceneEventTrack.cpp
--- a/Engine/Source/Runtime/MovieSceneTracks/Private/Tracks/MovieSceneEventTrack.cpp
+++ b/Engine/Source/Runtime/MovieSceneTracks/Private/Tracks/MovieSceneEventTrack.cpp
@@ -118,12 +118,6 @@ TRange<float> UMovieSceneEventTrack::GetSectionBoundaries() const
 }
 
 
-FName UMovieSceneEventTrack::GetTrackName() const
-{
-	return TrackName;
-}
-
-
 bool UMovieSceneEventTrack::HasSection(const UMovieSceneSection& Section) const
 {
 	return Sections.Contains(&Section);
diff --git a/Engine/Source/Runtime/MovieSceneTracks/Private/Tracks/MovieSceneFadeTrack.cpp b/Engine/Source/Runtime/MovieSceneTracks/Private/Tracks/MovieSceneFadeTrack.cpp
--- a/Engine/Source/Runtime/MovieSceneTracks/Private/Tracks/MovieSceneFadeTrack.cpp
+++ b/Engine/Source/Runtime/MovieSceneTracks/Private/Tracks/MovieSceneFadeTrack.cpp
@@ -6,7 +6,10 @@
 #include "MovieSceneFadeTrackInstance.h"
 
 
-/* UMovieSceneEventTrack interface
+#define LOCTEXT_NAMESPACE "MovieSceneFadeTrack"
+
+
+/* UMovieSceneFadeTrack interface
  *****************************************************************************/
 
 TSharedPtr<IMovieSceneTrackInstance> UMovieSceneFadeTrack::CreateInstance()
@@ -21,7 +24,12 @@ UMovieSceneSection* UMovieSceneFadeTrack::CreateNewSection()
 }
 
 
-FName UMovieSceneFadeTrack::GetTrackName() const
+#if WITH_EDITORONLY_DATA
+FText UMovieSceneFadeTrack::GetDisplayName() const
 {
-	return FName("Fade");
+	return LOCTEXT("TrackName", "Fade");
 }
+#endif
+
+
+#undef LOCTEXT_NAMESPACE
diff --git a/Engine/Source/Runtime/MovieSceneTracks/Private/Tracks/MovieSceneMaterialTrack.cpp b/Engine/Source/Runtime/MovieSceneTracks/Private/Tracks/MovieSceneMaterialTrack.cpp
--- a/Engine/Source/Runtime/MovieSceneTracks/Private/Tracks/MovieSceneMaterialTrack.cpp
+++ b/Engine/Source/Runtime/MovieSceneTracks/Private/Tracks/MovieSceneMaterialTrack.cpp
@@ -117,3 +117,11 @@ TSharedPtr<IMovieSceneTrackInstance> UMovieSceneComponentMaterialTrack::CreateIn
 {
 	return MakeShareable( new FMovieSceneComponentMaterialTrackInstance( *this ) );
 }
+
+
+#if WITH_EDITORONLY_DATA
+FText UMovieSceneComponentMaterialTrack::GetDisplayName() const
+{
+	return FText::FromString(FString::Printf(TEXT("Material Element %i"), MaterialIndex));
+}
+#endif
diff --git a/Engine/Source/Runtime/MovieSceneTracks/Private/Tracks/MovieSceneParticleTrack.cpp b/Engine/Source/Runtime/MovieSceneTracks/Private/Tracks/MovieSceneParticleTrack.cpp
--- a/Engine/Source/Runtime/MovieSceneTracks/Private/Tracks/MovieSceneParticleTrack.cpp
+++ b/Engine/Source/Runtime/MovieSceneTracks/Private/Tracks/MovieSceneParticleTrack.cpp
@@ -15,12 +15,6 @@ UMovieSceneParticleTrack::UMovieSceneParticleTrack( const FObjectInitializer& Ob
 { }
 
 
-FName UMovieSceneParticleTrack::GetTrackName() const
-{
-	return FName("ParticleSystem");
-}
-
-
 TSharedPtr<IMovieSceneTrackInstance> UMovieSceneParticleTrack::CreateInstance()
 {
 	return MakeShareable( new FMovieSceneParticleTrackInstance( *this ) ); 
@@ -73,6 +67,7 @@ TRange<float> UMovieSceneParticleTrack::GetSectionBoundaries() const
 	return TRange<float>::Hull(Bounds);
 }
 
+
 void UMovieSceneParticleTrack::AddNewSection( float SectionTime )
 {
 	if ( MovieSceneHelpers::FindSectionAtTime( ParticleSections, SectionTime ) == nullptr )
diff --git a/Engine/Source/Runtime/MovieSceneTracks/Private/Tracks/MovieScenePropertyTrack.cpp b/Engine/Source/Runtime/MovieSceneTracks/Private/Tracks/MovieScenePropertyTrack.cpp
--- a/Engine/Source/Runtime/MovieSceneTracks/Private/Tracks/MovieScenePropertyTrack.cpp
+++ b/Engine/Source/Runtime/MovieSceneTracks/Private/Tracks/MovieScenePropertyTrack.cpp
@@ -12,7 +12,7 @@ UMovieScenePropertyTrack::UMovieScenePropertyTrack( const FObjectInitializer& Ob
 
 void UMovieScenePropertyTrack::SetPropertyNameAndPath( FName InPropertyName, const FString& InPropertyPath )
 {
-	check( InPropertyName != NAME_None && !InPropertyPath.IsEmpty() );
+	check((InPropertyName != NAME_None) && !InPropertyPath.IsEmpty());
 
 	PropertyName = InPropertyName;
 	PropertyPath = InPropertyPath;
@@ -25,6 +25,20 @@ const TArray<UMovieSceneSection*>& UMovieScenePropertyTrack::GetAllSections() co
 }
 
 
+#if WITH_EDITORONLY_DATA
+FText UMovieScenePropertyTrack::GetDisplayName() const
+{
+	return FText::FromName(PropertyName);
+}
+#endif
+
+
+FName UMovieScenePropertyTrack::GetTrackName() const
+{
+	return PropertyName;
+}
+
+
 void UMovieScenePropertyTrack::RemoveAllAnimationData()
 {
 	Sections.Empty();
diff --git a/Engine/Source/Runtime/MovieSceneTracks/Private/Tracks/MovieSceneShotTrack.cpp b/Engine/Source/Runtime/MovieSceneTracks/Private/Tracks/MovieSceneShotTrack.cpp
--- a/Engine/Source/Runtime/MovieSceneTracks/Private/Tracks/MovieSceneShotTrack.cpp
+++ b/Engine/Source/Runtime/MovieSceneTracks/Private/Tracks/MovieSceneShotTrack.cpp
@@ -66,14 +66,6 @@ const TArray<UMovieSceneSection*>& UMovieSceneShotTrack::GetAllSections() const
 }
 
 
-FName UMovieSceneShotTrack::GetTrackName() const
-{
-	static const FName UniqueName("Shots");
-
-	return UniqueName;
-}
-
-
 void UMovieSceneShotTrack::RemoveSection(UMovieSceneSection& Section)
 {
 	Sections.Remove(&Section);
@@ -84,6 +76,14 @@ void UMovieSceneShotTrack::RemoveSection(UMovieSceneSection& Section)
 }
 
 
+#if WITH_EDITORONLY_DATA
+FText UMovieSceneShotTrack::GetDisplayName() const
+{
+	return LOCTEXT("TrackName", "Shots");
+}
+#endif
+
+
 #if WITH_EDITOR
 void UMovieSceneShotTrack::OnSectionMoved(UMovieSceneSection& Section)
 {
diff --git a/Engine/Source/Runtime/MovieSceneTracks/Private/Tracks/MovieSceneSkeletalAnimationTrack.cpp b/Engine/Source/Runtime/MovieSceneTracks/Private/Tracks/MovieSceneSkeletalAnimationTrack.cpp
--- a/Engine/Source/Runtime/MovieSceneTracks/Private/Tracks/MovieSceneSkeletalAnimationTrack.cpp
+++ b/Engine/Source/Runtime/MovieSceneTracks/Private/Tracks/MovieSceneSkeletalAnimationTrack.cpp
@@ -10,17 +10,47 @@
 #define LOCTEXT_NAMESPACE "MovieSceneSkeletalAnimationTrack"
 
 
+/* UMovieSceneSkeletalAnimationTrack structors
+ *****************************************************************************/
+
 UMovieSceneSkeletalAnimationTrack::UMovieSceneSkeletalAnimationTrack( const FObjectInitializer& ObjectInitializer )
 	: Super( ObjectInitializer )
 { }
 
 
-FName UMovieSceneSkeletalAnimationTrack::GetTrackName() const
+/* UMovieSceneSkeletalAnimationTrack interface
+ *****************************************************************************/
+
+void UMovieSceneSkeletalAnimationTrack::AddNewAnimation(float KeyTime, UAnimSequence* AnimSequence)
 {
-	return FName("Animation");
+	UMovieSceneSkeletalAnimationSection* NewSection = NewObject<UMovieSceneSkeletalAnimationSection>(this);
+	{
+		NewSection->InitialPlacement(AnimationSections, KeyTime, KeyTime + AnimSequence->SequenceLength, SupportsMultipleRows());
+		NewSection->SetAnimSequence(AnimSequence);
+	}
+
+	AddSection(*NewSection);
+}
+
+
+UMovieSceneSection* UMovieSceneSkeletalAnimationTrack::GetAnimSectionAtTime(float Time)
+{
+	for (auto Section : AnimationSections)
+	{
+		if (Section->IsTimeWithinSection(Time))
+		{
+			return Section;
+		}
+	}
+
+	UMovieSceneSection* NearestSection = MovieSceneHelpers::FindNearestSectionAtTime( AnimationSections, Time );
+	return NearestSection;
 }
 
 
+/* UMovieSceneTrack interface
+ *****************************************************************************/
+
 TSharedPtr<IMovieSceneTrackInstance> UMovieSceneSkeletalAnimationTrack::CreateInstance()
 {
 	return MakeShareable( new FMovieSceneSkeletalAnimationTrackInstance( *this ) ); 
@@ -76,30 +106,14 @@ TRange<float> UMovieSceneSkeletalAnimationTrack::GetSectionBoundaries() const
 }
 
 
-void UMovieSceneSkeletalAnimationTrack::AddNewAnimation(float KeyTime, UAnimSequence* AnimSequence)
-{
-	UMovieSceneSkeletalAnimationSection* NewSection = NewObject<UMovieSceneSkeletalAnimationSection>(this);
-	{
-		NewSection->InitialPlacement(AnimationSections, KeyTime, KeyTime + AnimSequence->SequenceLength, SupportsMultipleRows());
-		NewSection->SetAnimSequence(AnimSequence);
-	}
-
-	AddSection(*NewSection);
-}
-
+#if WITH_EDITORONLY_DATA
 
-UMovieSceneSection* UMovieSceneSkeletalAnimationTrack::GetAnimSectionAtTime(float Time)
+FText UMovieSceneSkeletalAnimationTrack::GetDisplayName() const
 {
-	for (auto Section : AnimationSections)
-	{
-		if (Section->IsTimeWithinSection(Time))
-		{
-			return Section;
-		}
-	}
-
-	return nullptr;
+	return LOCTEXT("TrackName", "Animation");
 }
 
+#endif
+
 
 #undef LOCTEXT_NAMESPACE
diff --git a/Engine/Source/Runtime/MovieSceneTracks/Private/Tracks/MovieSceneSlomoTrack.cpp b/Engine/Source/Runtime/MovieSceneTracks/Private/Tracks/MovieSceneSlomoTrack.cpp
--- a/Engine/Source/Runtime/MovieSceneTracks/Private/Tracks/MovieSceneSlomoTrack.cpp
+++ b/Engine/Source/Runtime/MovieSceneTracks/Private/Tracks/MovieSceneSlomoTrack.cpp
@@ -6,6 +6,9 @@
 #include "MovieSceneSlomoTrackInstance.h"
 
 
+#define LOCTEXT_NAMESPACE "MovieSceneSlomoTrack"
+
+
 /* UMovieSceneEventTrack interface
  *****************************************************************************/
 
@@ -21,7 +24,14 @@ UMovieSceneSection* UMovieSceneSlomoTrack::CreateNewSection()
 }
 
 
-FName UMovieSceneSlomoTrack::GetTrackName() const
+#if WITH_EDITORONLY_DATA
+
+FText UMovieSceneSlomoTrack::GetDisplayName() const
 {
-	return FName("Play Rate");
+	return LOCTEXT("TrackName", "Play Rate");
 }
+
+#endif
+
+
+#undef LOCTEXT_NAMESPACE
diff --git a/Engine/Source/Runtime/MovieSceneTracks/Private/Tracks/MovieSceneSpawnTrack.cpp b/Engine/Source/Runtime/MovieSceneTracks/Private/Tracks/MovieSceneSpawnTrack.cpp
--- a/Engine/Source/Runtime/MovieSceneTracks/Private/Tracks/MovieSceneSpawnTrack.cpp
+++ b/Engine/Source/Runtime/MovieSceneTracks/Private/Tracks/MovieSceneSpawnTrack.cpp
@@ -6,22 +6,19 @@
 #include "MovieSceneBoolSection.h"
 
 
+#define LOCTEXT_NAMESPACE "MovieSceneSpawnTrack"
+
+
+/* UMovieSceneSpawnTrack structors
+ *****************************************************************************/
+
 UMovieSceneSpawnTrack::UMovieSceneSpawnTrack(const FObjectInitializer& Obj)
 	: Super(Obj)
-{
-}
+{ }
 
-UMovieSceneSection* UMovieSceneSpawnTrack::CreateNewSection()
-{
-	UMovieSceneBoolSection* Section = NewObject<UMovieSceneBoolSection>(this, NAME_None, RF_Transactional);
-	Section->DefaultValue = true;
-	return Section;
-}
 
-TSharedPtr<IMovieSceneTrackInstance> UMovieSceneSpawnTrack::CreateInstance()
-{
-	return MakeShareable( new FMovieSceneSpawnTrackInstance( *this ) );
-}
+/* UMovieSceneSpawnTrack interface
+ *****************************************************************************/
 
 bool UMovieSceneSpawnTrack::Eval(float Position, float LastPostion, bool& bOutSpawned) const
 {
@@ -38,4 +35,70 @@ bool UMovieSceneSpawnTrack::Eval(float Position, float LastPostion, bool& bOutSp
 
 	bOutSpawned = Section->Eval(Position);
 	return true;
-}
\ No newline at end of file
+}
+
+
+/* UMovieSceneTrack interface
+ *****************************************************************************/
+
+UMovieSceneSection* UMovieSceneSpawnTrack::CreateNewSection()
+{
+	UMovieSceneBoolSection* Section = NewObject<UMovieSceneBoolSection>(this, NAME_None, RF_Transactional);
+	Section->DefaultValue = true;
+	return Section;
+}
+
+
+TSharedPtr<IMovieSceneTrackInstance> UMovieSceneSpawnTrack::CreateInstance()
+{
+	return MakeShareable( new FMovieSceneSpawnTrackInstance( *this ) );
+}
+
+
+bool UMovieSceneSpawnTrack::HasSection(const UMovieSceneSection& Section ) const
+{
+	return Sections.ContainsByPredicate([&](const UMovieSceneSection* In){ return In == &Section; });
+}
+
+
+void UMovieSceneSpawnTrack::AddSection(UMovieSceneSection& Section)
+{
+	Sections.Add(&Section);
+}
+
+
+void UMovieSceneSpawnTrack::RemoveSection(UMovieSceneSection& Section)
+{
+	Sections.RemoveAll([&](const UMovieSceneSection* In){ return In == &Section; });
+}
+
+
+bool UMovieSceneSpawnTrack::IsEmpty() const
+{
+	return Sections.Num() == 0;
+}
+
+
+TRange<float> UMovieSceneSpawnTrack::GetSectionBoundaries() const
+{
+	return TRange<float>::All();
+}
+
+
+const TArray<UMovieSceneSection*>& UMovieSceneSpawnTrack::GetAllSections() const
+{
+	return Sections;
+}
+
+
+#if WITH_EDITORONLY_DATA
+
+FText UMovieSceneSpawnTrack::GetDisplayName() const
+{
+	return LOCTEXT("TrackName", "Spawned");
+}
+
+#endif
+
+
+#undef LOCTEXT_NAMESPACE
diff --git a/Engine/Source/Runtime/MovieSceneTracks/Private/Tracks/MovieSceneSubTrack.cpp b/Engine/Source/Runtime/MovieSceneTracks/Private/Tracks/MovieSceneSubTrack.cpp
--- a/Engine/Source/Runtime/MovieSceneTracks/Private/Tracks/MovieSceneSubTrack.cpp
+++ b/Engine/Source/Runtime/MovieSceneTracks/Private/Tracks/MovieSceneSubTrack.cpp
@@ -8,10 +8,7 @@
 #include "MovieSceneSubTrackInstance.h"
 
 
-namespace MovieSceneSubTrackConstants
-{
-	const FName TrackName("SubSequences");
-}
+#define LOCTEXT_NAMESPACE "MovieSceneSubTrack"
 
 
 /* UMovieSceneSubTrack interface
@@ -101,12 +98,6 @@ const TArray<UMovieSceneSection*>& UMovieSceneSubTrack::GetAllSections() const
 }
 
 
-FName UMovieSceneSubTrack::GetTrackName() const
-{
-	return MovieSceneSubTrackConstants::TrackName;
-}
-
-
 TRange<float> UMovieSceneSubTrack::GetSectionBoundaries() const
 {
 	TArray<TRange<float>> Bounds;
@@ -147,4 +138,7 @@ void UMovieSceneSubTrack::RemoveSection(UMovieSceneSection& Section)
 bool UMovieSceneSubTrack::SupportsMultipleRows() const
 {
 	return true;
-}
\ No newline at end of file
+}
+
+
+#undef LOCTEXT_NAMESPACE
diff --git a/Engine/Source/Runtime/MovieSceneTracks/Public/Sections/MovieScene3DAttachSection.h b/Engine/Source/Runtime/MovieSceneTracks/Public/Sections/MovieScene3DAttachSection.h
--- a/Engine/Source/Runtime/MovieSceneTracks/Public/Sections/MovieScene3DAttachSection.h
+++ b/Engine/Source/Runtime/MovieSceneTracks/Public/Sections/MovieScene3DAttachSection.h
@@ -32,9 +32,8 @@ public:
 	 */
 	void AddAttach( float Time, float SequenceEndTime, const FGuid& InAttachId);
 	
-private:
-
-	//@todo need additional actor component/bone to attach to
+	UPROPERTY(EditAnywhere, Category="Attach")
+	FName AttachSocketName;
 
 	UPROPERTY(EditAnywhere, Category="Attach")
 	uint32 bConstrainTx:1;
diff --git a/Engine/Source/Runtime/MovieSceneTracks/Public/TrackInstances/MovieScene3DConstraintTrackInstance.h b/Engine/Source/Runtime/MovieSceneTracks/Public/TrackInstances/MovieScene3DConstraintTrackInstance.h
--- a/Engine/Source/Runtime/MovieSceneTracks/Public/TrackInstances/MovieScene3DConstraintTrackInstance.h
+++ b/Engine/Source/Runtime/MovieSceneTracks/Public/TrackInstances/MovieScene3DConstraintTrackInstance.h
@@ -19,10 +19,10 @@ public:
 	/** IMovieSceneTrackInstance interface */
 	virtual void SaveState (const TArray<UObject*>& RuntimeObjects, IMovieScenePlayer& Player, FMovieSceneSequenceInstance& SequenceInstance) override;
 	virtual void RestoreState (const TArray<UObject*>& RuntimeObjects, IMovieScenePlayer& Player, FMovieSceneSequenceInstance& SequenceInstance) override;
-	virtual void Update( float Position, float LastPosition, const TArray<UObject*>& RuntimeObjects, IMovieScenePlayer& Player, FMovieSceneSequenceInstance& SequenceInstance ) override;
+	virtual void Update( float Position, float LastPosition, const TArray<UObject*>& RuntimeObjects, IMovieScenePlayer& Player, FMovieSceneSequenceInstance& SequenceInstance, EMovieSceneUpdatePass UpdatePass ) override;
 	virtual void RefreshInstance( const TArray<UObject*>& RuntimeObjects, IMovieScenePlayer& Player, FMovieSceneSequenceInstance& SequenceInstance ) override {}
 	virtual void ClearInstance( IMovieScenePlayer& Player, FMovieSceneSequenceInstance& SequenceInstance ) override {}
-	virtual float EvalOrder() { return 1.f; }
+	virtual float EvalOrder() { return -1.f; }
 
 	virtual void UpdateConstraint(float Position, const TArray<UObject*>& RuntimeObjects, AActor* Actor, UMovieScene3DConstraintSection* ConstraintSection) = 0;
 
diff --git a/Engine/Source/Runtime/MovieSceneTracks/Public/TrackInstances/MovieScene3DTransformTrackInstance.h b/Engine/Source/Runtime/MovieSceneTracks/Public/TrackInstances/MovieScene3DTransformTrackInstance.h
--- a/Engine/Source/Runtime/MovieSceneTracks/Public/TrackInstances/MovieScene3DTransformTrackInstance.h
+++ b/Engine/Source/Runtime/MovieSceneTracks/Public/TrackInstances/MovieScene3DTransformTrackInstance.h
@@ -21,7 +21,8 @@ public:
 	/** IMovieSceneTrackInstance interface */
 	virtual void SaveState (const TArray<UObject*>& RuntimeObjects, IMovieScenePlayer& Player, FMovieSceneSequenceInstance& SequenceInstance) override;
 	virtual void RestoreState (const TArray<UObject*>& RuntimeObjects, IMovieScenePlayer& Player, FMovieSceneSequenceInstance& SequenceInstance) override;
-	virtual void Update( float Position, float LastPosition, const TArray<UObject*>& RuntimeObjects, IMovieScenePlayer& Player, FMovieSceneSequenceInstance& SequenceInstance ) override;
+	virtual void Update( float Position, float LastPosition, const TArray<UObject*>& RuntimeObjects, IMovieScenePlayer& Player, FMovieSceneSequenceInstance& SequenceInstance, EMovieSceneUpdatePass UpdatePass ) override;
+	virtual EMovieSceneUpdatePass HasUpdatePasses() override { return (EMovieSceneUpdatePass)(MSUP_PreUpdate | MSUP_Update); }
 	virtual void RefreshInstance( const TArray<UObject*>& RuntimeObjects, IMovieScenePlayer& Player, FMovieSceneSequenceInstance& SequenceInstance ) override {}
 	virtual void ClearInstance( IMovieScenePlayer& Player, FMovieSceneSequenceInstance& SequenceInstance ) override {}
 
diff --git a/Engine/Source/Runtime/MovieSceneTracks/Public/TrackInstances/MovieSceneAudioTrackInstance.h b/Engine/Source/Runtime/MovieSceneTracks/Public/TrackInstances/MovieSceneAudioTrackInstance.h
--- a/Engine/Source/Runtime/MovieSceneTracks/Public/TrackInstances/MovieSceneAudioTrackInstance.h
+++ b/Engine/Source/Runtime/MovieSceneTracks/Public/TrackInstances/MovieSceneAudioTrackInstance.h
@@ -22,7 +22,7 @@ public:
 	/** IMovieSceneTrackInstance interface */
 	virtual void SaveState(const TArray<UObject*>& RuntimeObjects, IMovieScenePlayer& Player, FMovieSceneSequenceInstance& SequenceInstance) override {}
 	virtual void RestoreState(const TArray<UObject*>& RuntimeObjects, IMovieScenePlayer& Player, FMovieSceneSequenceInstance& SequenceInstance) override {}
-	virtual void Update( float Position, float LastPosition, const TArray<UObject*>& RuntimeObjects, IMovieScenePlayer& Player, FMovieSceneSequenceInstance& SequenceInstance ) override;
+	virtual void Update( float Position, float LastPosition, const TArray<UObject*>& RuntimeObjects, IMovieScenePlayer& Player, FMovieSceneSequenceInstance& SequenceInstance, EMovieSceneUpdatePass UpdatePass ) override;
 	virtual void RefreshInstance( const TArray<UObject*>& RuntimeObjects, IMovieScenePlayer& Player, FMovieSceneSequenceInstance& SequenceInstance ) override {}
 	virtual void ClearInstance( IMovieScenePlayer& Player, FMovieSceneSequenceInstance& SequenceInstance ) override {}
 
diff --git a/Engine/Source/Runtime/MovieSceneTracks/Public/TrackInstances/MovieSceneBoolTrackInstance.h b/Engine/Source/Runtime/MovieSceneTracks/Public/TrackInstances/MovieSceneBoolTrackInstance.h
--- a/Engine/Source/Runtime/MovieSceneTracks/Public/TrackInstances/MovieSceneBoolTrackInstance.h
+++ b/Engine/Source/Runtime/MovieSceneTracks/Public/TrackInstances/MovieSceneBoolTrackInstance.h
@@ -22,7 +22,7 @@ public:
 	/** IMovieSceneTrackInstance interface */
 	virtual void SaveState (const TArray<UObject*>& RuntimeObjects, IMovieScenePlayer& Player, FMovieSceneSequenceInstance& SequenceInstance) override;
 	virtual void RestoreState (const TArray<UObject*>& RuntimeObjects, IMovieScenePlayer& Player, FMovieSceneSequenceInstance& SequenceInstance) override;
-	virtual void Update( float Position, float LastPosition, const TArray<UObject*>& RuntimeObjects, IMovieScenePlayer& Player, FMovieSceneSequenceInstance& SequenceInstance ) override;
+	virtual void Update( float Position, float LastPosition, const TArray<UObject*>& RuntimeObjects, IMovieScenePlayer& Player, FMovieSceneSequenceInstance& SequenceInstance, EMovieSceneUpdatePass UpdatePass ) override;
 	virtual void RefreshInstance( const TArray<UObject*>& RuntimeObjects, IMovieScenePlayer& Player, FMovieSceneSequenceInstance& SequenceInstance ) override;
 	virtual void ClearInstance( IMovieScenePlayer& Player, FMovieSceneSequenceInstance& SequenceInstance ) override {}
 
diff --git a/Engine/Source/Runtime/MovieSceneTracks/Public/TrackInstances/MovieSceneEventTrackInstance.h b/Engine/Source/Runtime/MovieSceneTracks/Public/TrackInstances/MovieSceneEventTrackInstance.h
--- a/Engine/Source/Runtime/MovieSceneTracks/Public/TrackInstances/MovieSceneEventTrackInstance.h
+++ b/Engine/Source/Runtime/MovieSceneTracks/Public/TrackInstances/MovieSceneEventTrackInstance.h
@@ -34,7 +34,7 @@ public:
 	virtual void RefreshInstance(const TArray<UObject*>& RuntimeObjects, IMovieScenePlayer& Player, FMovieSceneSequenceInstance& SequenceInstance) override { }
 	virtual void RestoreState(const TArray<UObject*>& RuntimeObjects, IMovieScenePlayer& Player, FMovieSceneSequenceInstance& SequenceInstance) override { }
 	virtual void SaveState(const TArray<UObject*>& RuntimeObjects, IMovieScenePlayer& Player, FMovieSceneSequenceInstance& SequenceInstance) override { }
-	virtual void Update(float Position, float LastPosition, const TArray<UObject*>& RuntimeObjects, IMovieScenePlayer& Player, FMovieSceneSequenceInstance& SequenceInstance) override;
+	virtual void Update(float Position, float LastPosition, const TArray<UObject*>& RuntimeObjects, IMovieScenePlayer& Player, FMovieSceneSequenceInstance& SequenceInstance, EMovieSceneUpdatePass UpdatePass) override;
 
 private:
 
diff --git a/Engine/Source/Runtime/MovieSceneTracks/Public/TrackInstances/MovieSceneFadeTrackInstance.h b/Engine/Source/Runtime/MovieSceneTracks/Public/TrackInstances/MovieSceneFadeTrackInstance.h
--- a/Engine/Source/Runtime/MovieSceneTracks/Public/TrackInstances/MovieSceneFadeTrackInstance.h
+++ b/Engine/Source/Runtime/MovieSceneTracks/Public/TrackInstances/MovieSceneFadeTrackInstance.h
@@ -34,7 +34,7 @@ public:
 	virtual void RefreshInstance(const TArray<UObject*>& RuntimeObjects, IMovieScenePlayer& Player, FMovieSceneSequenceInstance& SequenceInstance) override { }
 	virtual void RestoreState(const TArray<UObject*>& RuntimeObjects, IMovieScenePlayer& Player, FMovieSceneSequenceInstance& SequenceInstance) override;
 	virtual void SaveState(const TArray<UObject*>& RuntimeObjects, IMovieScenePlayer& Player, FMovieSceneSequenceInstance& SequenceInstance) override { }
-	virtual void Update(float Position, float LastPosition, const TArray<UObject*>& RuntimeObjects, IMovieScenePlayer& Player, FMovieSceneSequenceInstance& SequenceInstance) override;
+	virtual void Update(float Position, float LastPosition, const TArray<UObject*>& RuntimeObjects, IMovieScenePlayer& Player, FMovieSceneSequenceInstance& SequenceInstance, EMovieSceneUpdatePass UpdatePass) override;
 
 private:
 
diff --git a/Engine/Source/Runtime/MovieSceneTracks/Public/TrackInstances/MovieSceneShotTrackInstance.h b/Engine/Source/Runtime/MovieSceneTracks/Public/TrackInstances/MovieSceneShotTrackInstance.h
--- a/Engine/Source/Runtime/MovieSceneTracks/Public/TrackInstances/MovieSceneShotTrackInstance.h
+++ b/Engine/Source/Runtime/MovieSceneTracks/Public/TrackInstances/MovieSceneShotTrackInstance.h
@@ -25,15 +25,27 @@ public:
 	virtual void RefreshInstance(const TArray<UObject*>& RuntimeObjects, IMovieScenePlayer& Player, FMovieSceneSequenceInstance& SequenceInstance) override;
 	virtual void RestoreState(const TArray<UObject*>& RuntimeObjects, IMovieScenePlayer& Player, FMovieSceneSequenceInstance& SequenceInstance) override { }
 	virtual void SaveState(const TArray<UObject*>& RuntimeObjects, IMovieScenePlayer& Player, FMovieSceneSequenceInstance& SequenceInstance) override { }
-	virtual void Update(float Position, float LastPosition, const TArray<UObject*>& RuntimeObjects, IMovieScenePlayer& Player, FMovieSceneSequenceInstance& SequenceInstance) override;
+	virtual void Update(float Position, float LastPosition, const TArray<UObject*>& RuntimeObjects, IMovieScenePlayer& Player, FMovieSceneSequenceInstance& SequenceInstance, EMovieSceneUpdatePass UpdatePass) override;
+
+protected:
+
+	/**
+	 * Acquire the camera object for the specified camera identifier.
+	 *
+	 * @param ShotIndex The index of the shot section that needs the camera.
+	 * @param CameraGuid The unique identifier of the camera object to get.
+	 * @param SequenceInstance The sequence instance that owns this track instance.
+	 * @return The camera object, or nullptr if not found.
+	 */
+	UObject* AcquireCameraForShot(int32 ShotIndex, const FGuid& CameraGuid, FMovieSceneSequenceInstance& SequenceInstance);
 
 private:
 	
 	/** Runtime camera objects.  One for each shot.  Must be the same number of entries as sections */
-	TArray< TWeakObjectPtr<UObject> > RuntimeCameraObjects;
+	TArray<TWeakObjectPtr<UObject>> CachedCameraObjects;
 
 	/** Current camera object we are looking through.  Used to determine when making a new cut */
-	TWeakObjectPtr<UObject> CurrentCameraObject;
+	TWeakObjectPtr<UObject> LastCameraObject;
 
 	/** Track that is being instanced */
 	TWeakObjectPtr<UMovieSceneShotTrack> ShotTrack;
diff --git a/Engine/Source/Runtime/MovieSceneTracks/Public/TrackInstances/MovieSceneSpawnTrackInstance.h b/Engine/Source/Runtime/MovieSceneTracks/Public/TrackInstances/MovieSceneSpawnTrackInstance.h
--- a/Engine/Source/Runtime/MovieSceneTracks/Public/TrackInstances/MovieSceneSpawnTrackInstance.h
+++ b/Engine/Source/Runtime/MovieSceneTracks/Public/TrackInstances/MovieSceneSpawnTrackInstance.h
@@ -17,12 +17,12 @@ public:
 	/** IMovieSceneTrackInstance interface */
 	virtual void SaveState(const TArray<UObject*>& RuntimeObjects, IMovieScenePlayer& Player, FMovieSceneSequenceInstance& SequenceInstance) override {}
 	virtual void RestoreState(const TArray<UObject*>& RuntimeObjects, IMovieScenePlayer& Player, FMovieSceneSequenceInstance& SequenceInstance) override;
-	virtual void Update(float Position, float LastPosition, const TArray<UObject*>& RuntimeObjects, class IMovieScenePlayer& Player, FMovieSceneSequenceInstance& SequenceInstance) override;
+	virtual void Update(float Position, float LastPosition, const TArray<UObject*>& RuntimeObjects, class IMovieScenePlayer& Player, FMovieSceneSequenceInstance& SequenceInstance, EMovieSceneUpdatePass UpdatePass) override;
 
 	// recreate/destroy the spawned object here?
 	virtual void RefreshInstance(const TArray<UObject*>& RuntimeObjects, class IMovieScenePlayer& Player, FMovieSceneSequenceInstance& SequenceInstance) override {}
 	virtual void ClearInstance(IMovieScenePlayer& Player, FMovieSceneSequenceInstance& SequenceInstance) override;
-	virtual float EvalOrder() override { return 10000.f; }
+	virtual float EvalOrder() override { return -10000.f; }
 
 private:
 	/** Spawn track that is being instanced */
diff --git a/Engine/Source/Runtime/MovieSceneTracks/Public/TrackInstances/MovieSceneVisibilityTrackInstance.h b/Engine/Source/Runtime/MovieSceneTracks/Public/TrackInstances/MovieSceneVisibilityTrackInstance.h
--- a/Engine/Source/Runtime/MovieSceneTracks/Public/TrackInstances/MovieSceneVisibilityTrackInstance.h
+++ b/Engine/Source/Runtime/MovieSceneTracks/Public/TrackInstances/MovieSceneVisibilityTrackInstance.h
@@ -21,7 +21,7 @@ public:
 	/** IMovieSceneTrackInstance interface */
 	virtual void SaveState (const TArray<UObject*>& RuntimeObjects, IMovieScenePlayer& Player, FMovieSceneSequenceInstance& SequenceInstance) override;
 	virtual void RestoreState (const TArray<UObject*>& RuntimeObjects, IMovieScenePlayer& Player, FMovieSceneSequenceInstance& SequenceInstance) override;
-	virtual void Update( float Position, float LastPosition, const TArray<UObject*>& RuntimeObjects, class IMovieScenePlayer& Player, FMovieSceneSequenceInstance& SequenceInstance ) override;
+	virtual void Update( float Position, float LastPosition, const TArray<UObject*>& RuntimeObjects, class IMovieScenePlayer& Player, FMovieSceneSequenceInstance& SequenceInstance, EMovieSceneUpdatePass UpdatePass ) override;
 	virtual void RefreshInstance( const TArray<UObject*>& RuntimeObjects, class IMovieScenePlayer& Player, FMovieSceneSequenceInstance& SequenceInstance ) override;
 	virtual void ClearInstance( IMovieScenePlayer& Player, FMovieSceneSequenceInstance& SequenceInstance ) override {}
 
diff --git a/Engine/Source/Runtime/MovieSceneTracks/Public/Tracks/MovieScene3DAttachTrack.h b/Engine/Source/Runtime/MovieSceneTracks/Public/Tracks/MovieScene3DAttachTrack.h
--- a/Engine/Source/Runtime/MovieSceneTracks/Public/Tracks/MovieScene3DAttachTrack.h
+++ b/Engine/Source/Runtime/MovieSceneTracks/Public/Tracks/MovieScene3DAttachTrack.h
@@ -6,24 +6,29 @@
 #include "MovieScene3DConstraintTrack.h"
 #include "MovieScene3DAttachTrack.generated.h"
 
+
 /**
- * Handles manipulation of path tracks in a movie scene
+ * Handles manipulation of path tracks in a movie scene.
  */
 UCLASS( MinimalAPI )
-class UMovieScene3DAttachTrack : public UMovieScene3DConstraintTrack
+class UMovieScene3DAttachTrack
+	: public UMovieScene3DConstraintTrack
 {
 	GENERATED_UCLASS_BODY()
+
+public:
+
+	// UMovieScene3DConstraintTrack interface
+
+	virtual void AddConstraint( float Time, float ConstraintEndTime, const FName SocketName, const FGuid& ConstraintId ) override;
+
 public:
-	/** UMovieSceneTrack interface */
-	virtual FName GetTrackName() const override;
+
+	// UMovieSceneTrack interface
+
 	virtual TSharedPtr<IMovieSceneTrackInstance> CreateInstance() override;
 
-	/**
-	 * Adds a constraint
-	 *
-	 * @param Time				    The time relative to the owning movie scene where the section should be
-	 * @param ConstraintEndTime     Set the constraint to end at this time
-	 * @param ConstraintId			The id to the constraint
-	 */
-	virtual void AddConstraint( float Time, float ConstraintEndTime, const FGuid& ConstraintId ) override;
+#if WITH_EDITORONLY_DATA
+	virtual FText GetDisplayName() const override;
+#endif
 };
diff --git a/Engine/Source/Runtime/MovieSceneTracks/Public/Tracks/MovieScene3DConstraintTrack.h b/Engine/Source/Runtime/MovieSceneTracks/Public/Tracks/MovieScene3DConstraintTrack.h
--- a/Engine/Source/Runtime/MovieSceneTracks/Public/Tracks/MovieScene3DConstraintTrack.h
+++ b/Engine/Source/Runtime/MovieSceneTracks/Public/Tracks/MovieScene3DConstraintTrack.h
@@ -2,21 +2,38 @@
 
 #pragma once
 
-#include "MovieScene.h"
 #include "MovieSceneTrack.h"
 #include "MovieScene3DConstraintTrack.generated.h"
 
 
+class UMovieSceneSection;
+
+
 /**
- * Base class for constraint tracks - tracks that are dependent upon other objects
+ * Base class for constraint tracks (tracks that are dependent upon other objects).
  */
 UCLASS( MinimalAPI )
 class UMovieScene3DConstraintTrack
 	: public UMovieSceneTrack
 {
 	GENERATED_UCLASS_BODY()
+
+public:
+
+	/**
+	 * Adds a constraint.
+	 *
+	 * @param Time The time relative to the owning movie scene where the section should be.
+	 * @param ConstraintEndTime Set the constraint to end at this time.
+	 * @param SocketName The socket name for the constraint.
+	 * @param ConstraintId The id to the constraint.
+	 */
+	virtual void AddConstraint(float Time, float ConstraintEndTime, const FName SocketName, const FGuid& ConstraintId) { }
+
 public:
-	/** UMovieSceneTrack interface */
+
+	// UMovieSceneTrack interface
+
     virtual void RemoveAllAnimationData() override;
 	virtual bool HasSection(const UMovieSceneSection& Section) const override;
 	virtual void AddSection(UMovieSceneSection& Section) override;
@@ -25,18 +42,13 @@ public:
 	virtual TRange<float> GetSectionBoundaries() const override;
 	virtual const TArray<UMovieSceneSection*>& GetAllSections() const override;
 
-	/**
-	 * Adds a constraint
-	 *
-	 * @param Time				    The time relative to the owning movie scene where the section should be
-	 * @param ConstraintEndTime     Set the constraint to end at this time
-	 * @param ConstraintId			The id to the constraint
-	 */
-	virtual void AddConstraint( float Time, float ConstraintEndTime, const FGuid& ConstraintId ) {}
+#if WITH_EDITORONLY_DATA
+	virtual FText GetDisplayName() const override;
+#endif
 
 protected:
 
-	/** List of all constraint sections */
+	/** List of all constraint sections. */
 	UPROPERTY()
 	TArray<UMovieSceneSection*> ConstraintSections;
 };
diff --git a/Engine/Source/Runtime/MovieSceneTracks/Public/Tracks/MovieScene3DPathTrack.h b/Engine/Source/Runtime/MovieSceneTracks/Public/Tracks/MovieScene3DPathTrack.h
--- a/Engine/Source/Runtime/MovieSceneTracks/Public/Tracks/MovieScene3DPathTrack.h
+++ b/Engine/Source/Runtime/MovieSceneTracks/Public/Tracks/MovieScene3DPathTrack.h
@@ -2,28 +2,32 @@
 
 #pragma once
 
-#include "MovieScene.h"
 #include "MovieScene3DConstraintTrack.h"
 #include "MovieScene3DPathTrack.generated.h"
 
+
 /**
  * Handles manipulation of path tracks in a movie scene
  */
-UCLASS( MinimalAPI )
-class UMovieScene3DPathTrack : public UMovieScene3DConstraintTrack
+UCLASS(MinimalAPI)
+class UMovieScene3DPathTrack
+	: public UMovieScene3DConstraintTrack
 {
 	GENERATED_UCLASS_BODY()
+
+public:
+
+	// UMovieScene3DConstraintTrack interface
+
+	virtual void AddConstraint(float Time, float ConstraintEndTime,  const FName SocketName, const FGuid& ConstraintId) override;
+
 public:
-	/** UMovieSceneTrack interface */
-	virtual FName GetTrackName() const override;
+
+	// UMovieSceneTrack interface
+
 	virtual TSharedPtr<IMovieSceneTrackInstance> CreateInstance() override;
 
-	/**
-	 * Adds a constraint
-	 *
-	 * @param Time				    The time relative to the owning movie scene where the section should be
-	 * @param ConstraintEndTime     Set the constraint to end at this time
-	 * @param ConstraintId			The id to the constraint
-	 */
-	virtual void AddConstraint( float Time, float ConstraintEndTime, const FGuid& ConstraintId ) override;
+#if WITH_EDITORONLY_DATA
+	virtual FText GetDisplayName() const override;
+#endif
 };
diff --git a/Engine/Source/Runtime/MovieSceneTracks/Public/Tracks/MovieSceneAudioTrack.h b/Engine/Source/Runtime/MovieSceneTracks/Public/Tracks/MovieSceneAudioTrack.h
--- a/Engine/Source/Runtime/MovieSceneTracks/Public/Tracks/MovieSceneAudioTrack.h
+++ b/Engine/Source/Runtime/MovieSceneTracks/Public/Tracks/MovieSceneAudioTrack.h
@@ -2,32 +2,47 @@
 
 #pragma once
 
-#include "MovieScene.h"
-#include "MovieSceneTrack.h"
+#include "MovieSceneNameableTrack.h"
 #include "MovieSceneAudioTrack.generated.h"
 
 
+class UMovieSceneSection;
+class USoundBase;
+
+
 namespace AudioTrackConstants
 {
 	const float ScrubDuration = 0.050f;
-	const FName UniqueTrackName("MasterAudio");
 }
 
 
 /**
- * Handles manipulation of audio
+ * Handles manipulation of audio.
  */
-UCLASS( MinimalAPI )
+UCLASS(MinimalAPI)
 class UMovieSceneAudioTrack
-	: public UMovieSceneTrack
+	: public UMovieSceneNameableTrack
 {
 	GENERATED_UCLASS_BODY()
 
+public:
+
+	/** Adds a new sound cue to the audio */
+	virtual void AddNewSound(USoundBase* Sound, float Time);
+
+	/** @return The audio sections on this track */
+	const TArray<UMovieSceneSection*>& GetAudioSections() const
+	{
+		return AudioSections;
+	}
+
+	/** @return true if this is a master audio track */
+	bool IsAMasterTrack() const;
+
 public:
 
 	// UMovieSceneTrack interface
 
-	virtual FName GetTrackName() const override;
 	virtual TSharedPtr<IMovieSceneTrackInstance> CreateInstance() override;
 	virtual void RemoveAllAnimationData() override;
 	virtual bool HasSection(const UMovieSceneSection& Section) const override;
@@ -35,19 +50,8 @@ public:
 	virtual void RemoveSection(UMovieSceneSection& Section) override;
 	virtual bool IsEmpty() const override;
 	virtual TRange<float> GetSectionBoundaries() const override;
-	virtual bool SupportsMultipleRows() const override { return true; }
 	virtual const TArray<UMovieSceneSection*>& GetAllSections() const override;
-
-public:
-
-	/** Adds a new sound cue to the audio */
-	virtual void AddNewSound(class USoundBase* Sound, float Time);
-
-	/** @return The audio sections on this track */
-	const TArray<UMovieSceneSection*>& GetAudioSections() const { return AudioSections; }
-
-	/** @return true if this is a master audio track */
-	bool IsAMasterTrack() const;
+	virtual bool SupportsMultipleRows() const override;
 
 private:
 
diff --git a/Engine/Source/Runtime/MovieSceneTracks/Public/Tracks/MovieSceneEventTrack.h b/Engine/Source/Runtime/MovieSceneTracks/Public/Tracks/MovieSceneEventTrack.h
--- a/Engine/Source/Runtime/MovieSceneTracks/Public/Tracks/MovieSceneEventTrack.h
+++ b/Engine/Source/Runtime/MovieSceneTracks/Public/Tracks/MovieSceneEventTrack.h
@@ -2,7 +2,7 @@
 
 #pragma once
 
-#include "MovieSceneTrack.h"
+#include "MovieSceneNameableTrack.h"
 #include "MovieSceneEventTrack.generated.h"
 
 
@@ -11,7 +11,7 @@
  */
 UCLASS(MinimalAPI)
 class UMovieSceneEventTrack
-	: public UMovieSceneTrack
+	: public UMovieSceneNameableTrack
 {
 	GENERATED_BODY()
 
@@ -21,7 +21,6 @@ public:
 	UMovieSceneEventTrack()
 		: bFireEventsWhenForwards(true)
 		, bFireEventsWhenBackwards(true)
-		, TrackName("Events")
 	{ }
 
 public:
@@ -54,7 +53,6 @@ public:
 	virtual UMovieSceneSection* CreateNewSection() override;
 	virtual const TArray<UMovieSceneSection*>& GetAllSections() const override;
 	virtual TRange<float> GetSectionBoundaries() const override;
-	virtual FName GetTrackName() const override;
 	virtual bool HasSection(const UMovieSceneSection& Section) const override;
 	virtual bool IsEmpty() const override;
 	virtual void RemoveAllAnimationData() override;
@@ -73,8 +71,4 @@ private:
 	/** The track's sections. */
 	UPROPERTY()
 	TArray<UMovieSceneSection*> Sections;
-
-	/** Name of this track. */
-	UPROPERTY()
-	FName TrackName;
 };
diff --git a/Engine/Source/Runtime/MovieSceneTracks/Public/Tracks/MovieSceneFadeTrack.h b/Engine/Source/Runtime/MovieSceneTracks/Public/Tracks/MovieSceneFadeTrack.h
--- a/Engine/Source/Runtime/MovieSceneTracks/Public/Tracks/MovieSceneFadeTrack.h
+++ b/Engine/Source/Runtime/MovieSceneTracks/Public/Tracks/MovieSceneFadeTrack.h
@@ -21,5 +21,8 @@ public:
 
 	virtual TSharedPtr<IMovieSceneTrackInstance> CreateInstance() override;
 	virtual UMovieSceneSection* CreateNewSection() override;
-	virtual FName GetTrackName() const override;
+
+#if WITH_EDITORONLY_DATA
+	virtual FText GetDisplayName() const override;
+#endif
 };
diff --git a/Engine/Source/Runtime/MovieSceneTracks/Public/Tracks/MovieSceneMaterialTrack.h b/Engine/Source/Runtime/MovieSceneTracks/Public/Tracks/MovieSceneMaterialTrack.h
--- a/Engine/Source/Runtime/MovieSceneTracks/Public/Tracks/MovieSceneMaterialTrack.h
+++ b/Engine/Source/Runtime/MovieSceneTracks/Public/Tracks/MovieSceneMaterialTrack.h
@@ -116,7 +116,10 @@ public:
 	// UMovieSceneTrack interface
 
 	virtual TSharedPtr<IMovieSceneTrackInstance> CreateInstance() override;
-	virtual FName GetTrackName() const override { return TrackName; }
+
+#if WITH_EDITORONLY_DATA
+	virtual FText GetDisplayName() const override;
+#endif
 
 public:
 
@@ -127,15 +130,10 @@ public:
 	void SetMaterialIndex(int32 InMaterialIndex) 
 	{
 		MaterialIndex = InMaterialIndex;
-		TrackName = *FString::Printf(TEXT("Material Element %i"), MaterialIndex);
 	}
 
 private:
 
-	/** The name of this track .*/
-	UPROPERTY()
-	FName TrackName;
-
 	/** The index of this material this track is animating. */
 	UPROPERTY()
 	int32 MaterialIndex;
diff --git a/Engine/Source/Runtime/MovieSceneTracks/Public/Tracks/MovieSceneParticleTrack.h b/Engine/Source/Runtime/MovieSceneTracks/Public/Tracks/MovieSceneParticleTrack.h
--- a/Engine/Source/Runtime/MovieSceneTracks/Public/Tracks/MovieSceneParticleTrack.h
+++ b/Engine/Source/Runtime/MovieSceneTracks/Public/Tracks/MovieSceneParticleTrack.h
@@ -2,21 +2,38 @@
 
 #pragma once
 
-#include "MovieScene.h"
-#include "MovieSceneTrack.h"
+#include "MovieSceneNameableTrack.h"
 #include "MovieSceneParticleTrack.generated.h"
 
+
+class UMovieSceneSection;
+
+
 /**
  * Handles triggering of particle emitters
  */
-UCLASS( MinimalAPI )
-class UMovieSceneParticleTrack : public UMovieSceneTrack
+UCLASS(MinimalAPI)
+class UMovieSceneParticleTrack
+	: public UMovieSceneNameableTrack
 {
 	GENERATED_UCLASS_BODY()
 
 public:
-	/** UMovieSceneTrack interface */
-	virtual FName GetTrackName() const override;
+
+	/**
+	 * Get the track's particle sections.
+	 *
+	 * @return Particle sections collection.
+	 */
+	virtual TArray<UMovieSceneSection*> GetAllParticleSections() const
+	{
+		return ParticleSections;
+	}
+
+public:
+
+	// UMovieSceneTrack interface
+
 	virtual TSharedPtr<IMovieSceneTrackInstance> CreateInstance() override;
 	virtual void RemoveAllAnimationData() override;
 	virtual bool HasSection(const UMovieSceneSection& Section) const override;
@@ -25,13 +42,11 @@ public:
 	virtual bool IsEmpty() const override;
 	virtual TRange<float> GetSectionBoundaries() const override;
 	virtual const TArray<UMovieSceneSection*>& GetAllSections() const override;
-	
-	virtual void AddNewSection( float SectionTime );
-
-	virtual TArray<UMovieSceneSection*> GetAllParticleSections() const {return ParticleSections;}
+	virtual void AddNewSection(float SectionTime);
 
 private:
-	/** List of all particle sections */
+
+	/** List of all particle sections. */
 	UPROPERTY()
 	TArray<UMovieSceneSection*> ParticleSections;
 };
diff --git a/Engine/Source/Runtime/MovieSceneTracks/Public/Tracks/MovieScenePropertyTrack.h b/Engine/Source/Runtime/MovieSceneTracks/Public/Tracks/MovieScenePropertyTrack.h
--- a/Engine/Source/Runtime/MovieSceneTracks/Public/Tracks/MovieScenePropertyTrack.h
+++ b/Engine/Source/Runtime/MovieSceneTracks/Public/Tracks/MovieScenePropertyTrack.h
@@ -20,7 +20,7 @@ public:
 
 	// UMovieSceneTrack interface
 
-	virtual FName GetTrackName() const override { return PropertyName; }
+	virtual FName GetTrackName() const override;
 	virtual void RemoveAllAnimationData() override;
 	virtual bool HasSection(const UMovieSceneSection& Section ) const override;
 	virtual void AddSection(UMovieSceneSection& Section) override;
@@ -29,6 +29,10 @@ public:
 	virtual TRange<float> GetSectionBoundaries() const override;
 	virtual const TArray<UMovieSceneSection*>& GetAllSections() const override;
 
+#if WITH_EDITORONLY_DATA
+	virtual FText GetDisplayName() const override;
+#endif
+
 public:
 
 	/**
diff --git a/Engine/Source/Runtime/MovieSceneTracks/Public/Tracks/MovieSceneShotTrack.h b/Engine/Source/Runtime/MovieSceneTracks/Public/Tracks/MovieSceneShotTrack.h
--- a/Engine/Source/Runtime/MovieSceneTracks/Public/Tracks/MovieSceneShotTrack.h
+++ b/Engine/Source/Runtime/MovieSceneTracks/Public/Tracks/MovieSceneShotTrack.h
@@ -39,9 +39,12 @@ public:
 	virtual TSharedPtr<IMovieSceneTrackInstance> CreateInstance() override;
 	virtual UMovieSceneSection* CreateNewSection() override;
 	virtual const TArray<UMovieSceneSection*>& GetAllSections() const override;
-	virtual FName GetTrackName() const override;
 	virtual void RemoveSection(UMovieSceneSection& Section) override;
 
+#if WITH_EDITORONLY_DATA
+	virtual FText GetDisplayName() const override;
+#endif
+
 #if WITH_EDITOR
 	virtual void OnSectionMoved(UMovieSceneSection& Section) override;
 #endif
diff --git a/Engine/Source/Runtime/MovieSceneTracks/Public/Tracks/MovieSceneSkeletalAnimationTrack.h b/Engine/Source/Runtime/MovieSceneTracks/Public/Tracks/MovieSceneSkeletalAnimationTrack.h
--- a/Engine/Source/Runtime/MovieSceneTracks/Public/Tracks/MovieSceneSkeletalAnimationTrack.h
+++ b/Engine/Source/Runtime/MovieSceneTracks/Public/Tracks/MovieSceneSkeletalAnimationTrack.h
@@ -7,6 +7,9 @@
 #include "MovieSceneSkeletalAnimationTrack.generated.h"
 
 
+class UMovieSceneSection;
+
+
 /**
  * Handles animation of skeletal mesh actors
  */
@@ -16,11 +19,18 @@ class UMovieSceneSkeletalAnimationTrack
 {
 	GENERATED_UCLASS_BODY()
 
+public:
+
+	/** Adds a new animation to this track */
+	virtual void AddNewAnimation(float KeyTime, class UAnimSequence* AnimSequence);
+
+	/** Gets the animation section at a certain time, or NULL if there is none */
+	UMovieSceneSection* GetAnimSectionAtTime(float Time);
+
 public:
 
 	// UMovieSceneTrack interface
 
-	virtual FName GetTrackName() const override;
 	virtual TSharedPtr<IMovieSceneTrackInstance> CreateInstance() override;
 	virtual void RemoveAllAnimationData() override;
 	virtual bool HasSection(const UMovieSceneSection& Section ) const override;
@@ -31,13 +41,9 @@ public:
 	virtual bool SupportsMultipleRows() const override { return true; }
 	virtual const TArray<UMovieSceneSection*>& GetAllSections() const override;
 
-public:
-
-	/** Adds a new animation to this track */
-	virtual void AddNewAnimation(float KeyTime, class UAnimSequence* AnimSequence);
-
-	/** Gets the animation section at a certain time, or NULL if there is none */
-	class UMovieSceneSection* GetAnimSectionAtTime(float Time);
+#if WITH_EDITORONLY_DATA
+	virtual FText GetDisplayName() const override;
+#endif
 
 private:
 
diff --git a/Engine/Source/Runtime/MovieSceneTracks/Public/Tracks/MovieSceneSlomoTrack.h b/Engine/Source/Runtime/MovieSceneTracks/Public/Tracks/MovieSceneSlomoTrack.h
--- a/Engine/Source/Runtime/MovieSceneTracks/Public/Tracks/MovieSceneSlomoTrack.h
+++ b/Engine/Source/Runtime/MovieSceneTracks/Public/Tracks/MovieSceneSlomoTrack.h
@@ -21,5 +21,8 @@ public:
 
 	virtual TSharedPtr<IMovieSceneTrackInstance> CreateInstance() override;
 	virtual UMovieSceneSection* CreateNewSection() override;
-	virtual FName GetTrackName() const override;
+
+#if WITH_EDITORONLY_DATA
+	virtual FText GetDisplayName() const override;
+#endif
 };
diff --git a/Engine/Source/Runtime/MovieSceneTracks/Public/Tracks/MovieSceneSpawnTrack.h b/Engine/Source/Runtime/MovieSceneTracks/Public/Tracks/MovieSceneSpawnTrack.h
--- a/Engine/Source/Runtime/MovieSceneTracks/Public/Tracks/MovieSceneSpawnTrack.h
+++ b/Engine/Source/Runtime/MovieSceneTracks/Public/Tracks/MovieSceneSpawnTrack.h
@@ -2,10 +2,13 @@
 
 #pragma once
 
-#include "MovieScene.h"
 #include "MovieSceneTrack.h"
 #include "MovieSceneSpawnTrack.generated.h"
 
+
+class UMovieSceneSection;
+
+
 /**
  * Handles when a spawnable should be spawned and destroyed
  */
@@ -20,45 +23,13 @@ public:
 
 public:
 
-	// UMovieSceneTrack interface
-	virtual UMovieSceneSection* CreateNewSection() override;
-
-	virtual TSharedPtr<IMovieSceneTrackInstance> CreateInstance() override;
-	
-	virtual FName GetTrackName() const override
-	{
-		static FName Name = "Spawned";
-		return Name;
-	}
-
-	virtual bool HasSection(const UMovieSceneSection& Section ) const override
-	{
-		return Sections.ContainsByPredicate([&](const UMovieSceneSection* In){ return In == &Section; });
-	}
-
-	virtual void AddSection(UMovieSceneSection& Section) override
-	{
-		Sections.Add(&Section);
-	}
-
-	virtual void RemoveSection(UMovieSceneSection& Section) override
-	{
-		Sections.RemoveAll([&](const UMovieSceneSection* In){ return In == &Section; });
-	}
-
-	virtual bool IsEmpty() const override
-	{
-		return Sections.Num() == 0;
-	}
-
-	virtual TRange<float> GetSectionBoundaries() const override
-	{
-		return TRange<float>::All();
-	}
+	/** Evaluate whether the controlled object should currently be spawned or not */
+	bool Eval(float Position, float LastPostion, bool& bOutSpawned) const;
 
-	virtual const TArray<UMovieSceneSection*>& GetAllSections() const override
+	/** Get the object identifier that this spawn track controls */
+	const FGuid& GetObject() const
 	{
-		return Sections;
+		return ObjectGuid;
 	}
 
 	/** Set the object identifier that this spawn track controls */
@@ -67,14 +38,22 @@ public:
 		ObjectGuid = InGuid;
 	}
 
-	/** Get the object identifier that this spawn track controls */
-	const FGuid& GetObject() const
-	{
-		return ObjectGuid;
-	}
+public:
 
-	/** Evaluate whether the controlled object should currently be spawned or not */
-	bool Eval(float Position, float LastPostion, bool& bOutSpawned) const;
+	// UMovieSceneTrack interface
+
+	virtual UMovieSceneSection* CreateNewSection() override;
+	virtual TSharedPtr<IMovieSceneTrackInstance> CreateInstance() override;
+	virtual bool HasSection(const UMovieSceneSection& Section ) const override;
+	virtual void AddSection(UMovieSceneSection& Section) override;
+	virtual void RemoveSection(UMovieSceneSection& Section) override;
+	virtual bool IsEmpty() const override;
+	virtual TRange<float> GetSectionBoundaries() const override;
+	virtual const TArray<UMovieSceneSection*>& GetAllSections() const override;
+
+#if WITH_EDITORONLY_DATA
+	virtual FText GetDisplayName() const override;
+#endif
 
 protected:
 
diff --git a/Engine/Source/Runtime/MovieSceneTracks/Public/Tracks/MovieSceneSubTrack.h b/Engine/Source/Runtime/MovieSceneTracks/Public/Tracks/MovieSceneSubTrack.h
--- a/Engine/Source/Runtime/MovieSceneTracks/Public/Tracks/MovieSceneSubTrack.h
+++ b/Engine/Source/Runtime/MovieSceneTracks/Public/Tracks/MovieSceneSubTrack.h
@@ -2,7 +2,7 @@
 
 #pragma once
 
-#include "MovieSceneTrack.h"
+#include "MovieSceneNameableTrack.h"
 #include "MovieSceneSubTrack.generated.h"
 
 
@@ -15,7 +15,7 @@ class UMovieSceneSequence;
  */
 UCLASS(MinimalAPI)
 class UMovieSceneSubTrack
-	: public UMovieSceneTrack
+	: public UMovieSceneNameableTrack
 {
 	GENERATED_BODY()
 
@@ -47,7 +47,6 @@ public:
 	virtual UMovieSceneSection* CreateNewSection() override;
 	virtual const TArray<UMovieSceneSection*>& GetAllSections() const override;
 	virtual TRange<float> GetSectionBoundaries() const override;
-	virtual FName GetTrackName() const override;
 	virtual bool HasSection(const UMovieSceneSection& Section) const override;
 	virtual bool IsEmpty() const override;
 	virtual void RemoveAllAnimationData() override;
diff --git a/Engine/Source/Runtime/NetworkReplayStreaming/HttpNetworkReplayStreaming/Private/HttpNetworkReplayStreaming.cpp b/Engine/Source/Runtime/NetworkReplayStreaming/HttpNetworkReplayStreaming/Private/HttpNetworkReplayStreaming.cpp
--- a/Engine/Source/Runtime/NetworkReplayStreaming/HttpNetworkReplayStreaming/Private/HttpNetworkReplayStreaming.cpp
+++ b/Engine/Source/Runtime/NetworkReplayStreaming/HttpNetworkReplayStreaming/Private/HttpNetworkReplayStreaming.cpp
@@ -688,8 +688,9 @@ bool FQueuedHttpRequestAddEvent::PreProcess( const FString& ServerURL, const FSt
 
 	if ( !Name.IsEmpty() )
 	{
-		// Update existing event
-		Request->SetURL( FString::Printf( TEXT( "%sevent/%s" ), *ServerURL, *Name ) );
+		// Add or update existing event
+		const FString EventName = SessionName + TEXT( "_" ) + Name;
+		Request->SetURL( FString::Printf( TEXT( "%sreplay/%s/event/%s?group=%s&time1=%i&time2=%i&meta=%s" ), *ServerURL, *SessionName, *EventName, *Group, TimeInMS, TimeInMS, *FGenericPlatformHttp::UrlEncode( Meta ) ) );
 	}
 	else
 	{
