Copying //UE4/Dev-Sequencer to Dev-Main (//UE4/Dev-Main)

diff --git a/Engine/Source/Runtime/Core/Private/Math/BasicMathExpressionEvaluator.cpp b/Engine/Source/Runtime/Core/Private/Math/BasicMathExpressionEvaluator.cpp
--- a/Engine/Source/Runtime/Core/Private/Math/BasicMathExpressionEvaluator.cpp
+++ b/Engine/Source/Runtime/Core/Private/Math/BasicMathExpressionEvaluator.cpp
@@ -10,9 +10,13 @@ namespace ExpressionParser
 	const TCHAR* const FSubExpressionStart::Moniker = TEXT("(");
 	const TCHAR* const FSubExpressionEnd::Moniker = TEXT(")");
 	const TCHAR* const FPlus::Moniker = TEXT("+");
+	const TCHAR* const FPlusEquals::Moniker = TEXT("+=");
 	const TCHAR* const FMinus::Moniker = TEXT("-");
+	const TCHAR* const FMinusEquals::Moniker = TEXT("-=");
 	const TCHAR* const FStar::Moniker = TEXT("*");
+	const TCHAR* const FStarEquals::Moniker = TEXT("*=");
 	const TCHAR* const FForwardSlash::Moniker = TEXT("/");
+	const TCHAR* const FForwardSlashEquals::Moniker = TEXT("/=");
 	const TCHAR* const FPercent::Moniker = TEXT("%");
 	const TCHAR* const FSquareRoot::Moniker = TEXT("sqrt");
 	const TCHAR* const FPower::Moniker = TEXT("^");
@@ -121,6 +125,10 @@ FBasicMathExpressionEvaluator::FBasicMathExpressionEvaluator()
 	TokenDefinitions.IgnoreWhitespace();
 	TokenDefinitions.DefineToken(&ConsumeSymbol<FSubExpressionStart>);
 	TokenDefinitions.DefineToken(&ConsumeSymbol<FSubExpressionEnd>);
+	TokenDefinitions.DefineToken(&ConsumeSymbol<FPlusEquals>);
+	TokenDefinitions.DefineToken(&ConsumeSymbol<FMinusEquals>);
+	TokenDefinitions.DefineToken(&ConsumeSymbol<FStarEquals>);
+	TokenDefinitions.DefineToken(&ConsumeSymbol<FForwardSlashEquals>);
 	TokenDefinitions.DefineToken(&ConsumeSymbol<FPlus>);
 	TokenDefinitions.DefineToken(&ConsumeSymbol<FMinus>);
 	TokenDefinitions.DefineToken(&ConsumeSymbol<FStar>);
@@ -168,9 +176,60 @@ FBasicMathExpressionEvaluator::FBasicMathExpressionEvaluator()
 	});
 }
 
-TValueOrError<double, FExpressionError> FBasicMathExpressionEvaluator::Evaluate(const TCHAR* InExpression) const
+TValueOrError<double, FExpressionError> FBasicMathExpressionEvaluator::Evaluate(const TCHAR* InExpression, double InExistingValue) const
 {
-	TValueOrError<FExpressionNode, FExpressionError> Result = ExpressionParser::Evaluate(InExpression, TokenDefinitions, Grammar, JumpTable);
+	using namespace ExpressionParser;
+
+	TValueOrError<TArray<FExpressionToken>, FExpressionError> LexResult = ExpressionParser::Lex(InExpression, TokenDefinitions);
+	if (!LexResult.IsValid())
+	{
+		return MakeError(LexResult.StealError());
+	}
+
+	// Handle the += and -= tokens.
+	TArray<FExpressionToken> Tokens = LexResult.StealValue();
+	if (Tokens.Num())
+	{
+		FStringToken Context = Tokens[0].Context;
+		const FExpressionNode& FirstNode = Tokens[0].Node;
+		bool WasOpAssign = true;
+
+		if (FirstNode.Cast<FPlusEquals>())
+		{
+			Tokens.Insert(FExpressionToken(Context, FPlus()), 0);
+		}
+		else if (FirstNode.Cast<FMinusEquals>())
+		{
+			Tokens.Insert(FExpressionToken(Context, FMinus()), 0);
+		}
+		else if (FirstNode.Cast<FStarEquals>())
+		{
+			Tokens.Insert(FExpressionToken(Context, FStar()), 0);
+		}
+		else if (FirstNode.Cast<FForwardSlashEquals>())
+		{
+			Tokens.Insert(FExpressionToken(Context, FForwardSlash()), 0);
+		}
+		else
+		{
+			WasOpAssign = false;
+		}
+
+		if (WasOpAssign)
+		{
+			Tokens.Insert(FExpressionToken(Context, InExistingValue), 0);
+			Tokens.RemoveAt(2, 1, false);
+		}
+	}
+
+	TValueOrError<TArray<FCompiledToken>, FExpressionError> CompilationResult = ExpressionParser::Compile(MoveTemp(Tokens), Grammar);
+	if (!CompilationResult.IsValid())
+	{
+		return MakeError(CompilationResult.StealError());
+	}
+
+	TOperatorEvaluationEnvironment<> Env(JumpTable, nullptr);
+	TValueOrError<FExpressionNode, FExpressionError> Result = ExpressionParser::Evaluate(CompilationResult.GetValue(), Env);
 	if (!Result.IsValid())
 	{
 		return MakeError(Result.GetError());
@@ -182,10 +241,8 @@ TValueOrError<double, FExpressionError> FBasicMathExpressionEvaluator::Evaluate(
 	{
 		return MakeValue(*Numeric);
 	}
-	else
-	{
-		return MakeError(LOCTEXT("UnrecognizedResult", "Unrecognized result returned from expression"));
-	}
+
+	return MakeError(LOCTEXT("UnrecognizedResult", "Unrecognized result returned from expression"));
 }
 
 #if !(UE_BUILD_SHIPPING || UE_BUILD_TEST)
@@ -200,11 +257,13 @@ bool TestExpression(FAutomationTestBase* Test, const TCHAR* Expression, double E
 		Test->AddError(Result.GetError().Text.ToString());
 		return false;
 	}
-	else if (Result.GetValue() != Expected)
+	
+	if (Result.GetValue() != Expected)
 	{
 		Test->AddError(FString::Printf(TEXT("'%s' evaluation results: %.f != %.f"), Expression, Result.GetValue(), Expected));
 		return false;
 	}
+
 	return true;
 }
 
diff --git a/Engine/Source/Runtime/Core/Private/Math/UnitConversion.cpp b/Engine/Source/Runtime/Core/Private/Math/UnitConversion.cpp
--- a/Engine/Source/Runtime/Core/Private/Math/UnitConversion.cpp
+++ b/Engine/Source/Runtime/Core/Private/Math/UnitConversion.cpp
@@ -143,6 +143,10 @@ struct FUnitExpressionParser
 		TokenDefinitions.IgnoreWhitespace();
 			
 		// Defined in order of importance
+		TokenDefinitions.DefineToken(&ConsumeSymbol<FPlusEquals>);
+		TokenDefinitions.DefineToken(&ConsumeSymbol<FMinusEquals>);
+		TokenDefinitions.DefineToken(&ConsumeSymbol<FStarEquals>);
+		TokenDefinitions.DefineToken(&ConsumeSymbol<FForwardSlashEquals>);
 		TokenDefinitions.DefineToken(&ConsumeSymbol<FPlus>);
 		TokenDefinitions.DefineToken(&ConsumeSymbol<FMinus>);
 		TokenDefinitions.DefineToken(&ConsumeSymbol<FStar>);
@@ -300,9 +304,60 @@ struct FUnitExpressionParser
 		return TOptional<FExpressionError>();
 	}
 
-	TValueOrError<FNumericUnit<double>, FExpressionError> Evaluate(const TCHAR* InExpression) const
+	TValueOrError<FNumericUnit<double>, FExpressionError> Evaluate(const TCHAR* InExpression, const FNumericUnit<double>& InExistingValue) const
 	{
-		TValueOrError<FExpressionNode, FExpressionError> Result = ExpressionParser::Evaluate(InExpression, TokenDefinitions, Grammar, JumpTable);
+		using namespace ExpressionParser;
+
+		TValueOrError<TArray<FExpressionToken>, FExpressionError> LexResult = ExpressionParser::Lex(InExpression, TokenDefinitions);
+		if (!LexResult.IsValid())
+		{
+			return MakeError(LexResult.StealError());
+		}
+
+		// Handle the += and -= tokens.
+		TArray<FExpressionToken> Tokens = LexResult.StealValue();
+		if (Tokens.Num())
+		{
+			FStringToken Context = Tokens[0].Context;
+			const FExpressionNode& FirstNode = Tokens[0].Node;
+			bool WasOpAssign = true;
+
+			if (FirstNode.Cast<FPlusEquals>())
+			{			
+				Tokens.Insert(FExpressionToken(Context, FPlus()), 0);
+			}
+			else if (FirstNode.Cast<FMinusEquals>())
+			{
+				Tokens.Insert(FExpressionToken(Context, FMinus()), 0);
+			}
+			else if (FirstNode.Cast<FStarEquals>())
+			{
+				Tokens.Insert(FExpressionToken(Context, FStar()), 0);
+			}
+			else if (FirstNode.Cast<FForwardSlashEquals>())
+			{
+				Tokens.Insert(FExpressionToken(Context, FForwardSlash()), 0);
+			}
+			else
+			{
+				WasOpAssign = false;
+			}
+
+			if (WasOpAssign)
+			{
+				Tokens.Insert(FExpressionToken(Context, InExistingValue), 0);
+				Tokens.RemoveAt(2, 1, false);
+			}
+		}
+
+		TValueOrError<TArray<FCompiledToken>, FExpressionError> CompilationResult = ExpressionParser::Compile(MoveTemp(Tokens), Grammar);
+		if (!CompilationResult.IsValid())
+		{
+			return MakeError(CompilationResult.StealError());
+		}
+
+		TOperatorEvaluationEnvironment<> Env(JumpTable, nullptr);
+		TValueOrError<FExpressionNode, FExpressionError> Result = ExpressionParser::Evaluate(CompilationResult.GetValue(), Env);
 		if (!Result.IsValid())
 		{
 			return MakeError(Result.GetError());
@@ -576,10 +631,10 @@ namespace UnitConversion
 		}
 	}
 
-	TValueOrError<FNumericUnit<double>, FText> TryParseExpression(const TCHAR* InExpression, EUnit From)
+	TValueOrError<FNumericUnit<double>, FText> TryParseExpression(const TCHAR* InExpression, EUnit From, const FNumericUnit<double>& InExistingValue)
 	{
 		const FUnitExpressionParser Parser(From);
-		auto Result = Parser.Evaluate(InExpression);
+		auto Result = Parser.Evaluate(InExpression, InExistingValue);
 		if (Result.IsValid())
 		{
 			if (Result.GetValue().Units == EUnit::Unspecified)
diff --git a/Engine/Source/Runtime/Core/Private/Misc/ExpressionParser.cpp b/Engine/Source/Runtime/Core/Private/Misc/ExpressionParser.cpp
--- a/Engine/Source/Runtime/Core/Private/Misc/ExpressionParser.cpp
+++ b/Engine/Source/Runtime/Core/Private/Misc/ExpressionParser.cpp
@@ -613,10 +613,8 @@ namespace ExpressionParser
 		{
 			return MakeError(Error.GetValue());
 		}
-		else
-		{
-			return MakeValue(TokenConsumer.Extract());
-		}
+
+		return MakeValue(TokenConsumer.Extract());
 	}
 
 	CompileResultType Compile(const TCHAR* InExpression, const FTokenDefinitions& InTokenDefinitions, const FExpressionGrammar& InGrammar)
@@ -662,10 +660,8 @@ namespace ExpressionParser
 			{
 				return CompiledTokens[Index];
 			}
-			else
-			{
-				return RuntimeGeneratedTokens[Index - CompiledTokens.Num()];
-			}
+
+			return RuntimeGeneratedTokens[Index - CompiledTokens.Num()];
 		};
 
 		/** Add a new token to the runtime generated array */
@@ -750,10 +746,8 @@ namespace ExpressionParser
 		{
 			return MakeValue(GetToken(OperandStack[0]).Node.Copy());
 		}
-		else
-		{
-			return MakeError(LOCTEXT("SyntaxError_InvalidExpression", "Could not evaluate expression"));
-		}
+
+		return MakeError(LOCTEXT("SyntaxError_InvalidExpression", "Could not evaluate expression"));
 	}
 }
 
diff --git a/Engine/Source/Runtime/Core/Private/Tests/Internationalization/TextTest.cpp b/Engine/Source/Runtime/Core/Private/Tests/Internationalization/TextTest.cpp
--- a/Engine/Source/Runtime/Core/Private/Tests/Internationalization/TextTest.cpp
+++ b/Engine/Source/Runtime/Core/Private/Tests/Internationalization/TextTest.cpp
@@ -637,7 +637,7 @@ bool FTextTest::RunTest (const FString& Parameters)
 }
 
 
-IMPLEMENT_SIMPLE_AUTOMATION_TEST(FTextRoundingTest, "System.Core.Misc.TextRounding", EAutomationTestFlags::EditorContext | EAutomationTestFlags::ClientContext | EAutomationTestFlags::SmokeFilter)
+IMPLEMENT_SIMPLE_AUTOMATION_TEST(FTextRoundingTest, "System.Core.Misc.TextRounding", EAutomationTestFlags::EditorContext | EAutomationTestFlags::ClientContext | EAutomationTestFlags::EngineFilter)
 
 bool FTextRoundingTest::RunTest (const FString& Parameters)
 {
@@ -776,7 +776,7 @@ bool FTextRoundingTest::RunTest (const FString& Parameters)
 }
 
 
-IMPLEMENT_SIMPLE_AUTOMATION_TEST(FTextPaddingTest, "System.Core.Misc.TextPadding", EAutomationTestFlags::EditorContext | EAutomationTestFlags::ClientContext | EAutomationTestFlags::SmokeFilter)
+IMPLEMENT_SIMPLE_AUTOMATION_TEST(FTextPaddingTest, "System.Core.Misc.TextPadding", EAutomationTestFlags::EditorContext | EAutomationTestFlags::ClientContext | EAutomationTestFlags::EngineFilter)
 
 bool FTextPaddingTest::RunTest (const FString& Parameters)
 {
@@ -876,7 +876,7 @@ bool FTextPaddingTest::RunTest (const FString& Parameters)
 
 #if UE_ENABLE_ICU
 
-IMPLEMENT_SIMPLE_AUTOMATION_TEST(FICUTextTest, "System.Core.Misc.ICUText", EAutomationTestFlags::EditorContext | EAutomationTestFlags::ClientContext | EAutomationTestFlags::SmokeFilter)
+IMPLEMENT_SIMPLE_AUTOMATION_TEST(FICUTextTest, "System.Core.Misc.ICUText", EAutomationTestFlags::EditorContext | EAutomationTestFlags::ClientContext | EAutomationTestFlags::EngineFilter)
 
 bool FICUTextTest::RunTest (const FString& Parameters)
 {
diff --git a/Engine/Source/Runtime/Core/Private/Tests/Misc/CircularQueueTest.cpp b/Engine/Source/Runtime/Core/Private/Tests/Misc/CircularQueueTest.cpp
--- a/Engine/Source/Runtime/Core/Private/Tests/Misc/CircularQueueTest.cpp
+++ b/Engine/Source/Runtime/Core/Private/Tests/Misc/CircularQueueTest.cpp
@@ -8,7 +8,7 @@
 
 IMPLEMENT_SIMPLE_AUTOMATION_TEST(FCircularQueueTest, "System.Core.Misc.CircularQueue", EAutomationTestFlags::ApplicationContextMask | EAutomationTestFlags::SmokeFilter)
 
-bool FCircularQueueTest::RunTest( const FString& Parameters )
+bool FCircularQueueTest::RunTest(const FString& Parameters)
 {
 	const uint32 QueueSize = 8;
 
@@ -16,6 +16,7 @@ bool FCircularQueueTest::RunTest( const FString& Parameters )
 	{
 		TCircularQueue<int32> Queue(QueueSize);
 
+		TestEqual(TEXT("Newly created queues must have zero elements"), Queue.Count(), 0u);
 		TestTrue(TEXT("Newly created queues must be empty"), Queue.IsEmpty());
 		TestFalse(TEXT("Newly created queues must not be full"), Queue.IsFull());
 	}
@@ -26,6 +27,7 @@ bool FCircularQueueTest::RunTest( const FString& Parameters )
 		int32 Value = 0;
 	
 		TestTrue(TEXT("Adding to an empty queue must succeed"), Queue.Enqueue(666));
+		TestEqual(TEXT("After adding to an empty queue it must have one element"), Queue.Count(), 1u);
 		TestFalse(TEXT("Partially filled queues must not be empty"), Queue.IsEmpty());
 		TestFalse(TEXT("Partially filled queues must not be full"), Queue.IsFull());
 		TestTrue(TEXT("Peeking at a partially filled queue must succeed"), Queue.Peek(Value));
diff --git a/Engine/Source/Runtime/Core/Private/Tests/Misc/TripleBufferTest.cpp b/Engine/Source/Runtime/Core/Private/Tests/Misc/TripleBufferTest.cpp
new file mode 100644
--- /dev/null
+++ b/Engine/Source/Runtime/Core/Private/Tests/Misc/TripleBufferTest.cpp
@@ -0,0 +1,70 @@
+// Copyright 1998-2015 Epic Games, Inc. All Rights Reserved.
+
+#include "CorePrivatePCH.h"
+#include "AutomationTest.h"
+#include "TripleBuffer.h"
+
+
+IMPLEMENT_SIMPLE_AUTOMATION_TEST(FTripleBufferTest, "System.Core.Misc.TripleBuffer", EAutomationTestFlags::ApplicationContextMask | EAutomationTestFlags::SmokeFilter)
+
+
+bool FTripleBufferTest::RunTest(const FString& Parameters)
+{
+	// uninitialized buffer
+	TTripleBuffer<int32> b1_1(NoInit);
+
+	TestFalse(TEXT("Uninitialized triple buffer must not be dirty"), b1_1.IsDirty());
+
+	// initialized buffer
+	TTripleBuffer<int32> b2_1(1);
+
+	TestFalse(TEXT("Initialized triple buffer must not be dirty"), b2_1.IsDirty());
+	TestEqual(TEXT("Initialized triple buffer must have correct read buffer value"), b2_1.Read(), 1);
+
+	b2_1.SwapReadBuffers();
+
+	TestEqual(TEXT("Initialized triple buffer must have correct temp buffer value"), b2_1.Read(), 1);
+
+	b2_1.SwapWriteBuffers();
+
+	TestTrue(TEXT("Write buffer swap must set dirty flag"), b2_1.IsDirty());
+
+	b2_1.SwapReadBuffers();
+
+	TestFalse(TEXT("Read buffer swap must clear dirty flag"), b2_1.IsDirty());
+	TestEqual(TEXT("Initialized triple buffer must have correct temp buffer value"), b2_1.Read(), 1);
+
+	// Pre-set buffer
+	TTripleBuffer<int32> b3_1(1, 2, 3);
+
+	int32 Read = b3_1.Read();
+	TestEqual(TEXT("Pre-set triple buffer must have correct read buffer value"), Read, 1);
+
+	b3_1.SwapReadBuffers();
+
+	int32 Temp = b3_1.Read();
+	TestEqual(TEXT("Pre-set triple buffer must have correct temp buffer value"), Temp, 3);
+
+	b3_1.SwapWriteBuffers();
+	b3_1.SwapReadBuffers();
+
+	int32 Write = b3_1.Read();
+	TestEqual(TEXT("Pre-set triple buffer must have correct write buffer value"), Write, 2);
+
+	// operations
+	TTripleBuffer<int32> b4_1;
+
+	b4_1.WriteAndSwap(1);
+	b4_1.SwapReadBuffers();
+	b4_1.WriteAndSwap(2);
+	b4_1.Write(3);
+
+	TestEqual(TEXT("Triple buffer must read correct value (1)"), b4_1.Read(), 1);
+	TestEqual(TEXT("Triple buffer must read correct value (2)"), b4_1.SwapAndRead(), 2);
+
+	b4_1.SwapWriteBuffers();
+
+	TestEqual(TEXT("Triple buffer must read correct value (3)"), b4_1.SwapAndRead(), 3);
+
+	return true;
+}
diff --git a/Engine/Source/Runtime/Core/Public/Containers/Queue.h b/Engine/Source/Runtime/Core/Public/Containers/Queue.h
--- a/Engine/Source/Runtime/Core/Public/Containers/Queue.h
+++ b/Engine/Source/Runtime/Core/Public/Containers/Queue.h
@@ -3,20 +3,17 @@
 #pragma once
 
 
-namespace EQueueMode
+/**
+ * Enumerates concurrent queue modes.
+ */
+enum class EQueueMode
 {
-	/**
-	 * Enumerates concurrent queue modes.
-	 */
-	enum Type
-	{
-		/** Multiple-producers, single-consumer queue. */
-		Mpsc,
+	/** Multiple-producers, single-consumer queue. */
+	Mpsc,
 
-		/** Single-producer, single-consumer queue. */
-		Spsc
-	};
-}
+	/** Single-producer, single-consumer queue. */
+	Spsc
+};
 
 
 /**
@@ -34,7 +31,7 @@ namespace EQueueMode
  * @param Mode The queue mode (single-producer, single-consumer by default).
  * @todo gmp: Implement node pooling.
  */
-template<typename ItemType, EQueueMode::Type Mode = EQueueMode::Spsc>
+template<typename ItemType, EQueueMode Mode = EQueueMode::Spsc>
 class TQueue
 {
 public:
diff --git a/Engine/Source/Runtime/Core/Public/Math/BasicMathExpressionEvaluator.h b/Engine/Source/Runtime/Core/Public/Math/BasicMathExpressionEvaluator.h
--- a/Engine/Source/Runtime/Core/Public/Math/BasicMathExpressionEvaluator.h
+++ b/Engine/Source/Runtime/Core/Public/Math/BasicMathExpressionEvaluator.h
@@ -17,9 +17,13 @@ DEFINE_EXPRESSION_OPERATOR_NODE(FSubExpressionStart, 0xCC40A083, 0xADBF46E2, 0xA
 DEFINE_EXPRESSION_OPERATOR_NODE(FSubExpressionEnd, 0x125E4C67, 0x96EB48C4, 0x8894E09C, 0xB3CD56BF)
 
 DEFINE_EXPRESSION_OPERATOR_NODE(FPlus, 0x6F88756B, 0xF9234263, 0x9B13614F, 0x2706074B)
+DEFINE_EXPRESSION_OPERATOR_NODE(FPlusEquals, 0x05A878C9, 0x0E6C44A4, 0x9A73920D, 0x3175AB48)
 DEFINE_EXPRESSION_OPERATOR_NODE(FMinus, 0xE6240779, 0xEE2849CF, 0x95A0E648, 0x22D58713)
+DEFINE_EXPRESSION_OPERATOR_NODE(FMinusEquals, 0x8D1E08E3, 0x5F534245, 0xA987AD7E, 0xDB78A4B7)
 DEFINE_EXPRESSION_OPERATOR_NODE(FStar, 0xF287B3BF, 0x35DF4141, 0xA8B4F57B, 0x7E06DF47)
+DEFINE_EXPRESSION_OPERATOR_NODE(FStarEquals, 0xBA359BB0, 0xCEB54682, 0x9160EB4F, 0xD1687C7F)
 DEFINE_EXPRESSION_OPERATOR_NODE(FForwardSlash, 0xF99670F8, 0x74794256, 0xBB0CAE6D, 0xC67CD5B6)
+DEFINE_EXPRESSION_OPERATOR_NODE(FForwardSlashEquals, 0x4AFE0CF8, 0xF9054360, 0xBE5DCE80, 0xDC2E22F6)
 DEFINE_EXPRESSION_OPERATOR_NODE(FPercent, 0x936E4434, 0x0A014F2D, 0xBEEC90D3, 0x4D3ECEA2)
 DEFINE_EXPRESSION_OPERATOR_NODE(FSquareRoot, 0xE7C03E11, 0x9DE84B4B, 0xBA4C2B76, 0x69BF028E)
 DEFINE_EXPRESSION_OPERATOR_NODE(FPower, 0x93388F8D, 0x1D9B4DFE, 0xBD4D6CC4, 0x12D1DE99)
@@ -47,17 +51,17 @@ namespace ExpressionParser
 }
 
 /** A basic math expression evaluator */
-class FBasicMathExpressionEvaluator
+class CORE_API FBasicMathExpressionEvaluator
 {
 public:
 	/** Constructor that sets up the parser's lexer and compiler */
 	FBasicMathExpressionEvaluator();
 
 	/** Evaluate the given expression, resulting in either a double value, or an error */
-	TValueOrError<double, FExpressionError> Evaluate(const TCHAR* InExpression) const;
+	TValueOrError<double, FExpressionError> Evaluate(const TCHAR* InExpression, double InExistingValue = 0) const;
 	
 private:
 	FTokenDefinitions TokenDefinitions;
 	FExpressionGrammar Grammar;
 	FOperatorJumpTable JumpTable;
-};
\ No newline at end of file
+};
diff --git a/Engine/Source/Runtime/Core/Public/Math/UnitConversion.h b/Engine/Source/Runtime/Core/Public/Math/UnitConversion.h
--- a/Engine/Source/Runtime/Core/Public/Math/UnitConversion.h
+++ b/Engine/Source/Runtime/Core/Public/Math/UnitConversion.h
@@ -164,7 +164,7 @@ public:
 	FNumericUnit<NumericType> QuantizeUnitsToBestFit() const;
 
 	/** Try and parse an expression into a numeric unit */
-	static TValueOrError<FNumericUnit<NumericType>, FText> TryParseExpression(const TCHAR* InExpression, EUnit InDefaultUnit);
+	static TValueOrError<FNumericUnit<NumericType>, FText> TryParseExpression(const TCHAR* InExpression, EUnit InDefaultUnit, const FNumericUnit<NumericType>& InExistingValue);
 
 	/** Parse a numeric unit from a string */
 	static TOptional<FNumericUnit<NumericType>> TryParseString(const TCHAR* InSource);
diff --git a/Engine/Source/Runtime/Core/Public/Math/UnitConversion.inl b/Engine/Source/Runtime/Core/Public/Math/UnitConversion.inl
--- a/Engine/Source/Runtime/Core/Public/Math/UnitConversion.inl
+++ b/Engine/Source/Runtime/Core/Public/Math/UnitConversion.inl
@@ -25,7 +25,7 @@ namespace UnitConversion
 	CORE_API double TimeUnificationFactor(EUnit From);
 
 	/** Attempt to parse an expression into a numeric unit */
-	CORE_API TValueOrError<FNumericUnit<double>, FText> TryParseExpression(const TCHAR* InExpression, EUnit DefaultUnit);
+	CORE_API TValueOrError<FNumericUnit<double>, FText> TryParseExpression(const TCHAR* InExpression, EUnit DefaultUnit, const FNumericUnit<double>& InExistingValue);
 
 	/** Structure used to define the factor required to get from one unit type to the next. */
 	struct FQuantizationInfo
@@ -253,9 +253,9 @@ FNumericUnit<NumericType> FNumericUnit<NumericType>::QuantizeUnitsToBestFit() co
 }
 
 template<typename NumericType>
-TValueOrError<FNumericUnit<NumericType>, FText> FNumericUnit<NumericType>::TryParseExpression(const TCHAR* InExpression, EUnit InDefaultUnit)
+TValueOrError<FNumericUnit<NumericType>, FText> FNumericUnit<NumericType>::TryParseExpression(const TCHAR* InExpression, EUnit InDefaultUnit, const FNumericUnit<NumericType>& InExistingValue)
 {
-	TValueOrError<FNumericUnit<double>, FText> Result = UnitConversion::TryParseExpression(InExpression, InDefaultUnit);
+	TValueOrError<FNumericUnit<double>, FText> Result = UnitConversion::TryParseExpression(InExpression, InDefaultUnit, InExistingValue);
 	if (Result.IsValid())
 	{
 		const auto& Value = Result.GetValue();
diff --git a/Engine/Source/Runtime/Engine/Classes/Camera/CameraComponent.h b/Engine/Source/Runtime/Engine/Classes/Camera/CameraComponent.h
--- a/Engine/Source/Runtime/Engine/Classes/Camera/CameraComponent.h
+++ b/Engine/Source/Runtime/Engine/Classes/Camera/CameraComponent.h
@@ -15,7 +15,7 @@ class ENGINE_API UCameraComponent : public USceneComponent
 	GENERATED_UCLASS_BODY()
 
 	/** The horizontal field of view (in degrees) in perspective mode (ignored in Orthographic mode) */
-	UPROPERTY(Interp, EditAnywhere, BlueprintReadWrite, Category=CameraSettings, meta=(UIMin = "5.0", UIMax = "170", ClampMin = "0.001", ClampMax = "360.0"))
+	UPROPERTY(Interp, EditAnywhere, BlueprintReadWrite, Category=CameraSettings, meta=(UIMin = "5.0", UIMax = "170", ClampMin = "0.001", ClampMax = "360.0", Units = deg))
 	float FieldOfView;
 	UFUNCTION(BlueprintCallable, Category=Camera)
 	void SetFieldOfView(float InFieldOfView) { FieldOfView = InFieldOfView; }
@@ -57,7 +57,7 @@ class ENGINE_API UCameraComponent : public USceneComponent
 	void SetUseFieldOfViewForLOD(bool bInUseFieldOfViewForLOD) { bUseFieldOfViewForLOD = bInUseFieldOfViewForLOD; }
 
 	/** True if the camera's orientation and position should be locked to the HMD */
-	UPROPERTY(EditAnywhere, BlueprintReadWrite, Category=Camera)
+	UPROPERTY(EditAnywhere, BlueprintReadWrite, Category=CameraSettings)
 	uint32 bLockToHmd:1;
 
 	/**
@@ -74,13 +74,13 @@ class ENGINE_API UCameraComponent : public USceneComponent
 	void SetProjectionMode(ECameraProjectionMode::Type InProjectionMode) { ProjectionMode = InProjectionMode; }
 
 	/** Indicates if PostProcessSettings should be used when using this Camera to view through. */
-	UPROPERTY(Interp, EditAnywhere, BlueprintReadWrite, Category=CameraSettings, meta=(UIMin = "0.0", UIMax = "1.0"))
+	UPROPERTY(Interp, EditAnywhere, BlueprintReadWrite, Category=PostProcess, meta=(UIMin = "0.0", UIMax = "1.0"))
 	float PostProcessBlendWeight;
 	UFUNCTION(BlueprintCallable, Category=Camera)
 	void SetPostProcessBlendWeight(float InPostProcessBlendWeight) { PostProcessBlendWeight = InPostProcessBlendWeight; }
 
 	/** Post process settings to use for this camera. Don't forget to check the properties you want to override */
-	UPROPERTY(Interp, BlueprintReadWrite, Category=CameraSettings)
+	UPROPERTY(Interp, BlueprintReadWrite, Category = PostProcess)
 	struct FPostProcessSettings PostProcessSettings;
 
 	// UActorComponent interface
diff --git a/Engine/Source/Runtime/Engine/Classes/Camera/CineCameraActor.h b/Engine/Source/Runtime/Engine/Classes/Camera/CineCameraActor.h
new file mode 100644
--- /dev/null
+++ b/Engine/Source/Runtime/Engine/Classes/Camera/CineCameraActor.h
@@ -0,0 +1,22 @@
+// Copyright 1998-2015 Epic Games, Inc. All Rights Reserved.
+
+#pragma once
+
+#include "CineCameraActor.generated.h"
+
+/** 
+ * A CineCameraActor is a CameraActor specialized to work like a cinematic camera.
+ */
+UCLASS(ClassGroup = Common, hideCategories = (Input, Rendering, AutoPlayerActivation), showcategories = ("Input|MouseInput", "Input|TouchInput"), Blueprintable)
+class ENGINE_API ACineCameraActor : public ACameraActor
+{
+	GENERATED_BODY()
+
+public:
+	// Ctor
+	ACineCameraActor(const FObjectInitializer& ObjectInitializer);
+
+#if WITH_EDITOR
+	virtual void PostEditChangeProperty(FPropertyChangedEvent& PropertyChangedEvent) override;
+#endif
+};
diff --git a/Engine/Source/Runtime/Engine/Classes/Camera/CineCameraComponent.h b/Engine/Source/Runtime/Engine/Classes/Camera/CineCameraComponent.h
new file mode 100644
--- /dev/null
+++ b/Engine/Source/Runtime/Engine/Classes/Camera/CineCameraComponent.h
@@ -0,0 +1,232 @@
+// Copyright 1998-2016 Epic Games, Inc. All Rights Reserved.
+
+#pragma once
+#include "Camera/CameraComponent.h"
+#include "CineCameraComponent.generated.h"
+
+USTRUCT()
+struct FCameraFilmbackSettings
+{
+	GENERATED_USTRUCT_BODY()
+
+	UPROPERTY(EditAnywhere, BlueprintReadWrite, Category = "Filmback", meta = (ClampMin = "0.001", ForceUnits = mm))
+	float SensorWidth;
+
+	/** Vertical size of filmback or digital sensor, in mm. */
+	UPROPERTY(EditAnywhere, BlueprintReadWrite, Category = "Filmback", meta = (ClampMin = "0.001", ForceUnits = mm))
+	float SensorHeight;
+
+	UPROPERTY(VisibleAnywhere, BlueprintReadOnly, Category = "Filmback")
+	float SensorAspectRatio;
+};
+
+USTRUCT()
+struct FNamedFilmbackPreset
+{
+	GENERATED_USTRUCT_BODY()
+
+	UPROPERTY()
+	FString Name;
+
+	UPROPERTY()
+	FCameraFilmbackSettings FilmbackSettings;
+};
+
+USTRUCT()
+struct FCameraLensSettings
+{
+	GENERATED_USTRUCT_BODY()
+
+	UPROPERTY(EditAnywhere, BlueprintReadWrite, Category = "Lens", meta = (ForceUnits = mm))
+	float MinFocalLength;
+
+	UPROPERTY(EditAnywhere, BlueprintReadWrite, Category = "Lens", meta = (ForceUnits = mm))
+	float MaxFocalLength;
+
+	// #todo, split into MinFStopAtMinFocalLength, MinFStopAtMaxFocalLength for variable-min-fstop lenses
+
+	UPROPERTY(EditAnywhere, BlueprintReadWrite, Category = "Lens")
+	float MinFStop;
+
+	UPROPERTY(EditAnywhere, BlueprintReadWrite, Category = "Lens")
+	float MaxFStop;
+
+	UPROPERTY(EditAnywhere, BlueprintReadWrite, Category = "Lens", meta = (ForceUnits = mm))
+	float MinimumFocusDistance;
+
+	UPROPERTY(EditAnywhere, BlueprintReadWrite, Category = "Lens")
+	float MaxReproductionRatio;
+
+	bool operator==(const FCameraLensSettings& Other) const
+	{
+		return (MinFocalLength == Other.MinFocalLength)
+			&& (MaxFocalLength == Other.MaxFocalLength)
+			&& (MinFStop == Other.MinFStop)
+			&& (MaxFStop == Other.MaxFStop)
+			&& (MinimumFocusDistance == Other.MinimumFocusDistance)
+			&& (MaxReproductionRatio == Other.MaxReproductionRatio);
+	}
+};
+
+
+USTRUCT()
+struct FNamedLensPreset
+{
+	GENERATED_USTRUCT_BODY()
+
+	UPROPERTY()
+	FString Name;
+
+	UPROPERTY()
+	FCameraLensSettings LensSettings;
+};
+
+UENUM()
+enum class ECameraFocusMethod : uint8
+{
+	None,		/** Disable DoF entirely. */
+	Manual,		/** Allows for specifying or animating exact focus distances. */
+	Spot UMETA(Hidden), /** Raycasts into the scene, focuses on the hit location. */		// #todo, make this work
+	Tracking,	/** Locks focus to specific object. */
+};
+
+USTRUCT()
+struct FCameraSpotFocusSettings
+{
+	GENERATED_USTRUCT_BODY()
+
+	UPROPERTY(EditAnywhere, BlueprintReadWrite, Category="Spot Focus")
+	TArray<FVector2D> NormalizedSpotCoordinates;
+	
+	FCameraSpotFocusSettings()
+	{
+		NormalizedSpotCoordinates.Add(FVector2D(0.5f, 0.5f));
+	}
+};
+
+USTRUCT()
+struct FCameraTrackingFocusSettings
+{
+	GENERATED_USTRUCT_BODY()
+
+	UPROPERTY(EditInstanceOnly, BlueprintReadWrite, Category = "Tracking Focus")
+	AActor* ActorToTrack;
+
+	// #todo, ability to pick a component and socket/bone?
+
+// 	UPROPERTY(EditInstanceOnly, BlueprintReadWrite, Category = "Tracking Focus")
+// 	USceneComponent* ComponentToTrack;
+// 	
+// 	UPROPERTY(EditInstanceOnly, BlueprintReadWrite, Category = "Tracking Focus")
+// 	FName BoneOrSocketName;
+
+	FCameraTrackingFocusSettings()
+		: ActorToTrack(nullptr)
+	{}
+};
+
+
+USTRUCT()
+struct FCameraFocusSettings
+{
+	GENERATED_USTRUCT_BODY()
+
+	UPROPERTY(EditAnywhere, BlueprintReadWrite, Category = "Focus Settings")
+	TEnumAsByte<ECameraFocusMethod> FocusMethod;
+	
+	UPROPERTY(EditAnywhere, BlueprintReadWrite, Category = "Focus Settings")
+	uint8 bSmoothFocusChanges : 1;
+	
+	UPROPERTY(/*EditAnywhere, BlueprintReadWrite, Category = "Focus Settings"*/)
+	uint8 bSimulateFocusBreathing : 1;
+
+	UPROPERTY(Interp, EditAnywhere, BlueprintReadWrite, Category = "Focus Settings")
+	float ManualFocusDistance;
+
+	UPROPERTY(/*EditAnywhere, BlueprintReadWrite, Category = "Focus Settings"*/)
+	FCameraSpotFocusSettings SpotFocusSettings;
+
+	UPROPERTY(EditAnywhere, BlueprintReadWrite, Category = "Focus Settings")
+	FCameraTrackingFocusSettings TrackingFocusSettings;
+
+	UPROPERTY(EditAnywhere, BlueprintReadWrite, Category = "Focus Settings")
+	float FocusSmoothingInterpSpeed;
+
+	FCameraFocusSettings() : 
+		FocusMethod(ECameraFocusMethod::Manual),
+		bSmoothFocusChanges(0),
+		bSimulateFocusBreathing(0),
+		ManualFocusDistance(0.f),			// #todo better default?
+		FocusSmoothingInterpSpeed(8.f)
+	{}
+};
+
+/**
+ * A specialized version of a camera component, geared toward cinematic usage.
+ */
+UCLASS(
+	HideCategories = (CameraSettings), 
+	HideFunctions = (SetFieldOfView, SetAspectRatio, SetConstraintAspectRatio), 
+	Blueprintable, 
+	ClassGroup = Camera, 
+	meta = (BlueprintSpawnableComponent), 
+	Config = Engine
+	)
+class ENGINE_API UCineCameraComponent : public UCameraComponent
+{
+	GENERATED_BODY()
+
+public:
+	/** Default constuctor. */
+	UCineCameraComponent();
+
+	virtual void GetCameraView(float DeltaTime, FMinimalViewInfo& DesiredView) override;
+
+	UPROPERTY(EditAnywhere, BlueprintReadWrite, Category = "Filmback")
+	FCameraFilmbackSettings FilmbackSettings;
+
+	UPROPERTY(EditAnywhere, BlueprintReadWrite, Category = "Lens")
+	FCameraLensSettings LensSettings;
+
+	UPROPERTY(Interp, EditAnywhere, BlueprintReadWrite, Category = "Current Camera Settings")
+	FCameraFocusSettings FocusSettings;
+
+	UPROPERTY(Interp, EditAnywhere, BlueprintReadWrite, Category = "Current Camera Settings")
+	float CurrentFocalLength;
+
+	UPROPERTY(Interp, EditAnywhere, BlueprintReadWrite, Category = "Current Camera Settings")
+	float CurrentAperture;
+	
+	/** Read-only. Control this value via FocusSettings. */
+	UPROPERTY(VisibleAnywhere, BlueprintReadOnly, Category = "Current Camera Settings")
+	float CurrentFocusDistance;
+
+	// for interpto
+	float LastFocusDistance;
+
+	float GetHorizontalFieldOfView() const;
+	float GetVerticalFieldOfView() const;
+
+	static TArray<FNamedFilmbackPreset> const& GetFilmbackPresets();
+	static TArray<FNamedLensPreset> const& GetLensPresets();
+
+protected:
+
+	virtual void PostLoad() override;
+#if WITH_EDITORONLY_DATA
+	virtual void PostEditChangeProperty(FPropertyChangedEvent& PropertyChangedEvent) override;
+#endif
+
+	UPROPERTY(config)
+	TArray<FNamedFilmbackPreset> FilmbackPresets;
+
+	UPROPERTY(config)
+	TArray<FNamedLensPreset> LensPresets;
+
+	virtual void UpdateCameraLens(float DeltaTime, FMinimalViewInfo& DesiredView);
+
+private:
+	void RecalcDerivedData();
+	float GetDesiredFocusDistance(FMinimalViewInfo& DesiredView) const;
+	float GetWorldToMetersScale() const;
+};
diff --git a/Engine/Source/Runtime/Engine/Classes/Curves/CurveBase.h b/Engine/Source/Runtime/Engine/Classes/Curves/CurveBase.h
--- a/Engine/Source/Runtime/Engine/Classes/Curves/CurveBase.h
+++ b/Engine/Source/Runtime/Engine/Classes/Curves/CurveBase.h
@@ -523,6 +523,9 @@ public:
 	FRichCurveKey GetFirstKey() const;
 	FRichCurveKey GetLastKey() const;
 
+	/** Get the first key that matches any of the given key handles. */
+	FRichCurveKey& GetFirstMatchingKey(const TArray<FKeyHandle>& KeyHandles);
+
 	/** Get the next or previous key given the key handle */
 	FKeyHandle GetNextKey(FKeyHandle KeyHandle) const;
 	FKeyHandle GetPreviousKey(FKeyHandle KeyHandle) const;
@@ -609,6 +612,9 @@ public:
 	/** Determine if two RichCurves are the same */
 	bool operator == (const FRichCurve& Curve) const;
 
+	/** Remove redundant keys, comparing against Tolerance */
+	void RemoveRedundantKeys(float Tolerance);
+
 public:
 
 	// FIndexedCurve interface
diff --git a/Engine/Source/Runtime/Engine/Classes/Engine/Engine.h b/Engine/Source/Runtime/Engine/Classes/Engine/Engine.h
--- a/Engine/Source/Runtime/Engine/Classes/Engine/Engine.h
+++ b/Engine/Source/Runtime/Engine/Classes/Engine/Engine.h
@@ -1586,6 +1586,13 @@ public:
 
 	
 
+	/** Delegate broadcast after UEditorEngine::Tick has been called (or UGameEngine::Tick in standalone) */
+	DECLARE_EVENT_OneParam(UEditorEngine, FPostEditorTick, float /* DeltaTime */);
+	FPostEditorTick& OnPostEditorTick() { return PostEditorTickEvent; }
+
+	/** Called after UEditorEngine::Tick has been called (or UGameEngine::Tick in standalone) */
+	void BroadcastPostEditorTick(float DeltaSeconds) { PostEditorTickEvent.Broadcast(DeltaSeconds); }
+
 #endif // #if WITH_EDITOR
 
 	/** Event triggered after a server travel failure of any kind has occurred */
@@ -1646,8 +1653,6 @@ public:
 	bool HandleDumpTicksCommand( UWorld* InWorld, const TCHAR* Cmd, FOutputDevice& Ar );
 	bool HandleGammaCommand( const TCHAR* Cmd, FOutputDevice& Ar );
 
-	bool HandleRecordAnimationCommand(UWorld* InWorld, const TCHAR* InStr, FOutputDevice& Ar);
-	bool HandleStopRecordAnimationCommand(UWorld* InWorld, const TCHAR* InStr, FOutputDevice& Ar);
 	bool HandleShowLogCommand( const TCHAR* Cmd, FOutputDevice& Ar );
 
 	// Only compile in when STATS is set
@@ -2322,6 +2327,9 @@ private:
 	/** Broadcasts after an actor has been moved, rotated or scaled */
 	FOnActorMovedEvent		OnActorMovedEvent;
 	
+	/** Delegate broadcast after UEditorEngine::Tick has been called (or UGameEngine::Tick in standalone) */
+	FPostEditorTick PostEditorTickEvent;
+
 #endif // #if WITH_EDITOR
 
 	/** Thread preventing screen saver from kicking. Suspend most of the time. */
diff --git a/Engine/Source/Runtime/Engine/Classes/GameFramework/Actor.h b/Engine/Source/Runtime/Engine/Classes/GameFramework/Actor.h
--- a/Engine/Source/Runtime/Engine/Classes/GameFramework/Actor.h
+++ b/Engine/Source/Runtime/Engine/Classes/GameFramework/Actor.h
@@ -1513,7 +1513,7 @@ public:
 	FORCEINLINE class USceneComponent* GetRootComponent() const { return RootComponent; }
 
 	/**
-	 * Returns this actor's default attacjment component for attaching children to
+	 * Returns this actor's default attachment component for attaching children to
 	 * @return The scene component to be used as parent
 	 */
 	virtual class USceneComponent* GetDefaultAttachComponent() const { return GetRootComponent(); }
diff --git a/Engine/Source/Runtime/Engine/Classes/Materials/MaterialExpressionConstant3Vector.h b/Engine/Source/Runtime/Engine/Classes/Materials/MaterialExpressionConstant3Vector.h
--- a/Engine/Source/Runtime/Engine/Classes/Materials/MaterialExpressionConstant3Vector.h
+++ b/Engine/Source/Runtime/Engine/Classes/Materials/MaterialExpressionConstant3Vector.h
@@ -1,27 +1,29 @@
 // Copyright 1998-2016 Epic Games, Inc. All Rights Reserved.
 
-
 #pragma once
+
 #include "Materials/MaterialExpression.h"
 #include "MaterialExpressionConstant3Vector.generated.h"
 
+
 UCLASS(collapsecategories, hidecategories=Object, MinimalAPI)
-class UMaterialExpressionConstant3Vector : public UMaterialExpression
+class UMaterialExpressionConstant3Vector
+	: public UMaterialExpression
 {
 	GENERATED_UCLASS_BODY()
 
- 	UPROPERTY(EditAnywhere, Category=MaterialExpressionConstant3Vector, meta=(HideAlphaChannel))
+ 	UPROPERTY(EditAnywhere, Category=MaterialExpressionConstant3Vector, meta=(DontUpdateWhileEditing, HideAlphaChannel))
 	FLinearColor Constant;
 
-	//~ Begin UMaterialExpression Interface
+public:
+
+	//~ UMaterialExpression interface
+
 	virtual int32 Compile(class FMaterialCompiler* Compiler, int32 OutputIndex, int32 MultiplexIndex) override;
 	virtual void GetCaption(TArray<FString>& OutCaptions) const override;
+
 #if WITH_EDITOR
 	virtual FString GetDescription() const override;
 	virtual uint32 GetOutputType(int32 OutputIndex) override {return MCT_Float3;}
 #endif // WITH_EDITOR
-	//~ End UMaterialExpression Interface
 };
-
-
-
diff --git a/Engine/Source/Runtime/Engine/Private/Camera/CineCameraActor.cpp b/Engine/Source/Runtime/Engine/Private/Camera/CineCameraActor.cpp
new file mode 100644
--- /dev/null
+++ b/Engine/Source/Runtime/Engine/Private/Camera/CineCameraActor.cpp
@@ -0,0 +1,26 @@
+// Copyright 1998-2015 Epic Games, Inc. All Rights Reserved.
+
+#include "EnginePrivate.h"
+#include "Camera/CineCameraComponent.h"
+#include "Camera/CineCameraActor.h"
+
+#define LOCTEXT_NAMESPACE "CineCameraActor"
+
+//////////////////////////////////////////////////////////////////////////
+// ACameraActor
+
+ACineCameraActor::ACineCameraActor(const FObjectInitializer& ObjectInitializer)
+	: Super(ObjectInitializer
+			.SetDefaultSubobjectClass<UCineCameraComponent>(TEXT("CameraComponent"))
+	)
+{
+}
+
+#if WITH_EDITOR
+void ACineCameraActor::PostEditChangeProperty(FPropertyChangedEvent& PropertyChangedEvent)
+{
+	Super::PostEditChangeProperty(PropertyChangedEvent);
+}
+#endif
+
+#undef LOCTEXT_NAMESPACE
diff --git a/Engine/Source/Runtime/Engine/Private/Camera/CineCameraComponent.cpp b/Engine/Source/Runtime/Engine/Private/Camera/CineCameraComponent.cpp
new file mode 100644
--- /dev/null
+++ b/Engine/Source/Runtime/Engine/Private/Camera/CineCameraComponent.cpp
@@ -0,0 +1,228 @@
+// Copyright 1998-2016 Epic Games, Inc. All Rights Reserved.
+
+#include "EnginePrivate.h"
+#include "Camera/CineCameraComponent.h"
+
+#define LOCTEXT_NAMESPACE "CineCameraComponent"
+
+
+//////////////////////////////////////////////////////////////////////////
+// UCameraComponent
+
+UCineCameraComponent::UCineCameraComponent()
+{
+	// Super 35mm 4 Perf
+	FilmbackSettings.SensorWidth = 24.89f;
+	FilmbackSettings.SensorHeight = 18.67;
+	LensSettings.MinFocalLength = 50.f;
+	LensSettings.MaxFocalLength = 50.f;
+	LensSettings.MinFStop = 2.f;
+	LensSettings.MaxFStop = 2.f;
+	LensSettings.MinimumFocusDistance = 15.f;
+	LensSettings.MaxReproductionRatio = 1 / 11.f;
+
+	bConstrainAspectRatio = true;
+
+	RecalcDerivedData();
+}
+
+void UCineCameraComponent::PostLoad()
+{
+	RecalcDerivedData();
+	Super::PostLoad();
+}
+
+#if WITH_EDITORONLY_DATA
+void UCineCameraComponent::PostEditChangeProperty(FPropertyChangedEvent& PropertyChangedEvent)
+{
+	RecalcDerivedData();
+
+	Super::PostEditChangeProperty(PropertyChangedEvent);
+}
+#endif	// WITH_EDITORONLY_DATA
+
+
+float UCineCameraComponent::GetHorizontalFieldOfView() const
+{
+	return (CurrentFocalLength > 0.f)
+		? FMath::RadiansToDegrees(2.f * FMath::Atan(FilmbackSettings.SensorWidth / (2.f * CurrentFocalLength)))
+		: 0.f;
+}
+
+float UCineCameraComponent::GetVerticalFieldOfView() const
+{
+	return (CurrentFocalLength > 0.f)
+		? FMath::RadiansToDegrees(2.f * FMath::Atan(FilmbackSettings.SensorHeight / (2.f * CurrentFocalLength)))
+		: 0.f;
+}
+
+float UCineCameraComponent::GetWorldToMetersScale() const
+{
+	UWorld const* const World = GetWorld();
+	AWorldSettings const* const WorldSettings = World ? World->GetWorldSettings() : nullptr;
+	return WorldSettings ? WorldSettings->WorldToMeters : 100.f;
+}
+
+// static
+TArray<FNamedFilmbackPreset> const& UCineCameraComponent::GetFilmbackPresets()
+{
+	return GetDefault<UCineCameraComponent>()->FilmbackPresets;
+}
+
+// static
+TArray<FNamedLensPreset> const& UCineCameraComponent::GetLensPresets()
+{
+	return GetDefault<UCineCameraComponent>()->LensPresets;
+}
+
+void UCineCameraComponent::RecalcDerivedData()
+{
+	// respect physical limits of the (simulated) hardware
+	CurrentFocalLength = FMath::Clamp(CurrentFocalLength, LensSettings.MinFocalLength, LensSettings.MaxFocalLength);
+	CurrentAperture = FMath::Clamp(CurrentAperture, LensSettings.MinFStop, LensSettings.MaxFStop);
+
+	float const MinFocusDistInWorldUnits = LensSettings.MinimumFocusDistance * (GetWorldToMetersScale() / 1000.f);	// convert mm to uu
+	FocusSettings.ManualFocusDistance = FMath::Max(FocusSettings.ManualFocusDistance, MinFocusDistInWorldUnits);
+
+	FieldOfView = GetHorizontalFieldOfView();
+	FilmbackSettings.SensorAspectRatio = (FilmbackSettings.SensorHeight > 0.f) ? (FilmbackSettings.SensorWidth / FilmbackSettings.SensorHeight) : 0.f;
+	AspectRatio = FilmbackSettings.SensorAspectRatio;
+}
+
+
+static const FName NAME_GetDesiredFocusDistance(TEXT("GetDesiredFocusDistance"));
+static const float SpotFocusTraceDist = 1000000.f;
+static const float FocusDistInfinity = 1000000.f;
+
+float UCineCameraComponent::GetDesiredFocusDistance(FMinimalViewInfo& DesiredView) const
+{
+	float DesiredFocusDistance = 0.f;
+
+	// get focus distance
+	switch (FocusSettings.FocusMethod)
+	{
+	case ECameraFocusMethod::Manual:
+		DesiredFocusDistance = FocusSettings.ManualFocusDistance;
+		break;
+
+	case ECameraFocusMethod::Spot:
+		{
+			// AutoFocus, do a trace
+			AActor* const OwningActor = GetOwner();
+
+			// trace to get depth at center of screen
+			FCollisionQueryParams TraceParams(NAME_GetDesiredFocusDistance, true, OwningActor);
+			FHitResult Hit;
+
+			FVector const TraceStart = DesiredView.Location;
+			FVector const TraceEnd = TraceStart + DesiredView.Rotation.Vector() * SpotFocusTraceDist;
+			bool const bHit = GetWorld()->LineTraceSingleByChannel(Hit, TraceStart, TraceEnd, ECC_WorldStatic, TraceParams);
+			DesiredFocusDistance = bHit ? (TraceStart - Hit.ImpactPoint).Size() : FocusDistInfinity;
+		}
+		break;
+
+	case ECameraFocusMethod::Tracking:
+		{
+			AActor* const TrackedActor = FocusSettings.TrackingFocusSettings.ActorToTrack;
+			if (TrackedActor && TrackedActor->GetRootComponent())
+			{
+				DesiredFocusDistance = (TrackedActor->GetActorLocation() - DesiredView.Location).Size();
+			}
+
+// 			// try to get the position of the tracked object
+// 			UObject* const TrackedObj = FocusSettings.TrackingFocusSettings.ObjectToTrack;
+// 
+// 			// infinity is the fallback in case none of the other code paths produces a good result
+// 			DesiredFocusDistance = FocusDistInfinity;
+// 			if (TrackedObj)
+// 			{
+// 				// we have a socket/bone, see if the object is a scenecomponent
+// 				USceneComponent* const TrackedSceneComp = Cast<USceneComponent>(TrackedObj);
+// 				if (TrackedSceneComp)
+// 				{
+// 					if (FocusSettings.TrackingFocusSettings.BoneOrSocketName != NAME_None)
+// 					{
+// 						FVector const SocketLoc = TrackedSceneComp->GetSocketLocation(FocusSettings.TrackingFocusSettings.BoneOrSocketName);
+// 						DesiredFocusDistance = (SocketLoc - DesiredView.Location).Size();
+// 					}
+// 					else
+// 					{
+// 						DesiredFocusDistance = (TrackedSceneComp->GetComponentLocation() - DesiredView.Location).Size();
+// 					}
+// 				}
+// 				else
+// 				{
+// 					AActor* const TrackedActor = Cast<AActor>(TrackedObj);
+// 					if (TrackedActor)
+// 					{
+// 						if (FocusSettings.TrackingFocusSettings.BoneOrSocketName != NAME_None)
+// 						{
+// 							// user asked for an actor and a bone or socket, try to find a suitable component
+// 							TArray<USceneComponent*> OutComponents;
+// 							TrackedActor->GetComponents(OutComponents);
+// 							for (USceneComponent const* Comp : OutComponents)
+// 							{
+// 								if (Comp->DoesSocketExist(FocusSettings.TrackingFocusSettings.BoneOrSocketName))
+// 								{
+// 									FVector const SocketLoc = Comp->GetSocketLocation(FocusSettings.TrackingFocusSettings.BoneOrSocketName);
+// 									DesiredFocusDistance = (SocketLoc - DesiredView.Location).Size();
+// 									break;
+// 								}
+// 							}
+// 
+// 						}
+// 						else
+// 						{
+// 							DesiredFocusDistance = (TrackedActor->GetActorLocation() - DesiredView.Location).Size();
+// 						}
+// 					}
+// 				}
+// 			}
+		}
+		break;
+	}
+	
+	return DesiredFocusDistance;
+}
+
+void UCineCameraComponent::GetCameraView(float DeltaTime, FMinimalViewInfo& DesiredView)
+{
+	RecalcDerivedData();
+
+	Super::GetCameraView(DeltaTime, DesiredView);
+
+	UpdateCameraLens(DeltaTime, DesiredView);
+}
+
+
+void UCineCameraComponent::UpdateCameraLens(float DeltaTime, FMinimalViewInfo& DesiredView)
+{
+	// Update focus/DoF
+	DesiredView.PostProcessBlendWeight = 1.f;
+	DesiredView.PostProcessSettings.bOverride_DepthOfFieldMethod = true;
+	DesiredView.PostProcessSettings.DepthOfFieldMethod = DOFM_CircleDOF;
+
+	DesiredView.PostProcessSettings.bOverride_DepthOfFieldFstop = true;
+	DesiredView.PostProcessSettings.DepthOfFieldFstop = CurrentAperture;
+
+	CurrentFocusDistance = GetDesiredFocusDistance(DesiredView);
+
+	// clamp to min focus distance
+	float const MinFocusDistInWorldUnits = LensSettings.MinimumFocusDistance * (GetWorldToMetersScale() / 1000.f);	// convert mm to uu
+	CurrentFocusDistance = FMath::Max(CurrentFocusDistance, MinFocusDistInWorldUnits);
+
+	// smoothing, if desired
+ 	if (FocusSettings.bSmoothFocusChanges)
+ 	{
+		CurrentFocusDistance = FMath::FInterpTo(LastFocusDistance, CurrentFocusDistance, DeltaTime, FocusSettings.FocusSmoothingInterpSpeed);
+ 	}
+	LastFocusDistance = CurrentFocusDistance;
+
+	DesiredView.PostProcessSettings.bOverride_DepthOfFieldFocalDistance = true;
+	DesiredView.PostProcessSettings.DepthOfFieldFocalDistance = CurrentFocusDistance;
+}
+
+#undef LOCTEXT_NAMESPACE
+
+
+#undef LOCTEXT_NAMESPACE
diff --git a/Engine/Source/Runtime/Engine/Private/CapturePin.cpp b/Engine/Source/Runtime/Engine/Private/CapturePin.cpp
--- a/Engine/Source/Runtime/Engine/Private/CapturePin.cpp
+++ b/Engine/Source/Runtime/Engine/Private/CapturePin.cpp
@@ -326,7 +326,6 @@ HRESULT FCapturePin::DoBufferProcessingLoop(void)
 		}
 	}
 
-	static_cast<FCaptureSource*>(m_pFilter)->OnFinishedCapturing();
 	return S_FALSE;
 }
 
diff --git a/Engine/Source/Runtime/Engine/Private/CurveBase.cpp b/Engine/Source/Runtime/Engine/Private/CurveBase.cpp
--- a/Engine/Source/Runtime/Engine/Private/CurveBase.cpp
+++ b/Engine/Source/Runtime/Engine/Private/CurveBase.cpp
@@ -572,6 +572,21 @@ FRichCurveKey FRichCurve::GetLastKey() const
 }
 
 
+FRichCurveKey& FRichCurve::GetFirstMatchingKey(const TArray<FKeyHandle>& KeyHandles)
+{
+	for (const auto& KeyHandle : KeyHandles)
+	{
+		if (IsKeyHandleValid(KeyHandle))
+		{
+			return GetKey(KeyHandle);
+		}
+	}
+
+	checkf(false, TEXT("No matching rich curve key found."));
+	return Keys[INDEX_NONE];
+}
+
+
 FKeyHandle FRichCurve::GetNextKey(FKeyHandle KeyHandle) const
 {
 	int32 KeyIndex = GetIndex(KeyHandle);
@@ -1189,6 +1204,34 @@ void FRichCurve::ReadjustTimeRange(float NewMinTimeRange, float NewMaxTimeRange,
 	}
 }
 
+void FRichCurve::RemoveRedundantKeys(float Tolerance)
+{
+	for(int32 KeyIndex = 0; KeyIndex < Keys.Num(); ++KeyIndex)
+	{
+		// copy the key
+		FRichCurveKey OriginalKey = Keys[KeyIndex];
+
+		FKeyHandle KeyHandle = GetKeyHandle(KeyIndex);
+
+		// remove it
+		DeleteKey(KeyHandle);
+
+		// eval to check validity
+		float NewValue = Eval(OriginalKey.Time, DefaultValue);
+			
+		// outside tolerance? re-add the key.
+		if(FMath::Abs(NewValue - OriginalKey.Value) > Tolerance)
+		{
+			FKeyHandle NewKeyHandle = AddKey(OriginalKey.Time, OriginalKey.Value, false, KeyHandle);
+			FRichCurveKey& NewKey = GetKey(NewKeyHandle);
+			NewKey = OriginalKey;
+		}
+		else
+		{
+			KeyIndex--;
+		}
+	}
+}
 
 /** Util to find float value on bezier defined by 4 control points */ 
 static float BezierInterp(float P0, float P1, float P2, float P3, float Alpha)
diff --git a/Engine/Source/Runtime/Engine/Private/GameEngine.cpp b/Engine/Source/Runtime/Engine/Private/GameEngine.cpp
--- a/Engine/Source/Runtime/Engine/Private/GameEngine.cpp
+++ b/Engine/Source/Runtime/Engine/Private/GameEngine.cpp
@@ -42,10 +42,6 @@
 #include "Engine/CoreSettings.h"
 #include "EngineAnalytics.h"
 
-#if WITH_EDITOR
-#include "Editor/UnrealEd/Public/Animation/AnimationRecorder.h"
-#endif
-
 ENGINE_API bool GDisallowNetworkTravel = false;
 
 // How slow must a frame be (in seconds) to be logged out (<= 0 to disable)
@@ -1154,11 +1150,7 @@ void UGameEngine::Tick( float DeltaSeconds, bool bIdleMode )
 	}
 
 #if WITH_EDITOR
-	// tick animation recorder. available only in editor builds
-	if (!IsRunningDedicatedServer() && !IsRunningCommandlet())
-	{
-		FAnimationRecorderManager::Get().Tick(DeltaSeconds);
-	}
+	BroadcastPostEditorTick(DeltaSeconds);
 #endif
 }
 
diff --git a/Engine/Source/Runtime/Engine/Private/UnrealEngine.cpp b/Engine/Source/Runtime/Engine/Private/UnrealEngine.cpp
--- a/Engine/Source/Runtime/Engine/Private/UnrealEngine.cpp
+++ b/Engine/Source/Runtime/Engine/Private/UnrealEngine.cpp
@@ -124,10 +124,6 @@
 #include "ABTesting.h"
 #include "Performance/EnginePerformanceTargets.h"
 
-#if WITH_EDITOR
-#include "Editor/UnrealEd/Public/Animation/AnimationRecorder.h"
-#endif
-
 #include "InstancedReferenceSubobjectHelper.h"
 #include "Engine/EndUserSettings.h"
 #include "ABTesting.h"
@@ -138,10 +134,6 @@
 #include "IAutomationWorkerModule.h"
 #endif	// UE_BUILD_SHIPPING
 
-#if WITH_EDITOR
-#include "Editor/UnrealEd/Public/Animation/AnimationRecorder.h"
-#endif
-
 DEFINE_LOG_CATEGORY(LogEngine);
 IMPLEMENT_MODULE( FEngineModule, Engine );
 
@@ -2814,16 +2806,6 @@ bool UEngine::Exec( UWorld* InWorld, const TCHAR* Cmd, FOutputDevice& Ar )
 	{
 		GNetworkProfiler.Exec( InWorld, Cmd, Ar );
 	}
-#endif
-#if WITH_EDITOR
-	else if (FParse::Command(&Cmd, TEXT("RecordAnimation")))
-	{
-		return HandleRecordAnimationCommand(InWorld, Cmd, Ar);
-	}
-	else if (FParse::Command(&Cmd, TEXT("StopRecordingAnimation")))
-	{
-		return HandleStopRecordAnimationCommand(InWorld, Cmd, Ar);
-	}
 #endif
 	else 
 	{
@@ -2965,110 +2947,6 @@ bool UEngine::HandleGammaCommand( const TCHAR* Cmd, FOutputDevice& Ar )
 	return true;
 }
 
-bool UEngine::HandleRecordAnimationCommand( UWorld* InWorld, const TCHAR* InStr, FOutputDevice& Ar )
-{
-#if WITH_EDITOR
-	const TCHAR* Str = InStr;
-	// parse actor name
-	TCHAR ActorName[128];
-	AActor* FoundActor = nullptr;
-	if (FParse::Token( Str, ActorName, ARRAY_COUNT( ActorName ), 0 ))
-	{
-		FString const ActorNameStr = FString( ActorName );
-		for (ULevel const* Level : InWorld->GetLevels())
-		{
-			if (Level)
-			{
-				for (AActor* Actor : Level->Actors)
-				{
-					if (Actor)
-					{
-						if (Actor->GetName() == ActorNameStr)
-						{
-							FoundActor = Actor;
-							break;
-						}
-					}
-				}
-			}
-		}
-	}
-
-	if (FoundActor)
-	{
-		USkeletalMeshComponent* const SkelComp = FoundActor->FindComponentByClass<USkeletalMeshComponent>();
-		if (SkelComp)
-		{
-			TCHAR AssetPath[256];
-			FParse::Token( Str, AssetPath, ARRAY_COUNT( AssetPath ), 0 );
-			FString const AssetName = FPackageName::GetLongPackageAssetName( AssetPath );
-			return FAnimationRecorderManager::Get().RecordAnimation( FoundActor, SkelComp, AssetPath, AssetName );
-		}
-	}
-#endif		// WITH_EDITOR
-
-	return false;
-}
-
-bool UEngine::HandleStopRecordAnimationCommand( UWorld* InWorld, const TCHAR* InStr, FOutputDevice& Ar )
-{
-#if WITH_EDITOR
-	const TCHAR* Str = InStr;
-
-	// parse actor name
-	TCHAR ActorName[128];
-	AActor* FoundActor = nullptr;
-	bool bStopAll = false;
-	if (FParse::Token( Str, ActorName, ARRAY_COUNT( ActorName ), 0 ))
-	{
-		FString const ActorNameStr = FString( ActorName );
-
-		if (ActorNameStr.ToLower() == TEXT( "all" ))
-		{
-			bStopAll = true;
-		}
-		else if (InWorld)
-		{
-			// search for the actor by name
-			for (ULevel* Level : InWorld->GetLevels())
-			{
-				if (Level)
-				{
-					for (AActor* Actor : Level->Actors)
-					{
-						if (Actor)
-						{
-							if (Actor->GetName() == ActorNameStr)
-							{
-								FoundActor = Actor;
-								break;
-							}
-						}
-					}
-				}
-			}
-		}
-	}
-
-	if (bStopAll)
-	{
-		FAnimationRecorderManager::Get().StopRecordingAllAnimations();
-		return true;
-	}
-	else if (FoundActor)
-	{
-		USkeletalMeshComponent* const SkelComp = FoundActor->FindComponentByClass<USkeletalMeshComponent>();
-		if (SkelComp)
-		{
-			FAnimationRecorderManager::Get().StopRecordingAnimation( FoundActor, SkelComp );
-			return true;
-		}
-	}
-#endif		// WITH_EDITOR
-
-	return false;
-}
-
 bool UEngine::HandleShowLogCommand( const TCHAR* Cmd, FOutputDevice& Ar )
 {
 	// Toggle display of console log window.
diff --git a/Engine/Source/Runtime/Launch/Private/LaunchEngineLoop.cpp b/Engine/Source/Runtime/Launch/Private/LaunchEngineLoop.cpp
--- a/Engine/Source/Runtime/Launch/Private/LaunchEngineLoop.cpp
+++ b/Engine/Source/Runtime/Launch/Private/LaunchEngineLoop.cpp
@@ -1963,14 +1963,6 @@ bool FEngineLoop::LoadStartupCoreModules()
 	if (FPlatformProcess::SupportsMultithreading())
 	{
 		FModuleManager::LoadModuleChecked<IMessagingModule>("Messaging");
-
-		if (!IsRunningCommandlet())
-		{
-			SessionService = FModuleManager::LoadModuleChecked<ISessionServicesModule>("SessionServices").GetSessionService();
-			SessionService->Start();
-		}
-
-		EngineService = new FEngineService();
 	}
 
 	SlowTask.EnterProgressFrame(10);
@@ -2196,6 +2188,18 @@ int32 FEngineLoop::Init()
 
 	SlowTask.EnterProgressFrame(30);
 
+	// initialize engine instance discovery
+	if (FPlatformProcess::SupportsMultithreading())
+	{
+		if (!IsRunningCommandlet())
+		{
+			SessionService = FModuleManager::LoadModuleChecked<ISessionServicesModule>("SessionServices").GetSessionService();
+			SessionService->Start();
+		}
+
+		EngineService = new FEngineService();
+	}
+
 	// Load all the post-engine init modules
 	if (!IProjectManager::Get().LoadModulesForProject(ELoadingPhase::PostEngineInit) || !IPluginManager::Get().LoadModulesForEnabledPlugins(ELoadingPhase::PostEngineInit))
 	{
diff --git a/Engine/Source/Runtime/LevelSequence/Private/LevelSequencePlayer.cpp b/Engine/Source/Runtime/LevelSequence/Private/LevelSequencePlayer.cpp
--- a/Engine/Source/Runtime/LevelSequence/Private/LevelSequencePlayer.cpp
+++ b/Engine/Source/Runtime/LevelSequence/Private/LevelSequencePlayer.cpp
@@ -339,7 +339,8 @@ void ULevelSequencePlayer::UpdateMovieSceneInstance(float CurrentPosition, float
 {
 	if(RootMovieSceneInstance.IsValid())
 	{
-		RootMovieSceneInstance->Update(CurrentPosition + StartTime, PreviousPosition + StartTime, *this);
+		EMovieSceneUpdateData UpdateData(CurrentPosition + StartTime, PreviousPosition + StartTime);
+		RootMovieSceneInstance->Update(UpdateData, *this);
 #if WITH_EDITOR
 		OnLevelSequencePlayerUpdate.Broadcast(*this, CurrentPosition, PreviousPosition);
 #endif
diff --git a/Engine/Source/Runtime/MediaAssets/Private/Assets/MediaSoundWave.cpp b/Engine/Source/Runtime/MediaAssets/Private/Assets/MediaSoundWave.cpp
--- a/Engine/Source/Runtime/MediaAssets/Private/Assets/MediaSoundWave.cpp
+++ b/Engine/Source/Runtime/MediaAssets/Private/Assets/MediaSoundWave.cpp
@@ -34,7 +34,7 @@ TSharedPtr<IMediaPlayer> UMediaSoundWave::GetPlayer() const
 }
 
 
-void UMediaSoundWave::SetMediaPlayer( UMediaPlayer* InMediaPlayer)
+void UMediaSoundWave::SetMediaPlayer(UMediaPlayer* InMediaPlayer)
 {
 	MediaPlayer = InMediaPlayer;
 
@@ -45,8 +45,19 @@ void UMediaSoundWave::SetMediaPlayer( UMediaPlayer* InMediaPlayer)
 /* USoundWave overrides
  *****************************************************************************/
 
-int32 UMediaSoundWave::GeneratePCMData( uint8* PCMData, const int32 SamplesNeeded)
+int32 UMediaSoundWave::GeneratePCMData(uint8* PCMData, const int32 SamplesNeeded)
 {
+	if (AudioTrack.IsValid())
+	{
+		NumChannels = AudioTrack->GetNumChannels();
+		SampleRate = AudioTrack->GetSamplesPerSecond();
+	}
+
+	if ((NumChannels == 0) || (SampleRate == 0))
+	{
+		return 0;
+	}
+
 	// drain media sample queue
 	TSharedPtr<TArray<uint8>, ESPMode::ThreadSafe> Sample;
 
@@ -74,25 +85,25 @@ int32 UMediaSoundWave::GeneratePCMData( uint8* PCMData, const int32 SamplesNeede
 }
 
 
-FByteBulkData* UMediaSoundWave::GetCompressedData( FName Format)
+FByteBulkData* UMediaSoundWave::GetCompressedData(FName Format)
 {
 	return nullptr;
 }
 
 
-int32 UMediaSoundWave::GetResourceSizeForFormat( FName Format)
+int32 UMediaSoundWave::GetResourceSizeForFormat(FName Format)
 {
 	return 0;
 }
 
 
-void UMediaSoundWave::InitAudioResource( FByteBulkData& CompressedData)
+void UMediaSoundWave::InitAudioResource(FByteBulkData& CompressedData)
 {
 	check(false); // should never be pushing compressed data to this class
 }
 
 
-bool UMediaSoundWave::InitAudioResource( FName Format)
+bool UMediaSoundWave::InitAudioResource(FName Format)
 {
 	return true;
 }
@@ -101,7 +112,7 @@ bool UMediaSoundWave::InitAudioResource( FName Format)
 /* UObject overrides
  *****************************************************************************/
 
-void UMediaSoundWave::GetAssetRegistryTags( TArray<FAssetRegistryTag>& OutTags) const
+void UMediaSoundWave::GetAssetRegistryTags(TArray<FAssetRegistryTag>& OutTags) const
 {
 }
 
@@ -112,7 +123,7 @@ SIZE_T UMediaSoundWave::GetResourceSize(EResourceSizeMode::Type Mode)
 }
 
 
-void UMediaSoundWave::Serialize( FArchive& Ar)
+void UMediaSoundWave::Serialize(FArchive& Ar)
 {
 	// do not call the USoundWave version of serialize
 	USoundBase::Serialize(Ar);
@@ -200,17 +211,6 @@ void UMediaSoundWave::InitializeTrack()
 		}
 	}
 
-	if (AudioTrack.IsValid())
-	{
-		NumChannels = AudioTrack->GetNumChannels();
-		SampleRate = AudioTrack->GetSamplesPerSecond();
-	}
-	else
-	{
-		NumChannels = 0;
-		SampleRate = 0;
-	}
-
 	// connect to new track
 	if (AudioTrack.IsValid())
 	{
diff --git a/Engine/Source/Runtime/MovieScene/Private/MovieScene.cpp b/Engine/Source/Runtime/MovieScene/Private/MovieScene.cpp
--- a/Engine/Source/Runtime/MovieScene/Private/MovieScene.cpp
+++ b/Engine/Source/Runtime/MovieScene/Private/MovieScene.cpp
@@ -84,6 +84,11 @@ FMovieSceneSpawnable* UMovieScene::FindSpawnable( const TFunctionRef<bool(FMovie
 #endif //WITH_EDITOR
 
 
+FMovieSceneSpawnable& UMovieScene::GetSpawnable(int32 Index)
+{
+	return Spawnables[Index];
+}
+
 int32 UMovieScene::GetSpawnableCount() const
 {
 	return Spawnables.Num();
@@ -392,30 +397,30 @@ class UMovieSceneTrack* UMovieScene::AddMasterTrack( TSubclassOf<UMovieSceneTrac
 }
 
 
-UMovieSceneTrack* UMovieScene::AddShotTrack( TSubclassOf<UMovieSceneTrack> TrackClass )
+UMovieSceneTrack* UMovieScene::AddCameraCutTrack( TSubclassOf<UMovieSceneTrack> TrackClass )
 {
-	if( !ShotTrack )
+	if( !CameraCutTrack )
 	{
 		Modify();
-		ShotTrack = NewObject<UMovieSceneTrack>(this, TrackClass, FName("Shots"), RF_Transactional);
+		CameraCutTrack = NewObject<UMovieSceneTrack>(this, TrackClass, FName("Camera Cuts"), RF_Transactional);
 	}
 
-	return ShotTrack;
+	return CameraCutTrack;
 }
 
 
-UMovieSceneTrack* UMovieScene::GetShotTrack()
+UMovieSceneTrack* UMovieScene::GetCameraCutTrack()
 {
-	return ShotTrack;
+	return CameraCutTrack;
 }
 
 
-void UMovieScene::RemoveShotTrack()
+void UMovieScene::RemoveCameraCutTrack()
 {
-	if( ShotTrack )
+	if( CameraCutTrack )
 	{
 		Modify();
-		ShotTrack = nullptr;
+		CameraCutTrack = nullptr;
 	}
 }
 
diff --git a/Engine/Source/Runtime/MovieScene/Private/MovieSceneCommonHelpers.cpp b/Engine/Source/Runtime/MovieScene/Private/MovieSceneCommonHelpers.cpp
--- a/Engine/Source/Runtime/MovieScene/Private/MovieSceneCommonHelpers.cpp
+++ b/Engine/Source/Runtime/MovieScene/Private/MovieSceneCommonHelpers.cpp
@@ -125,6 +125,45 @@ UMovieSceneSection* MovieSceneHelpers::FindNearestSectionAtTime( const TArray<UM
 }
 
 
+void MovieSceneHelpers::SortConsecutiveSections(TArray<UMovieSceneSection*>& Sections)
+{
+	Sections.Sort([](const UMovieSceneSection& A, const UMovieSceneSection& B)
+		{
+			return A.GetStartTime() < B.GetStartTime();
+		}
+	);
+}
+
+void MovieSceneHelpers::FixupConsecutiveSections(TArray<UMovieSceneSection*>& Sections, UMovieSceneSection& Section, bool bDelete)
+{
+	// Find the previous section and extend it to take the place of the section being deleted
+	int32 SectionIndex = INDEX_NONE;
+
+	if (Sections.Find(&Section, SectionIndex))
+	{
+		int32 PrevSectionIndex = SectionIndex - 1;
+		if( Sections.IsValidIndex( PrevSectionIndex ) )
+		{
+			// Extend the previous section
+			Sections[PrevSectionIndex]->SetEndTime( bDelete ? Section.GetEndTime() : Section.GetStartTime() );
+		}
+
+		if( !bDelete )
+		{
+			int32 NextSectionIndex = SectionIndex + 1;
+			if(Sections.IsValidIndex(NextSectionIndex))
+			{
+				// Shift the next CameraCut's start time so that it starts when the new CameraCut ends
+				Sections[NextSectionIndex]->SetStartTime(Section.GetEndTime());
+			}
+		}
+	}
+
+	SortConsecutiveSections(Sections);
+}
+
+
+
 USceneComponent* MovieSceneHelpers::SceneComponentFromRuntimeObject(UObject* Object)
 {
 	AActor* Actor = Cast<AActor>(Object);
diff --git a/Engine/Source/Runtime/MovieScene/Private/MovieSceneNameableTrack.cpp b/Engine/Source/Runtime/MovieScene/Private/MovieSceneNameableTrack.cpp
--- a/Engine/Source/Runtime/MovieScene/Private/MovieSceneNameableTrack.cpp
+++ b/Engine/Source/Runtime/MovieScene/Private/MovieSceneNameableTrack.cpp
@@ -37,12 +37,18 @@ FText UMovieSceneNameableTrack::GetDisplayName() const
 {
 	if (DisplayName.IsEmpty())
 	{
-		return LOCTEXT("UnnamedTrackName", "Unnamed Track");
+		return GetDefaultDisplayName();
 	}
 
 	return DisplayName;
 }
 
+FText UMovieSceneNameableTrack::GetDefaultDisplayName() const
+
+{ 
+	return LOCTEXT("UnnamedTrackName", "Unnamed Track"); 
+}
+
 #endif
 
 
diff --git a/Engine/Source/Runtime/MovieScene/Private/MovieSceneSection.cpp b/Engine/Source/Runtime/MovieScene/Private/MovieSceneSection.cpp
--- a/Engine/Source/Runtime/MovieScene/Private/MovieSceneSection.cpp
+++ b/Engine/Source/Runtime/MovieScene/Private/MovieSceneSection.cpp
@@ -4,7 +4,7 @@
 #include "MovieSceneSection.h"
 
 
-UMovieSceneSection::UMovieSceneSection( const FObjectInitializer& ObjectInitializer )
+UMovieSceneSection::UMovieSceneSection(const FObjectInitializer& ObjectInitializer)
 	: Super( ObjectInitializer )
 	, StartTime(0.0f)
 	, EndTime(0.0f)
@@ -15,15 +15,17 @@ UMovieSceneSection::UMovieSceneSection( const FObjectInitializer& ObjectInitiali
 	, bIsInfinite(false)
 { }
 
-bool
-UMovieSceneSection::TryModify(bool bAlwaysMarkDirty)
+
+bool UMovieSceneSection::TryModify(bool bAlwaysMarkDirty)
 {
 	if (IsLocked())
 	{
 		return false;
 	}
 
-	return Modify(bAlwaysMarkDirty);
+	Modify(bAlwaysMarkDirty);
+
+	return true;
 }
 
 
@@ -133,6 +135,7 @@ void UMovieSceneSection::TrimSection(float TrimTime, bool bTrimLeft)
 	}
 }
 
+
 void UMovieSceneSection::AddKeyToCurve(FRichCurve& InCurve, float Time, float Value, EMovieSceneKeyInterpolation Interpolation, const bool bUnwindRotation)
 {
 	if (IsTimeWithinSection(Time))
@@ -140,7 +143,6 @@ void UMovieSceneSection::AddKeyToCurve(FRichCurve& InCurve, float Time, float Va
 		if (TryModify())
 		{
 			FKeyHandle ExistingKeyHandle = InCurve.FindKey(Time);
-				
 			FKeyHandle NewKeyHandle = InCurve.UpdateOrAddKey(Time, Value, bUnwindRotation);
 
 			if (!InCurve.IsKeyHandleValid(ExistingKeyHandle) && InCurve.IsKeyHandleValid(NewKeyHandle))
@@ -151,6 +153,7 @@ void UMovieSceneSection::AddKeyToCurve(FRichCurve& InCurve, float Time, float Va
 	}
 }
 
+
 void UMovieSceneSection::SetCurveDefault(FRichCurve& InCurve, float Value)
 {
 	if (TryModify())
@@ -158,4 +161,3 @@ void UMovieSceneSection::SetCurveDefault(FRichCurve& InCurve, float Value)
 		InCurve.SetDefaultValue(Value);
 	}
 }
-
diff --git a/Engine/Source/Runtime/MovieScene/Private/MovieSceneSequenceInstance.cpp b/Engine/Source/Runtime/MovieScene/Private/MovieSceneSequenceInstance.cpp
--- a/Engine/Source/Runtime/MovieScene/Private/MovieSceneSequenceInstance.cpp
+++ b/Engine/Source/Runtime/MovieScene/Private/MovieSceneSequenceInstance.cpp
@@ -46,21 +46,27 @@ FGuid FMovieSceneSequenceInstance::FindParentObjectId(UObject& Object) const
 UObject* FMovieSceneSequenceInstance::FindObject(const FGuid& ObjectId, const IMovieScenePlayer& Player) const
 {
 	// Attempt to find a possessable first
-	FMovieScenePossessable* Possessable = MovieSceneSequence->GetMovieScene()->FindPossessable(ObjectId);
-	if (Possessable)
+	UMovieScene* MovieScene = MovieSceneSequence->GetMovieScene();
+	if (MovieScene != nullptr)
 	{
-		UObject* ParentObject = Player.GetPlaybackContext();
-		if (Possessable->GetParent().IsValid())
+		FMovieScenePossessable* Possessable = MovieScene->FindPossessable(ObjectId);
+		if (Possessable)
 		{
-			ParentObject = FindObject(Possessable->GetParent(), Player);
-		}
+			UObject* ParentObject = Player.GetPlaybackContext();
+			if (Possessable->GetParent().IsValid())
+			{
+				ParentObject = FindObject(Possessable->GetParent(), Player);
+			}
 
-		return MovieSceneSequence->FindPossessableObject(ObjectId, ParentObject);
-	}
-	else
-	{
-		return FindSpawnedObject(ObjectId);
+			return MovieSceneSequence->FindPossessableObject(ObjectId, ParentObject);
+		}
+		else
+		{
+			return FindSpawnedObject(ObjectId);
+		}
 	}
+
+	return nullptr;
 }
 
 UObject* FMovieSceneSequenceInstance::FindSpawnedObject(const FGuid& ObjectId) const
@@ -72,9 +78,9 @@ UObject* FMovieSceneSequenceInstance::FindSpawnedObject(const FGuid& ObjectId) c
 void FMovieSceneSequenceInstance::SaveState(class IMovieScenePlayer& Player)
 {
 	TArray<UObject*> NoObjects;
-	if (ShotTrackInstance.IsValid())
+	if (CameraCutTrackInstance.IsValid())
 	{
-		ShotTrackInstance->SaveState(NoObjects, Player, *this);
+		CameraCutTrackInstance->SaveState(NoObjects, Player, *this);
 	}
 
 	for (auto& Pair : MasterTrackInstances)
@@ -121,9 +127,9 @@ void FMovieSceneSequenceInstance::RestoreState(class IMovieScenePlayer& Player)
 	}
 
 	TArray<UObject*> NoObjects;
-	if (ShotTrackInstance.IsValid())
+	if (CameraCutTrackInstance.IsValid())
 	{
-		ShotTrackInstance->RestoreState(NoObjects, Player, *this);
+		CameraCutTrackInstance->RestoreState(NoObjects, Player, *this);
 	}
 
 	for (auto& Pair : MasterTrackInstances)
@@ -133,7 +139,7 @@ void FMovieSceneSequenceInstance::RestoreState(class IMovieScenePlayer& Player)
 }
 
 
-void FMovieSceneSequenceInstance::Update( float Position, float LastPosition, class IMovieScenePlayer& Player )
+void FMovieSceneSequenceInstance::Update( EMovieSceneUpdateData& UpdateData, class IMovieScenePlayer& Player )
 {
 	// Remove any stale runtime objects
 	TMap<FGuid, FMovieSceneObjectBindingInstance>::TIterator ObjectIt = ObjectBindingInstances.CreateIterator();
@@ -156,14 +162,17 @@ void FMovieSceneSequenceInstance::Update( float Position, float LastPosition, cl
 
 	Player.GetSpawnRegister().PreUpdateSequenceInstance(*this, Player);
 
-	UpdateInternal(Position, LastPosition, Player, MSUP_PreUpdate);
-	UpdateInternal(Position, LastPosition, Player, MSUP_Update);
-	UpdateInternal(Position, LastPosition, Player, MSUP_PostUpdate);
+	UpdateData.UpdatePass = MSUP_PreUpdate;
+	UpdateInternal(UpdateData, Player);
+	UpdateData.UpdatePass = MSUP_Update;
+	UpdateInternal(UpdateData, Player);
+	UpdateData.UpdatePass = MSUP_PostUpdate;
+	UpdateInternal(UpdateData, Player);
 
 	Player.GetSpawnRegister().PostUpdateSequenceInstance(*this, Player);
 }
 
-void FMovieSceneSequenceInstance::UpdateInternal( float Position, float LastPosition, class IMovieScenePlayer& Player, EMovieSceneUpdatePass UpdatePass )
+void FMovieSceneSequenceInstance::UpdateInternal( EMovieSceneUpdateData& UpdateData, class IMovieScenePlayer& Player )
 {
 	// Refresh time range so that spawnables can be created if they fall within the playback range, or destroyed if not
 	UMovieScene* MovieScene = MovieSceneSequence->GetMovieScene();
@@ -174,9 +183,9 @@ void FMovieSceneSequenceInstance::UpdateInternal( float Position, float LastPosi
 	// update each master track
 	for( FMovieSceneInstanceMap::TIterator It( MasterTrackInstances ); It; ++It )
 	{
-		if (It.Value()->HasUpdatePasses() & UpdatePass)
+		if (It.Value()->HasUpdatePasses() & UpdateData.UpdatePass)
 		{
-			It.Value()->Update( Position, LastPosition, NoObjects, Player, *this, UpdatePass);
+			It.Value()->Update( UpdateData, NoObjects, Player, *this);
 		}
 	}
 
@@ -188,19 +197,19 @@ void FMovieSceneSequenceInstance::UpdateInternal( float Position, float LastPosi
 		
 		for( FMovieSceneInstanceMap::TIterator It = ObjectBindingInstance.TrackInstances.CreateIterator(); It; ++It )
 		{
-			if (It.Value()->HasUpdatePasses() & UpdatePass)
+			if (It.Value()->HasUpdatePasses() & UpdateData.UpdatePass)
 			{
-				It.Value()->Update( Position, LastPosition, ObjectBindingInstance.RuntimeObjects, Player, *this, UpdatePass );
+				It.Value()->Update( UpdateData, ObjectBindingInstance.RuntimeObjects, Player, *this );
 			}
 		}
 	}
 
-	// update shot track last to make sure spawnable cameras are there, and to override sub-shots
-	if (ShotTrackInstance.IsValid())
+	// update camera cut track last to make sure spawnable cameras are there, and to override sub-shots
+	if (CameraCutTrackInstance.IsValid())
 	{
-		if (ShotTrackInstance->HasUpdatePasses() & UpdatePass)
+		if (CameraCutTrackInstance->HasUpdatePasses() & UpdateData.UpdatePass)
 		{
-			ShotTrackInstance->Update(Position, LastPosition, NoObjects, Player, *this, UpdatePass );
+			CameraCutTrackInstance->Update(UpdateData, NoObjects, Player, *this );
 		}
 	}
 }
@@ -211,28 +220,28 @@ void FMovieSceneSequenceInstance::RefreshInstance( IMovieScenePlayer& Player )
 	UMovieScene* MovieScene = MovieSceneSequence->GetMovieScene();
 	TimeRange = MovieScene->GetPlaybackRange();
 
-	UMovieSceneTrack* ShotTrack = MovieScene->GetShotTrack();
+	UMovieSceneTrack* CameraCutTrack = MovieScene->GetCameraCutTrack();
 
-	if (ShotTrack != nullptr)
+	if (CameraCutTrack != nullptr)
 	{
-		FMovieSceneInstanceMap ShotTrackInstanceMap;
+		FMovieSceneInstanceMap CameraCutTrackInstanceMap;
 
-		if (ShotTrackInstance.IsValid())
+		if (CameraCutTrackInstance.IsValid())
 		{
-			ShotTrackInstanceMap.Add(ShotTrack, ShotTrackInstance);
+			CameraCutTrackInstanceMap.Add(CameraCutTrack, CameraCutTrackInstance);
 		}
 
 		TArray<UObject*> Objects;
 		TArray<UMovieSceneTrack*> Tracks;
-		Tracks.Add(ShotTrack);
-		RefreshInstanceMap(Tracks, Objects, ShotTrackInstanceMap, Player);
+		Tracks.Add(CameraCutTrack);
+		RefreshInstanceMap(Tracks, Objects, CameraCutTrackInstanceMap, Player);
 
-		ShotTrackInstance = ShotTrackInstanceMap.FindRef(ShotTrack);
+		CameraCutTrackInstance = CameraCutTrackInstanceMap.FindRef(CameraCutTrack);
 	}
-	else if(ShotTrackInstance.IsValid())
+	else if(CameraCutTrackInstance.IsValid())
 	{
-		ShotTrackInstance->ClearInstance(Player, *this);
-		ShotTrackInstance.Reset();
+		CameraCutTrackInstance->ClearInstance(Player, *this);
+		CameraCutTrackInstance.Reset();
 	}
 
 	// Get all the master tracks and create instances for them if needed
diff --git a/Engine/Source/Runtime/MovieScene/Public/IMovieSceneSpawnRegister.h b/Engine/Source/Runtime/MovieScene/Public/IMovieSceneSpawnRegister.h
--- a/Engine/Source/Runtime/MovieScene/Public/IMovieSceneSpawnRegister.h
+++ b/Engine/Source/Runtime/MovieScene/Public/IMovieSceneSpawnRegister.h
@@ -3,9 +3,24 @@
 #pragma once
 
 #include "MovieSceneSequenceInstance.h"
+#include "ValueOrError.h"
 
 class IMovieScenePlayer;
 class FMovieSceneSequenceInstance;
+class UMovieScene;
+
+/** Struct used for defining a new spawnable type */
+struct FNewSpawnable
+{
+	FNewSpawnable() : Blueprint(nullptr) {}
+	FNewSpawnable(UBlueprint* InBlueprint, FString InName) : Blueprint(InBlueprint), Name(MoveTemp(InName)) {}
+
+	/** The blueprint that defines the spawnable */
+	UBlueprint* Blueprint;
+
+	/** The desired name of the new spawnable */
+	FString Name;
+};
 
 /**
  * Class responsible for managing spawnables in a movie scene
@@ -87,6 +102,17 @@ public:
 	 * @param Player 		The player that is playing the movie scene
 	 */
 	virtual void PostUpdateSequenceInstance(FMovieSceneSequenceInstance& Instance, IMovieScenePlayer& Player) { }
+
+#if WITH_EDITOR
+	/**
+	 * Create a new spawnable type from the given source object
+	 *
+	 * @param SourceObject		The source object to create the spawnable from
+	 * @param OwnerMovieScene	The owner movie scene that this spawnable type should reside in
+	 * @return the new spawnable type, or nullptr
+	 */
+	virtual TValueOrError<FNewSpawnable, FText> CreateNewSpawnableType(UObject& SourceObject, UMovieScene& OwnerMovieScene) { return MakeError(NSLOCTEXT("SpawnRegister", "NotSupported", "Not supported")); }
+#endif
 };
 
 
diff --git a/Engine/Source/Runtime/MovieScene/Public/IMovieSceneTrackInstance.h b/Engine/Source/Runtime/MovieScene/Public/IMovieSceneTrackInstance.h
--- a/Engine/Source/Runtime/MovieScene/Public/IMovieSceneTrackInstance.h
+++ b/Engine/Source/Runtime/MovieScene/Public/IMovieSceneTrackInstance.h
@@ -13,6 +13,29 @@ enum EMovieSceneUpdatePass
 	MSUP_PostUpdate = 0x00000004
 };
 
+struct EMovieSceneUpdateData
+{
+	float Position;
+	float LastPosition;
+	bool bPreroll;
+	EMovieSceneUpdatePass UpdatePass;
+
+	EMovieSceneUpdateData()
+	{
+		Position = 0.0f;
+		LastPosition = 0.0f;
+		bPreroll = false;
+		UpdatePass = MSUP_PreUpdate;
+	}
+	EMovieSceneUpdateData(float InPosition, float InLastPosition)
+	{
+		Position = InPosition;
+		LastPosition = InLastPosition;
+		bPreroll = false;
+		UpdatePass = MSUP_PreUpdate;
+	}
+};
+
 /**
  * A track instance holds the live objects for a track.  
  */
@@ -38,13 +61,11 @@ public:
 	/**
 	 * Main update function for track instances.  Called in game and in editor when we update a moviescene.
 	 *
-	 * @param Position			The current position of the moviescene that is playing
-	 * @param LastPosition		The previous playback position
+	 * @param UpdateData		The current and previous position of the moviescene that is playing. The update pass.
 	 * @param RuntimeObjects	Runtime objects bound to this instance (if any)
 	 * @param Player			The playback interface.  Contains state and some other functionality for runtime playback
-	 * @param UpdatePass        Which update pass
 	 */
-	virtual void Update(float Position, float LastPosition, const TArray<UObject*>& RuntimeObjects, IMovieScenePlayer& Player, FMovieSceneSequenceInstance& SequenceInstance, EMovieSceneUpdatePass UpdatePass) = 0;
+	virtual void Update(EMovieSceneUpdateData& UpdateData, const TArray<UObject*>& RuntimeObjects, IMovieScenePlayer& Player, FMovieSceneSequenceInstance& SequenceInstance) = 0;
 	
 	/*
 	 * Which update passes does this track instance evaluate in?
diff --git a/Engine/Source/Runtime/MovieScene/Public/MovieScene.h b/Engine/Source/Runtime/MovieScene/Public/MovieScene.h
--- a/Engine/Source/Runtime/MovieScene/Public/MovieScene.h
+++ b/Engine/Source/Runtime/MovieScene/Public/MovieScene.h
@@ -123,6 +123,14 @@ public:
 	 */
 	FMovieSceneSpawnable* FindSpawnable(const FGuid& Guid);
 
+	/**
+	 * Grabs a reference to a specific spawnable by index.
+	 *
+	 * @param Index of spawnable to return. Must be between 0 and GetSpawnableCount()
+	 * @return Returns the specified spawnable by index.
+	 */
+	FMovieSceneSpawnable& GetSpawnable(int32 Index);
+
 	/**
 	 * Get the number of spawnable objects in this scene.
 	 *
@@ -324,20 +332,20 @@ public:
 	// @todo sequencer: the following methods really shouldn't be here
 
 	/**
-	 * Adds a new shot track if it doesn't exist 
-	 * A shot track is a special kind of sub-movie scene track that allows for cutting between camera views
+	 * Adds a new camera cut track if it doesn't exist 
+	 * A camera cut track allows for cutting between camera views
 	 * There is only one per movie scene. 
 	 *
-	 * @param TrackClass  The shot track class type
-	 * @return The created shot track
+	 * @param TrackClass  The camera cut track class type
+	 * @return The created camera cut track
 	 */
-	UMovieSceneTrack* AddShotTrack( TSubclassOf<UMovieSceneTrack> TrackClass );
+	UMovieSceneTrack* AddCameraCutTrack( TSubclassOf<UMovieSceneTrack> TrackClass );
 	
-	/** @return The shot track if it exists. */
-	UMovieSceneTrack* GetShotTrack();
+	/** @return The camera cut track if it exists. */
+	UMovieSceneTrack* GetCameraCutTrack();
 
-	/** Removes the shot track if it exists. */
-	void RemoveShotTrack();
+	/** Removes the camera cut track if it exists. */
+	void RemoveCameraCutTrack();
 
 public:
 
@@ -453,9 +461,9 @@ private:
 	UPROPERTY()
 	TArray<UMovieSceneTrack*> MasterTracks;
 
-	/** The shot track is a specialized track for switching between cameras on a cinematic */
+	/** The camera cut track is a specialized track for switching between cameras on a cinematic */
 	UPROPERTY()
-	UMovieSceneTrack* ShotTrack;
+	UMovieSceneTrack* CameraCutTrack;
 
 #if WITH_EDITORONLY_DATA
 	/** Maps object GUIDs to user defined display names. */
diff --git a/Engine/Source/Runtime/MovieScene/Public/MovieSceneCommonHelpers.h b/Engine/Source/Runtime/MovieScene/Public/MovieSceneCommonHelpers.h
--- a/Engine/Source/Runtime/MovieScene/Public/MovieSceneCommonHelpers.h
+++ b/Engine/Source/Runtime/MovieScene/Public/MovieSceneCommonHelpers.h
@@ -34,6 +34,20 @@ public:
 	 */
 	static UMovieSceneSection* FindNearestSectionAtTime( const TArray<UMovieSceneSection*>& Sections, float Time );
 
+	/*
+	 * Fix up consecutive sections so that there are no gaps
+	 * 
+	 * @param Sections All the sections
+	 * @param Section The section that was modified 
+	 * @param bDelete Was this a deletion?
+	 */
+	static void FixupConsecutiveSections(TArray<UMovieSceneSection*>& Sections, UMovieSceneSection& Section, bool bDelete);
+
+	/*
+ 	 * Sort consecutive sections so that they are in order based on start time
+ 	 */
+	static void SortConsecutiveSections(TArray<UMovieSceneSection*>& Sections);
+
 	/**
 	 * Get the scene component from the runtime object
 	 *
diff --git a/Engine/Source/Runtime/MovieScene/Public/MovieSceneNameableTrack.h b/Engine/Source/Runtime/MovieScene/Public/MovieSceneNameableTrack.h
--- a/Engine/Source/Runtime/MovieScene/Public/MovieSceneNameableTrack.h
+++ b/Engine/Source/Runtime/MovieScene/Public/MovieSceneNameableTrack.h
@@ -30,6 +30,7 @@ public:
 	// UMovieSceneTrack interface
 
 	MOVIESCENE_API virtual FText GetDisplayName() const override;
+	MOVIESCENE_API virtual FText GetDefaultDisplayName() const;
 
 private:
 
diff --git a/Engine/Source/Runtime/MovieScene/Public/MovieSceneSection.h b/Engine/Source/Runtime/MovieScene/Public/MovieSceneSection.h
--- a/Engine/Source/Runtime/MovieScene/Public/MovieSceneSection.h
+++ b/Engine/Source/Runtime/MovieScene/Public/MovieSceneSection.h
@@ -9,19 +9,19 @@
 /**
  * Base class for movie scene sections
  */
-UCLASS( abstract, MinimalAPI )
+UCLASS(abstract, MinimalAPI)
 class UMovieSceneSection
 	: public UObject
 {
 	GENERATED_UCLASS_BODY()
 public:
 
-	/*
-	 * Calls Modify if this section can be modified. ie. can't be modified if it's locked
+	/**
+	 * Calls Modify if this section can be modified, i.e. can't be modified if it's locked
 	 *
-	 * @return The return value of Modify()
+	 * @return Returns whether this section is locked or not
 	 */
-	virtual MOVIESCENE_API bool TryModify( bool bAlwaysMarkDirty=true );
+	virtual MOVIESCENE_API bool TryModify(bool bAlwaysMarkDirty=true);
 
 	/**
 	 * @return The start time of the section
@@ -52,7 +52,7 @@ public:
 	 * 
 	 * @param InEndTime	The new end time
 	 */
-	void SetStartTime( float NewStartTime )
+	void SetStartTime(float NewStartTime)
 	{ 
 		if (TryModify())
 		{
@@ -65,7 +65,7 @@ public:
 	 * 
 	 * @param InEndTime	The new end time
 	 */
-	void SetEndTime( float NewEndTime )
+	void SetEndTime(float NewEndTime)
 	{ 
 		if (TryModify())
 		{
@@ -108,7 +108,7 @@ public:
 	 * @param Position	The position to check
 	 * @return true if the position is within the timespan, false otherwise
 	 */
-	bool IsTimeWithinSection( float Position ) const 
+	bool IsTimeWithinSection(float Position) const 
 	{
 		return Position >= StartTime && Position <= EndTime;
 	}
@@ -119,7 +119,7 @@ public:
 	 * @param DeltaTime	The distance in time to move the curve
 	 * @param KeyHandles The key handles to operate on
 	 */
-	virtual void MoveSection( float DeltaTime, TSet<FKeyHandle>& KeyHandles )
+	virtual void MoveSection(float DeltaTime, TSet<FKeyHandle>& KeyHandles)
 	{
 		if (TryModify())
 		{
@@ -135,7 +135,7 @@ public:
 	 * @param bFromStart Whether to dilate from the beginning or end (whichever stays put)
 	 * @param KeyHandles The key handles to operate on
 	 */
-	virtual void DilateSection( float DilationFactor, float Origin, TSet<FKeyHandle>& KeyHandles )
+	virtual void DilateSection(float DilationFactor, float Origin, TSet<FKeyHandle>& KeyHandles)
 	{
 		if (TryModify())
 		{
@@ -167,6 +167,17 @@ public:
 	 */
 	virtual void GetKeyHandles(TSet<FKeyHandle>& KeyHandles) const {};
 
+	/**
+	 * Get the data structure representing the specified key.
+	 *
+	 * @param KeyHandle The handle of the key.
+	 * @return The key's data structure representation, or nullptr if key not found or no structure available.
+	 */
+	virtual TSharedPtr<FStructOnScope> GetKeyStruct(const TArray<FKeyHandle>& KeyHandles)
+	{
+		return nullptr;
+	}
+
 	/**
 	 * Gets all snap times for this section
 	 *
@@ -211,7 +222,7 @@ public:
 	 * @param KeyParams The keying parameters
 	 * @param bUnwindRotation Unwind rotation
 	 */
-	void MOVIESCENE_API AddKeyToCurve( FRichCurve& InCurve, float Time, float Value, EMovieSceneKeyInterpolation Interpolation, const bool bUnwindRotation = false );
+	void MOVIESCENE_API AddKeyToCurve(FRichCurve& InCurve, float Time, float Value, EMovieSceneKeyInterpolation Interpolation, const bool bUnwindRotation = false);
 
 	/**
 	 * Sets the default value for a curve.
@@ -219,7 +230,7 @@ public:
 	 * @param InCurve The curve to set a default value on.
 	 * @param Value The value to use as the default.
 	 */
-	void MOVIESCENE_API SetCurveDefault( FRichCurve& InCurve, float Value );
+	void MOVIESCENE_API SetCurveDefault(FRichCurve& InCurve, float Value);
 
 	/**
 	 * Checks to see if this section overlaps with an array of other sections
diff --git a/Engine/Source/Runtime/MovieScene/Public/MovieSceneSequenceInstance.h b/Engine/Source/Runtime/MovieScene/Public/MovieSceneSequenceInstance.h
--- a/Engine/Source/Runtime/MovieScene/Public/MovieSceneSequenceInstance.h
+++ b/Engine/Source/Runtime/MovieScene/Public/MovieSceneSequenceInstance.h
@@ -71,10 +71,10 @@ public:
 	/**
 	 * Updates this movie scene.
 	 *
-	 * @param Position The local playback position.
+	 * @param UpdateData The current and previous position of the moviescene that is playing. The update pass.
 	 * @param Player Movie scene player interface for interaction with runtime data.
 	 */
-	MOVIESCENE_API void Update(float Position, float LastPosition, IMovieScenePlayer& Player);
+	MOVIESCENE_API void Update(EMovieSceneUpdateData& UpdateData, IMovieScenePlayer& Player);
 
 	/**
 	 * Refreshes the existing instance.
@@ -128,7 +128,7 @@ protected:
 
 	void RefreshInstanceMap(const TArray<UMovieSceneTrack*>& Tracks, const TArray<UObject*>& RuntimeObjects, FMovieSceneInstanceMap& TrackInstances, IMovieScenePlayer& Player);
 
-	void UpdateInternal(float Position, float LastPosition, IMovieScenePlayer& Player, EMovieSceneUpdatePass UpdatePass);
+	void UpdateInternal(EMovieSceneUpdateData& UpdateData, IMovieScenePlayer& Player);
 
 	/** Update the object binding instance for the specified object */
 	void UpdateObjectBinding(const FGuid& ObjectId, IMovieScenePlayer& Player);
@@ -152,8 +152,8 @@ private:
 	/** MovieScene that is instanced */
 	const TWeakObjectPtr<UMovieSceneSequence> MovieSceneSequence;
 
-	/** The shot track instance map */
-	TSharedPtr<IMovieSceneTrackInstance> ShotTrackInstance;
+	/** The camera cut track instance map */
+	TSharedPtr<IMovieSceneTrackInstance> CameraCutTrackInstance;
 
 	/** All Master track instances */
 	FMovieSceneInstanceMap MasterTrackInstances;
diff --git a/Engine/Source/Runtime/MovieScene/Public/MovieSceneSpawnable.h b/Engine/Source/Runtime/MovieScene/Public/MovieSceneSpawnable.h
--- a/Engine/Source/Runtime/MovieScene/Public/MovieSceneSpawnable.h
+++ b/Engine/Source/Runtime/MovieScene/Public/MovieSceneSpawnable.h
@@ -30,7 +30,7 @@ public:
 	/** FMovieSceneSpawnable default constructor */
 	FMovieSceneSpawnable() { }
 
-/** FMovieSceneSpawnable initialization constructor */
+	/** FMovieSceneSpawnable initialization constructor */
 	FMovieSceneSpawnable(const FString& InitName, UClass* InitClass)
 		: Guid(FGuid::NewGuid())
 		, Name(InitName)
@@ -62,6 +62,17 @@ public:
 		return Guid;
 	}
 
+	/**
+	 * Set the unique identifier of this spawnable object
+	 * @param InGuid The new GUID for this spawnable
+	 * @note Be careful - this guid may be referenced by spawnable/possessable child-parent relationships.
+	 * @see GetGuid
+	 */
+	void SetGuid(const FGuid& InGuid)
+	{
+		Guid = InGuid;
+	}
+
 	/**
 	 * Get the name of the spawnable object.
 	 *
diff --git a/Engine/Source/Runtime/MovieScene/Public/MovieSceneTrack.h b/Engine/Source/Runtime/MovieScene/Public/MovieSceneTrack.h
--- a/Engine/Source/Runtime/MovieScene/Public/MovieSceneTrack.h
+++ b/Engine/Source/Runtime/MovieScene/Public/MovieSceneTrack.h
@@ -51,6 +51,14 @@ public:
 		return false;
 	}
 
+	/**
+	 * @return Whether or not this track's section bounds should be added to the play range
+	 */
+	virtual bool AddsSectionBoundsToPlayRange() const
+	{
+		return false;
+	}
+
 public:
 
 	/**
diff --git a/Engine/Source/Runtime/MovieSceneCapture/Private/CompositionGraphCaptureProtocol.cpp b/Engine/Source/Runtime/MovieSceneCapture/Private/CompositionGraphCaptureProtocol.cpp
--- a/Engine/Source/Runtime/MovieSceneCapture/Private/CompositionGraphCaptureProtocol.cpp
+++ b/Engine/Source/Runtime/MovieSceneCapture/Private/CompositionGraphCaptureProtocol.cpp
@@ -86,6 +86,7 @@ struct FSceneViewExtension : ISceneViewExtension
 			PostProcessingMaterial->OverrideBlendableSettings(InView, 1.f);
 		}
 		
+
 		// Ensure we're rendering at full size
 		InView.ViewRect = InView.UnscaledViewRect;
 
@@ -126,6 +127,18 @@ bool FCompositionGraphCaptureProtocol::Initialize(const FCaptureProtocolInitSett
 		RenderPasses = ProtocolSettings->IncludeRenderPasses.Value;
 		bCaptureFramesInHDR = ProtocolSettings->bCaptureFramesInHDR;
 		PostProcessingMaterial = Cast<UMaterialInterface>(ProtocolSettings->PostProcessingMaterial.TryLoad());
+
+		FString OverrideRenderPasses;
+		if( FParse::Value( FCommandLine::Get(), TEXT( "-CustomRenderPasses=" ), OverrideRenderPasses ) )
+		{
+			OverrideRenderPasses.ParseIntoArray(RenderPasses, TEXT(","), true);
+		}
+
+		bool bOverrideCaptureFramesInHDR;
+		if( FParse::Bool( FCommandLine::Get(), TEXT( "-CaptureFramesInHDR=" ), bOverrideCaptureFramesInHDR ) )
+		{
+			bCaptureFramesInHDR = bOverrideCaptureFramesInHDR;
+		}
 	}
 
 	ViewExtension = MakeShareable(new FSceneViewExtension(RenderPasses, bCaptureFramesInHDR, PostProcessingMaterial));
diff --git a/Engine/Source/Runtime/MovieSceneCapture/Public/MovieSceneCapture.h b/Engine/Source/Runtime/MovieSceneCapture/Public/MovieSceneCapture.h
--- a/Engine/Source/Runtime/MovieSceneCapture/Public/MovieSceneCapture.h
+++ b/Engine/Source/Runtime/MovieSceneCapture/Public/MovieSceneCapture.h
@@ -63,6 +63,10 @@ public:
 	UPROPERTY(config, EditAnywhere, Category=General, AdvancedDisplay)
 	bool bUseSeparateProcess;
 
+	/** When enabled, the editor will shutdown when the capture starts */
+	UPROPERTY(EditAnywhere, config, Category=General, AdvancedDisplay, meta=(EditCondition=bUseSeparateProcess))
+	bool bCloseEditorWhenCaptureStarts;
+
 	/** Additional command line arguments to pass to the external process when capturing */
 	UPROPERTY(EditAnywhere, config, Category=General, AdvancedDisplay, meta=(EditCondition=bUseSeparateProcess))
 	FString AdditionalCommandLineArguments;
diff --git a/Engine/Source/Runtime/MovieSceneTracks/Private/Sections/MovieScene3DTransformSection.cpp b/Engine/Source/Runtime/MovieSceneTracks/Private/Sections/MovieScene3DTransformSection.cpp
--- a/Engine/Source/Runtime/MovieSceneTracks/Private/Sections/MovieScene3DTransformSection.cpp
+++ b/Engine/Source/Runtime/MovieSceneTracks/Private/Sections/MovieScene3DTransformSection.cpp
@@ -5,34 +5,124 @@
 #include "MovieScene3DTransformTrack.h"
 
 
-UMovieScene3DTransformSection::UMovieScene3DTransformSection( const FObjectInitializer& ObjectInitializer )
-	: Super( ObjectInitializer )
+/* FMovieScene3DTransformKeyStruct interface
+ *****************************************************************************/
+
+void FMovieScene3DTransformKeyStruct::PropagateChanges(const FPropertyChangedEvent& ChangeEvent)
+{
+	for (int32 Index = 0; Index <= 2; ++Index)
+	{
+		TranslationKeys[Index]->Value = Translation[Index];
+		ScaleKeys[Index]->Value = Scale[Index];
+	}
+
+	RotationKeys[0]->Value = Rotation.Pitch;
+	RotationKeys[1]->Value = Rotation.Yaw;
+	RotationKeys[2]->Value = Rotation.Roll;
+}
+
+
+/* FMovieScene3DTransformKeyStruct interface
+ *****************************************************************************/
+
+UMovieScene3DTransformSection::UMovieScene3DTransformSection(const FObjectInitializer& ObjectInitializer)
+	: Super(ObjectInitializer)
 { }
 
 
-void UMovieScene3DTransformSection::MoveSection( float DeltaTime, TSet<FKeyHandle>& KeyHandles )
+/* UMovieScene3DTransformSection interface
+ *****************************************************************************/
+void UMovieScene3DTransformSection::EvalTranslation(float Time, FVector& OutTranslation) const
+{
+	OutTranslation.X = Translation[0].Eval(Time);
+	OutTranslation.Y = Translation[1].Eval(Time);
+	OutTranslation.Z = Translation[2].Eval(Time);
+}
+
+
+void UMovieScene3DTransformSection::EvalRotation(float Time, FRotator& OutRotation) const
+{
+	OutRotation.Roll = Rotation[0].Eval(Time);
+	OutRotation.Pitch = Rotation[1].Eval(Time);
+	OutRotation.Yaw = Rotation[2].Eval(Time);
+}
+
+
+void UMovieScene3DTransformSection::EvalScale(float Time, FVector& OutScale) const
+{
+	OutScale.X = Scale[0].Eval(Time);
+	OutScale.Y = Scale[1].Eval(Time);
+	OutScale.Z = Scale[2].Eval(Time);
+}
+
+
+/**
+ * Chooses an appropriate curve from an axis and a set of curves
+ */
+static FRichCurve* ChooseCurve(EAxis::Type Axis, FRichCurve* Curves)
+{
+	switch(Axis)
+	{
+	case EAxis::X:
+		return &Curves[0];
+
+	case EAxis::Y:
+		return &Curves[1];
+
+	case EAxis::Z:
+		return &Curves[2];
+
+	default:
+		check(false);
+		return nullptr;
+	}
+}
+
+
+FRichCurve& UMovieScene3DTransformSection::GetTranslationCurve(EAxis::Type Axis) 
+{
+	return *ChooseCurve(Axis, Translation);
+}
+
+
+FRichCurve& UMovieScene3DTransformSection::GetRotationCurve(EAxis::Type Axis)
+{
+	return *ChooseCurve(Axis, Rotation);
+}
+
+
+FRichCurve& UMovieScene3DTransformSection::GetScaleCurve(EAxis::Type Axis)
 {
-	Super::MoveSection( DeltaTime, KeyHandles );
+	return *ChooseCurve(Axis, Scale);
+}
+
+
+/* UMovieSceneSection interface
+ *****************************************************************************/
+
+void UMovieScene3DTransformSection::MoveSection(float DeltaTime, TSet<FKeyHandle>& KeyHandles)
+{
+	Super::MoveSection(DeltaTime, KeyHandles);
 
 	// Move all the curves in this section
-	for( int32 Axis = 0; Axis < 3; ++Axis )
+	for (int32 Axis = 0; Axis < 3; ++Axis)
 	{
-		Translation[Axis].ShiftCurve( DeltaTime, KeyHandles );
-		Rotation[Axis].ShiftCurve( DeltaTime, KeyHandles );
-		Scale[Axis].ShiftCurve( DeltaTime, KeyHandles );
+		Translation[Axis].ShiftCurve(DeltaTime, KeyHandles);
+		Rotation[Axis].ShiftCurve(DeltaTime, KeyHandles);
+		Scale[Axis].ShiftCurve(DeltaTime, KeyHandles);
 	}
 }
 
 
-void UMovieScene3DTransformSection::DilateSection( float DilationFactor, float Origin, TSet<FKeyHandle>& KeyHandles )
+void UMovieScene3DTransformSection::DilateSection(float DilationFactor, float Origin, TSet<FKeyHandle>& KeyHandles)
 {
 	Super::DilateSection(DilationFactor, Origin, KeyHandles);
 
-	for( int32 Axis = 0; Axis < 3; ++Axis )
+	for (int32 Axis = 0; Axis < 3; ++Axis)
 	{
-		Translation[Axis].ScaleCurve( Origin, DilationFactor, KeyHandles );
-		Rotation[Axis].ScaleCurve( Origin, DilationFactor, KeyHandles );
-		Scale[Axis].ScaleCurve( Origin, DilationFactor, KeyHandles );
+		Translation[Axis].ScaleCurve(Origin, DilationFactor, KeyHandles);
+		Rotation[Axis].ScaleCurve(Origin, DilationFactor, KeyHandles);
+		Scale[Axis].ScaleCurve(Origin, DilationFactor, KeyHandles);
 	}
 }
 
@@ -68,81 +158,42 @@ void UMovieScene3DTransformSection::GetKeyHandles(TSet<FKeyHandle>& KeyHandles)
 			}
 		}
 	}
-}	
-
-
-void UMovieScene3DTransformSection::EvalTranslation( float Time, FVector& OutTranslation ) const
-{
-	OutTranslation.X = Translation[0].Eval( Time );
-	OutTranslation.Y = Translation[1].Eval( Time );
-	OutTranslation.Z = Translation[2].Eval( Time );
 }
 
 
-void UMovieScene3DTransformSection::EvalRotation( float Time, FRotator& OutRotation ) const
+TSharedPtr<FStructOnScope> UMovieScene3DTransformSection::GetKeyStruct(const TArray<FKeyHandle>& KeyHandles)
 {
-	OutRotation.Roll = Rotation[0].Eval( Time );
-	OutRotation.Pitch = Rotation[1].Eval( Time );
-	OutRotation.Yaw = Rotation[2].Eval( Time );
-}
-
-
-void UMovieScene3DTransformSection::EvalScale( float Time, FVector& OutScale ) const
-{
-	OutScale.X = Scale[0].Eval( Time );
-	OutScale.Y = Scale[1].Eval( Time );
-	OutScale.Z = Scale[2].Eval( Time );
-}
-
-
-/**
- * Chooses an appropriate curve from an axis and a set of curves
- */
-static FRichCurve* ChooseCurve( EAxis::Type Axis, FRichCurve* Curves )
-{
-	switch( Axis )
+	TSharedRef<FStructOnScope> KeyStruct = MakeShareable(new FStructOnScope(FMovieScene3DTransformKeyStruct::StaticStruct()));
+	auto Struct = (FMovieScene3DTransformKeyStruct*)KeyStruct->GetStructMemory();
 	{
-	case EAxis::X:
-		return &Curves[0];
-
-	case EAxis::Y:
-		return &Curves[1];
-
-	case EAxis::Z:
-		return &Curves[2];
+		for (int32 Index = 0; Index <= 2; ++Index)
+		{
+			Struct->TranslationKeys[Index] = &Translation[Index].GetFirstMatchingKey(KeyHandles);
+			Struct->Translation[Index] = Struct->TranslationKeys[Index]->Value;
+			Struct->RotationKeys[Index] = &Rotation[Index].GetFirstMatchingKey(KeyHandles);
+			Struct->ScaleKeys[Index] = &Scale[Index].GetFirstMatchingKey(KeyHandles);
+			Struct->Scale[Index] = Struct->ScaleKeys[Index]->Value;
+		}
 
-	default:
-		check(false);
-		return nullptr;
+		Struct->Rotation.Pitch = Struct->RotationKeys[0]->Value;
+		Struct->Rotation.Yaw = Struct->RotationKeys[1]->Value;
+		Struct->Rotation.Roll = Struct->RotationKeys[2]->Value;
 	}
-}
-
 
-FRichCurve& UMovieScene3DTransformSection::GetTranslationCurve( EAxis::Type Axis ) 
-{
-	return *ChooseCurve( Axis, Translation );
-}
-
-
-FRichCurve& UMovieScene3DTransformSection::GetRotationCurve( EAxis::Type Axis )
-{
-	return *ChooseCurve( Axis, Rotation );
+	return KeyStruct;
 }
 
 
-FRichCurve& UMovieScene3DTransformSection::GetScaleCurve( EAxis::Type Axis )
-{
-	return *ChooseCurve( Axis, Scale );
-}
-
+/* UMovieSceneSection interface
+ *****************************************************************************/
 
 // This function uses a template to avoid duplicating this for const and non-const versions
 template<typename CurveType>
-CurveType* GetCurveForChannelAndAxis( EKey3DTransformChannel::Type Channel, EAxis::Type Axis,
-	CurveType* TranslationCurves, CurveType* RotationCurves, CurveType* ScaleCurves )
+CurveType* GetCurveForChannelAndAxis(EKey3DTransformChannel::Type Channel, EAxis::Type Axis,
+	CurveType* TranslationCurves, CurveType* RotationCurves, CurveType* ScaleCurves)
 {
 	CurveType* ChannelCurves = nullptr;
-	switch ( Channel )
+	switch (Channel)
 	{
 	case EKey3DTransformChannel::Translation:
 		ChannelCurves = TranslationCurves;
@@ -154,9 +205,10 @@ CurveType* GetCurveForChannelAndAxis( EKey3DTransformChannel::Type Channel, EAxi
 		ChannelCurves = ScaleCurves;
 		break;
 	}
-	if ( ChannelCurves != nullptr )
+
+	if (ChannelCurves != nullptr)
 	{
-		switch ( Axis )
+		switch (Axis)
 		{
 		case EAxis::X:
 			return &ChannelCurves[0];
@@ -166,36 +218,35 @@ CurveType* GetCurveForChannelAndAxis( EKey3DTransformChannel::Type Channel, EAxi
 			return &ChannelCurves[2];
 		}
 	}
-	checkf( false, TEXT( "Invalid channel and axis combination." ) );
+
+	checkf(false, TEXT("Invalid channel and axis combination."));
 	return nullptr;
 }
 
 
-bool UMovieScene3DTransformSection::NewKeyIsNewData( float Time, const FTransformKey& TransformKey ) const
+bool UMovieScene3DTransformSection::NewKeyIsNewData(float Time, const FTransformKey& TransformKey) const
 {
-	const FRichCurve* KeyCurve = GetCurveForChannelAndAxis( TransformKey.Channel, TransformKey.Axis, Translation, Rotation, Scale );
-	return FMath::IsNearlyEqual( KeyCurve->Eval( Time ), TransformKey.Value ) == false;
+	const FRichCurve* KeyCurve = GetCurveForChannelAndAxis(TransformKey.Channel, TransformKey.Axis, Translation, Rotation, Scale);
+	return FMath::IsNearlyEqual(KeyCurve->Eval(Time), TransformKey.Value) == false;
 }
 
 
-bool UMovieScene3DTransformSection::HasKeys( const FTransformKey& TransformKey ) const
+bool UMovieScene3DTransformSection::HasKeys(const FTransformKey& TransformKey) const
 {
-	const FRichCurve* KeyCurve = GetCurveForChannelAndAxis( TransformKey.Channel, TransformKey.Axis, Translation, Rotation, Scale );
+	const FRichCurve* KeyCurve = GetCurveForChannelAndAxis(TransformKey.Channel, TransformKey.Axis, Translation, Rotation, Scale);
 	return KeyCurve->GetNumKeys() != 0;
 }
 
 
-void UMovieScene3DTransformSection::AddKey( float Time, const FTransformKey& TransformKey, EMovieSceneKeyInterpolation KeyInterpolation )
+void UMovieScene3DTransformSection::AddKey(float Time, const FTransformKey& TransformKey, EMovieSceneKeyInterpolation KeyInterpolation)
 {
-	FRichCurve* KeyCurve = GetCurveForChannelAndAxis( TransformKey.Channel, TransformKey.Axis, Translation, Rotation, Scale );
-	AddKeyToCurve( *KeyCurve, Time, TransformKey.Value, KeyInterpolation );
+	FRichCurve* KeyCurve = GetCurveForChannelAndAxis(TransformKey.Channel, TransformKey.Axis, Translation, Rotation, Scale);
+	AddKeyToCurve(*KeyCurve, Time, TransformKey.Value, KeyInterpolation);
 }
 
 
-void UMovieScene3DTransformSection::SetDefault( const FTransformKey& TransformKey )
+void UMovieScene3DTransformSection::SetDefault(const FTransformKey& TransformKey)
 {
-	FRichCurve* KeyCurve = GetCurveForChannelAndAxis( TransformKey.Channel, TransformKey.Axis, Translation, Rotation, Scale );
-	SetCurveDefault( *KeyCurve, TransformKey.Value );
+	FRichCurve* KeyCurve = GetCurveForChannelAndAxis(TransformKey.Channel, TransformKey.Axis, Translation, Rotation, Scale);
+	SetCurveDefault(*KeyCurve, TransformKey.Value);
 }
-
-
diff --git a/Engine/Source/Runtime/MovieSceneTracks/Private/Sections/MovieSceneCameraCutSection.cpp b/Engine/Source/Runtime/MovieSceneTracks/Private/Sections/MovieSceneCameraCutSection.cpp
new file mode 100644
--- /dev/null
+++ b/Engine/Source/Runtime/MovieSceneTracks/Private/Sections/MovieSceneCameraCutSection.cpp
@@ -0,0 +1,8 @@
+// Copyright 1998-2016 Epic Games, Inc. All Rights Reserved.
+
+#include "MovieSceneTracksPrivatePCH.h"
+#include "MovieSceneCameraCutSection.h"
+
+
+/* UMovieSceneCameraCutSection interface
+ *****************************************************************************/
diff --git a/Engine/Source/Runtime/MovieSceneTracks/Private/Sections/MovieSceneCinematicShotSection.cpp b/Engine/Source/Runtime/MovieSceneTracks/Private/Sections/MovieSceneCinematicShotSection.cpp
new file mode 100644
--- /dev/null
+++ b/Engine/Source/Runtime/MovieSceneTracks/Private/Sections/MovieSceneCinematicShotSection.cpp
@@ -0,0 +1,13 @@
+// Copyright 1998-2016 Epic Games, Inc. All Rights Reserved.
+
+#include "MovieSceneTracksPrivatePCH.h"
+#include "MovieSceneCinematicShotSection.h"
+#include "MovieSceneSubSection.h"
+
+
+/* UMovieSceneCinematicshotSection structors
+ *****************************************************************************/
+
+UMovieSceneCinematicShotSection::UMovieSceneCinematicShotSection() : UMovieSceneSubSection()
+{ }
+
