Copying //UE4/Dev-Sequencer to Dev-Main (//UE4/Dev-Main)

diff --git a/Engine/Source/Runtime/MovieSceneTracks/Private/Sections/MovieSceneColorSection.cpp b/Engine/Source/Runtime/MovieSceneTracks/Private/Sections/MovieSceneColorSection.cpp
--- a/Engine/Source/Runtime/MovieSceneTracks/Private/Sections/MovieSceneColorSection.cpp
+++ b/Engine/Source/Runtime/MovieSceneTracks/Private/Sections/MovieSceneColorSection.cpp
@@ -5,14 +5,45 @@
 #include "MovieSceneColorTrack.h"
 
 
-UMovieSceneColorSection::UMovieSceneColorSection( const FObjectInitializer& ObjectInitializer )
-	: Super( ObjectInitializer )
+/* FMovieSceneColorKeyStruct interface
+ *****************************************************************************/
+
+void FMovieSceneColorKeyStruct::PropagateChanges(const FPropertyChangedEvent& ChangeEvent)
+{
+	for (int32 Index = 0; Index <= 3; ++Index)
+	{
+		Keys[Index]->Value = Color.Component(Index);
+	}
+}
+
+
+/* UMovieSceneColorSection structors
+ *****************************************************************************/
+
+UMovieSceneColorSection::UMovieSceneColorSection(const FObjectInitializer& ObjectInitializer)
+	: Super(ObjectInitializer)
 { }
 
 
-void UMovieSceneColorSection::MoveSection( float DeltaTime, TSet<FKeyHandle>& KeyHandles )
+/* UMovieSceneColorSection interface
+ *****************************************************************************/
+
+FLinearColor UMovieSceneColorSection::Eval(float Position, const FLinearColor& DefaultColor) const
+{
+	return FLinearColor(
+		RedCurve.Eval(Position, DefaultColor.R),
+		GreenCurve.Eval(Position, DefaultColor.G),
+		BlueCurve.Eval(Position, DefaultColor.B),
+		AlphaCurve.Eval(Position, DefaultColor.A));
+}
+
+
+/* UMovieSceneSection interface
+ *****************************************************************************/
+
+void UMovieSceneColorSection::MoveSection(float DeltaTime, TSet<FKeyHandle>& KeyHandles)
 {
-	Super::MoveSection( DeltaTime, KeyHandles );
+	Super::MoveSection(DeltaTime, KeyHandles);
 
 	// Move all the curves in this section
 	RedCurve.ShiftCurve(DeltaTime, KeyHandles);
@@ -22,7 +53,7 @@ void UMovieSceneColorSection::MoveSection( float DeltaTime, TSet<FKeyHandle>& Ke
 }
 
 
-void UMovieSceneColorSection::DilateSection( float DilationFactor, float Origin, TSet<FKeyHandle>& KeyHandles )
+void UMovieSceneColorSection::DilateSection(float DilationFactor, float Origin, TSet<FKeyHandle>& KeyHandles)
 {
 	Super::DilateSection(DilationFactor, Origin, KeyHandles);
 
@@ -71,18 +102,33 @@ void UMovieSceneColorSection::GetKeyHandles(TSet<FKeyHandle>& KeyHandles) const
 }
 
 
-FLinearColor UMovieSceneColorSection::Eval( float Position, const FLinearColor& DefaultColor ) const
+TSharedPtr<FStructOnScope> UMovieSceneColorSection::GetKeyStruct(const TArray<FKeyHandle>& KeyHandles)
 {
-	return FLinearColor(RedCurve.Eval(Position, DefaultColor.R),
-						GreenCurve.Eval(Position, DefaultColor.G),
-						BlueCurve.Eval(Position, DefaultColor.B),
-						AlphaCurve.Eval(Position, DefaultColor.A));
+	TSharedRef<FStructOnScope> KeyStruct = MakeShareable(new FStructOnScope(FMovieSceneColorKeyStruct::StaticStruct()));
+	auto Struct = (FMovieSceneColorKeyStruct*)KeyStruct->GetStructMemory();
+	{
+		Struct->Keys[0] = &RedCurve.GetFirstMatchingKey(KeyHandles);
+		Struct->Keys[1] = &GreenCurve.GetFirstMatchingKey(KeyHandles);
+		Struct->Keys[2] = &BlueCurve.GetFirstMatchingKey(KeyHandles);
+		Struct->Keys[3] = &AlphaCurve.GetFirstMatchingKey(KeyHandles);
+
+		for (int32 Index = 0; Index <= 3; ++Index)
+		{
+			Struct->Color.Component(Index) = Struct->Keys[Index]->Value;
+		}
+	}
+
+	return KeyStruct;
 }
 
+
+/* IKeyframeSection interface
+ *****************************************************************************/
+
 template<typename CurveType>
-CurveType* GetCurveForChannel( EKeyColorChannel Channel, CurveType* RedCurve, CurveType* GreenCurve, CurveType* BlueCurve, CurveType* AlphaCurve )
+CurveType* GetCurveForChannel(EKeyColorChannel Channel, CurveType* RedCurve, CurveType* GreenCurve, CurveType* BlueCurve, CurveType* AlphaCurve)
 {
-	switch ( Channel )
+	switch (Channel)
 	{
 	case EKeyColorChannel::Red:
 		return RedCurve;
@@ -93,35 +139,35 @@ CurveType* GetCurveForChannel( EKeyColorChannel Channel, CurveType* RedCurve, Cu
 	case EKeyColorChannel::Alpha:
 		return AlphaCurve;
 	default:
-		checkf( false, TEXT( "Invalid key color channel" ) );
+		checkf(false, TEXT("Invalid key color channel"));
 		return nullptr;
 	}
 }
 
 
-void UMovieSceneColorSection::AddKey( float Time, const FColorKey& Key, EMovieSceneKeyInterpolation KeyInterpolation )
+void UMovieSceneColorSection::AddKey(float Time, const FColorKey& Key, EMovieSceneKeyInterpolation KeyInterpolation)
 {
-	FRichCurve* ChannelCurve = GetCurveForChannel( Key.Channel, &RedCurve, &GreenCurve, &BlueCurve, &AlphaCurve );
-	AddKeyToCurve( *ChannelCurve, Time, Key.ChannelValue, KeyInterpolation );
+	FRichCurve* ChannelCurve = GetCurveForChannel(Key.Channel, &RedCurve, &GreenCurve, &BlueCurve, &AlphaCurve);
+	AddKeyToCurve(*ChannelCurve, Time, Key.ChannelValue, KeyInterpolation);
 }
 
 
-bool UMovieSceneColorSection::NewKeyIsNewData( float Time, const FColorKey& Key ) const
+bool UMovieSceneColorSection::NewKeyIsNewData(float Time, const FColorKey& Key) const
 {
-	const FRichCurve* ChannelCurve = GetCurveForChannel( Key.Channel, &RedCurve, &GreenCurve, &BlueCurve, &AlphaCurve );
-	return FMath::IsNearlyEqual( ChannelCurve->Eval( Time ), Key.ChannelValue ) == false;
+	const FRichCurve* ChannelCurve = GetCurveForChannel(Key.Channel, &RedCurve, &GreenCurve, &BlueCurve, &AlphaCurve);
+	return FMath::IsNearlyEqual(ChannelCurve->Eval(Time), Key.ChannelValue) == false;
 }
 
 
-bool UMovieSceneColorSection::HasKeys( const FColorKey& Key ) const
+bool UMovieSceneColorSection::HasKeys(const FColorKey& Key) const
 {
-	const FRichCurve* ChannelCurve = GetCurveForChannel( Key.Channel, &RedCurve, &GreenCurve, &BlueCurve, &AlphaCurve );
+	const FRichCurve* ChannelCurve = GetCurveForChannel(Key.Channel, &RedCurve, &GreenCurve, &BlueCurve, &AlphaCurve);
 	return ChannelCurve->GetNumKeys() != 0;
 }
 
 
-void UMovieSceneColorSection::SetDefault( const FColorKey& Key )
+void UMovieSceneColorSection::SetDefault(const FColorKey& Key)
 {
-	FRichCurve* ChannelCurve = GetCurveForChannel( Key.Channel, &RedCurve, &GreenCurve, &BlueCurve, &AlphaCurve );
-	return SetCurveDefault( *ChannelCurve, Key.ChannelValue );
+	FRichCurve* ChannelCurve = GetCurveForChannel(Key.Channel, &RedCurve, &GreenCurve, &BlueCurve, &AlphaCurve);
+	return SetCurveDefault(*ChannelCurve, Key.ChannelValue);
 }
diff --git a/Engine/Source/Runtime/MovieSceneTracks/Private/Sections/MovieSceneSubSection.cpp b/Engine/Source/Runtime/MovieSceneTracks/Private/Sections/MovieSceneSubSection.cpp
--- a/Engine/Source/Runtime/MovieSceneTracks/Private/Sections/MovieSceneSubSection.cpp
+++ b/Engine/Source/Runtime/MovieSceneTracks/Private/Sections/MovieSceneSubSection.cpp
@@ -12,4 +12,5 @@
 UMovieSceneSubSection::UMovieSceneSubSection()
 	: StartOffset(0.0f)
 	, TimeScale(1.0f)
+	, PrerollTime(0.0f)
 { }
diff --git a/Engine/Source/Runtime/MovieSceneTracks/Private/Sections/MovieSceneVectorSection.cpp b/Engine/Source/Runtime/MovieSceneTracks/Private/Sections/MovieSceneVectorSection.cpp
--- a/Engine/Source/Runtime/MovieSceneTracks/Private/Sections/MovieSceneVectorSection.cpp
+++ b/Engine/Source/Runtime/MovieSceneTracks/Private/Sections/MovieSceneVectorSection.cpp
@@ -5,77 +5,48 @@
 #include "MovieSceneVectorTrack.h"
 
 
-UMovieSceneVectorSection::UMovieSceneVectorSection( const FObjectInitializer& ObjectInitializer )
-	: Super( ObjectInitializer )
-{
-	ChannelsUsed = 0;
-}
-
-
-FVector4 UMovieSceneVectorSection::Eval( float Position, const FVector4& DefaultVector ) const
-{
-	return FVector4(
-		Curves[0].Eval( Position, DefaultVector.X ),
-		Curves[1].Eval( Position, DefaultVector.Y ),
-		Curves[2].Eval( Position, DefaultVector.Z ),
-		Curves[3].Eval( Position, DefaultVector.W ) );
-}
+/* FMovieSceneVectorKeyStruct interface
+ *****************************************************************************/
 
-
-template<typename CurveType>
-CurveType* GetCurveForChannel( EKeyVectorChannel Channel, CurveType* Curves, int32 ChannelsUsed )
+void FMovieSceneVectorKeyStruct::PropagateChanges(const FPropertyChangedEvent& ChangeEvent)
 {
-	switch ( Channel )
+	for (int32 Index = 0; Index <= 3; ++Index)
 	{
-		case EKeyVectorChannel::X:
-			return &Curves[0];
-		case EKeyVectorChannel::Y:
-			return &Curves[1];
-		case EKeyVectorChannel::Z:
-			checkf( ChannelsUsed >= 3, TEXT( "Can not get Z channel, it is not in use on this section." ) );
-			return &Curves[2];
-		case EKeyVectorChannel::W:
-			checkf( ChannelsUsed >= 4, TEXT("Can not get W channel, it is not in use on this section." ) );
-			return &Curves[3];
+		Keys[Index]->Value = Vector[Index];
 	}
-	checkf( false, TEXT( "Invalid channel requested" ) );
-	return nullptr;
 }
 
 
-void UMovieSceneVectorSection::AddKey( float Time, const FVectorKey& Key, EMovieSceneKeyInterpolation KeyInterpolation )
-{
-	FRichCurve* ChannelCurve = GetCurveForChannel( Key.Channel, Curves, ChannelsUsed );
-	AddKeyToCurve( *ChannelCurve, Time, Key.Value, KeyInterpolation );
-}
-
+/* UMovieSceneVectorSection structors
+ *****************************************************************************/
 
-bool UMovieSceneVectorSection::NewKeyIsNewData( float Time, const FVectorKey& Key ) const
+UMovieSceneVectorSection::UMovieSceneVectorSection(const FObjectInitializer& ObjectInitializer)
+	: Super(ObjectInitializer)
 {
-	const FRichCurve* ChannelCurve = GetCurveForChannel( Key.Channel, Curves, ChannelsUsed );
-	return FMath::IsNearlyEqual( ChannelCurve->Eval( Time ), Key.Value ) == false;
+	ChannelsUsed = 0;
 }
 
 
-bool UMovieSceneVectorSection::HasKeys( const FVectorKey& Key ) const
-{
-	const FRichCurve* ChannelCurve = GetCurveForChannel( Key.Channel, Curves, ChannelsUsed );
-	return ChannelCurve->GetNumKeys() > 0;
-}
+/* UMovieSceneVectorSection interface
+ *****************************************************************************/
 
-
-void UMovieSceneVectorSection::SetDefault( const FVectorKey& Key )
+FVector4 UMovieSceneVectorSection::Eval(float Position, const FVector4& DefaultVector) const
 {
-	FRichCurve* ChannelCurve = GetCurveForChannel( Key.Channel, Curves, ChannelsUsed );
-	ChannelCurve->SetDefaultValue( Key.Value );
+	return FVector4(
+		Curves[0].Eval(Position, DefaultVector.X),
+		Curves[1].Eval(Position, DefaultVector.Y),
+		Curves[2].Eval(Position, DefaultVector.Z),
+		Curves[3].Eval(Position, DefaultVector.W));
 }
 
 
-void UMovieSceneVectorSection::MoveSection( float DeltaTime, TSet<FKeyHandle>& KeyHandles )
+/* UMovieSceneSection interface
+ *****************************************************************************/
+
+void UMovieSceneVectorSection::MoveSection(float DeltaTime, TSet<FKeyHandle>& KeyHandles)
 {
 	check(ChannelsUsed >= 2 && ChannelsUsed <= 4);
-
-	Super::MoveSection( DeltaTime, KeyHandles );
+	Super::MoveSection(DeltaTime, KeyHandles);
 
 	for (int32 i = 0; i < ChannelsUsed; ++i)
 	{
@@ -84,7 +55,7 @@ void UMovieSceneVectorSection::MoveSection( float DeltaTime, TSet<FKeyHandle>& K
 }
 
 
-void UMovieSceneVectorSection::DilateSection( float DilationFactor, float Origin, TSet<FKeyHandle>& KeyHandles )
+void UMovieSceneVectorSection::DilateSection(float DilationFactor, float Origin, TSet<FKeyHandle>& KeyHandles)
 {
 	check(ChannelsUsed >= 2 && ChannelsUsed <= 4);
 	Super::DilateSection(DilationFactor, Origin, KeyHandles);
@@ -110,3 +81,71 @@ void UMovieSceneVectorSection::GetKeyHandles(TSet<FKeyHandle>& KeyHandles) const
 		}
 	}
 }
+
+
+TSharedPtr<FStructOnScope> UMovieSceneVectorSection::GetKeyStruct(const TArray<FKeyHandle>& KeyHandles)
+{
+	TSharedRef<FStructOnScope> KeyStruct = MakeShareable(new FStructOnScope(FMovieSceneVectorKeyStruct::StaticStruct()));
+	auto Struct = (FMovieSceneVectorKeyStruct*)KeyStruct->GetStructMemory();
+	{
+		for (int32 Index = 0; Index <= 3; ++Index)
+		{
+			Struct->Keys[Index] = &Curves[Index].GetFirstMatchingKey(KeyHandles);
+			Struct->Vector[Index] = Struct->Keys[Index]->Value;
+		}
+	}
+
+	return KeyStruct;
+}
+
+
+/* IKeyframeSection interface
+ *****************************************************************************/
+
+template<typename CurveType>
+CurveType* GetCurveForChannel(EKeyVectorChannel Channel, CurveType* Curves, int32 ChannelsUsed)
+{
+	switch (Channel)
+	{
+		case EKeyVectorChannel::X:
+			return &Curves[0];
+		case EKeyVectorChannel::Y:
+			return &Curves[1];
+		case EKeyVectorChannel::Z:
+			checkf(ChannelsUsed >= 3, TEXT("Can not get Z channel, it is not in use on this section."));
+			return &Curves[2];
+		case EKeyVectorChannel::W:
+			checkf(ChannelsUsed >= 4, TEXT("Can not get W channel, it is not in use on this section."));
+			return &Curves[3];
+	}
+	checkf(false, TEXT("Invalid channel requested"));
+	return nullptr;
+}
+
+
+void UMovieSceneVectorSection::AddKey(float Time, const FVectorKey& Key, EMovieSceneKeyInterpolation KeyInterpolation)
+{
+	FRichCurve* ChannelCurve = GetCurveForChannel(Key.Channel, Curves, ChannelsUsed);
+	AddKeyToCurve(*ChannelCurve, Time, Key.Value, KeyInterpolation);
+}
+
+
+bool UMovieSceneVectorSection::NewKeyIsNewData(float Time, const FVectorKey& Key) const
+{
+	const FRichCurve* ChannelCurve = GetCurveForChannel(Key.Channel, Curves, ChannelsUsed);
+	return FMath::IsNearlyEqual(ChannelCurve->Eval(Time), Key.Value) == false;
+}
+
+
+bool UMovieSceneVectorSection::HasKeys(const FVectorKey& Key) const
+{
+	const FRichCurve* ChannelCurve = GetCurveForChannel(Key.Channel, Curves, ChannelsUsed);
+	return ChannelCurve->GetNumKeys() > 0;
+}
+
+
+void UMovieSceneVectorSection::SetDefault(const FVectorKey& Key)
+{
+	FRichCurve* ChannelCurve = GetCurveForChannel(Key.Channel, Curves, ChannelsUsed);
+	ChannelCurve->SetDefaultValue(Key.Value);
+}
diff --git a/Engine/Source/Runtime/MovieSceneTracks/Private/TrackInstances/MovieScene3DConstraintTrackInstance.cpp b/Engine/Source/Runtime/MovieSceneTracks/Private/TrackInstances/MovieScene3DConstraintTrackInstance.cpp
--- a/Engine/Source/Runtime/MovieSceneTracks/Private/TrackInstances/MovieScene3DConstraintTrackInstance.cpp
+++ b/Engine/Source/Runtime/MovieSceneTracks/Private/TrackInstances/MovieScene3DConstraintTrackInstance.cpp
@@ -53,7 +53,7 @@ void FMovieScene3DConstraintTrackInstance::RestoreState(const TArray<UObject*>&
 }
 
 
-void FMovieScene3DConstraintTrackInstance::Update( float Position, float LastPosition, const TArray<UObject*>& RuntimeObjects, class IMovieScenePlayer& Player, FMovieSceneSequenceInstance& SequenceInstance, EMovieSceneUpdatePass UpdatePass ) 
+void FMovieScene3DConstraintTrackInstance::Update(EMovieSceneUpdateData& UpdateData, const TArray<UObject*>& RuntimeObjects, class IMovieScenePlayer& Player, FMovieSceneSequenceInstance& SequenceInstance) 
 {
 	UMovieScene3DConstraintSection* FirstConstraintSection = nullptr;
 
@@ -63,7 +63,7 @@ void FMovieScene3DConstraintTrackInstance::Update( float Position, float LastPos
 	{
 		UMovieScene3DConstraintSection* ConstraintSection = CastChecked<UMovieScene3DConstraintSection>(ConstraintSections[ConstraintIndex]);
 
-		if (ConstraintSection->IsTimeWithinSection(Position) &&
+		if (ConstraintSection->IsTimeWithinSection(UpdateData.Position) &&
 			(FirstConstraintSection == nullptr || FirstConstraintSection->GetRowIndex() > ConstraintSection->GetRowIndex()))
 		{
 			TArray<UObject*> ConstraintObjects;
@@ -78,7 +78,7 @@ void FMovieScene3DConstraintTrackInstance::Update( float Position, float LastPos
 					AActor* Actor = Cast<AActor>(ConstraintObjects[ConstraintObjectIndex]);
 					if (Actor)
 					{
-						UpdateConstraint(Position, RuntimeObjects, Actor, ConstraintSection);
+						UpdateConstraint(UpdateData.Position, RuntimeObjects, Actor, ConstraintSection);
 					}	
 				}
 			}
diff --git a/Engine/Source/Runtime/MovieSceneTracks/Private/TrackInstances/MovieScene3DTransformTrackInstance.cpp b/Engine/Source/Runtime/MovieSceneTracks/Private/TrackInstances/MovieScene3DTransformTrackInstance.cpp
--- a/Engine/Source/Runtime/MovieSceneTracks/Private/TrackInstances/MovieScene3DTransformTrackInstance.cpp
+++ b/Engine/Source/Runtime/MovieSceneTracks/Private/TrackInstances/MovieScene3DTransformTrackInstance.cpp
@@ -79,9 +79,9 @@ void FMovieScene3DTransformTrackInstance::UpdateRuntimeMobility(const TArray<UOb
 	}
 }
 
-void FMovieScene3DTransformTrackInstance::Update( float Position, float LastPosition, const TArray<UObject*>& RuntimeObjects, class IMovieScenePlayer& Player, FMovieSceneSequenceInstance& SequenceInstance, EMovieSceneUpdatePass UpdatePass ) 
+void FMovieScene3DTransformTrackInstance::Update(EMovieSceneUpdateData& UpdateData, const TArray<UObject*>& RuntimeObjects, class IMovieScenePlayer& Player, FMovieSceneSequenceInstance& SequenceInstance) 
 {
-	if (UpdatePass == MSUP_PreUpdate)
+	if (UpdateData.UpdatePass == MSUP_PreUpdate)
 	{
 		UpdateRuntimeMobility(RuntimeObjects);
 	}
@@ -90,7 +90,7 @@ void FMovieScene3DTransformTrackInstance::Update( float Position, float LastPosi
 	FRotator Rotation;
 	FVector Scale;
 
-	if( TransformTrack->Eval( Position, LastPosition, Translation, Rotation, Scale ) )
+	if( TransformTrack->Eval( UpdateData.Position, UpdateData.LastPosition, Translation, Rotation, Scale ) )
 	{
 		for( int32 ObjIndex = 0; ObjIndex < RuntimeObjects.Num(); ++ObjIndex )
 		{
@@ -98,11 +98,11 @@ void FMovieScene3DTransformTrackInstance::Update( float Position, float LastPosi
 
 			if (SceneComponent != nullptr)
 			{
-				if (UpdatePass == MSUP_PreUpdate)
+				if (UpdateData.UpdatePass == MSUP_PreUpdate)
 				{
 					SceneComponent->ResetRelativeTransform();
 				}
-				else if (UpdatePass == MSUP_Update)
+				else if (UpdateData.UpdatePass == MSUP_Update)
 				{
 					SceneComponent->AddRelativeLocation(Translation);
 					SceneComponent->AddRelativeRotation(Rotation);
diff --git a/Engine/Source/Runtime/MovieSceneTracks/Private/TrackInstances/MovieSceneAudioTrackInstance.cpp b/Engine/Source/Runtime/MovieSceneTracks/Private/TrackInstances/MovieSceneAudioTrackInstance.cpp
--- a/Engine/Source/Runtime/MovieSceneTracks/Private/TrackInstances/MovieSceneAudioTrackInstance.cpp
+++ b/Engine/Source/Runtime/MovieSceneTracks/Private/TrackInstances/MovieSceneAudioTrackInstance.cpp
@@ -16,7 +16,7 @@ FMovieSceneAudioTrackInstance::FMovieSceneAudioTrackInstance( UMovieSceneAudioTr
 }
 
 
-void FMovieSceneAudioTrackInstance::Update( float Position, float LastPosition, const TArray<UObject*>& RuntimeObjects, class IMovieScenePlayer& Player, FMovieSceneSequenceInstance& SequenceInstance, EMovieSceneUpdatePass UpdatePass ) 
+void FMovieSceneAudioTrackInstance::Update(EMovieSceneUpdateData& UpdateData, const TArray<UObject*>& RuntimeObjects, class IMovieScenePlayer& Player, FMovieSceneSequenceInstance& SequenceInstance) 
 {
 	const TArray<UMovieSceneSection*>& AudioSections = AudioTrack->GetAudioSections();
 
@@ -24,7 +24,7 @@ void FMovieSceneAudioTrackInstance::Update( float Position, float LastPosition,
 
 	if (Player.GetPlaybackStatus() == EMovieScenePlayerStatus::Playing)
 	{
-		if (Position > LastPosition)
+		if (UpdateData.Position > UpdateData.LastPosition)
 		{
 			TMap<int32, TArray<UMovieSceneAudioSection*> > AudioSectionsBySectionIndex;
 			for (int32 i = 0; i < AudioSections.Num(); ++i)
@@ -52,11 +52,11 @@ void FMovieSceneAudioTrackInstance::Update( float Position, float LastPosition,
 						for (int32 i = 0; i < MovieSceneAudioSections.Num(); ++i)
 						{
 							UMovieSceneAudioSection* AudioSection = MovieSceneAudioSections[i];
-							if (AudioSection->IsTimeWithinAudioRange(Position))
+							if (AudioSection->IsTimeWithinAudioRange(UpdateData.Position))
 							{
-								if (!AudioSection->IsTimeWithinAudioRange(LastPosition) || !Component->IsPlaying())
+								if (!AudioSection->IsTimeWithinAudioRange(UpdateData.LastPosition) || !Component->IsPlaying())
 								{
-									PlaySound(AudioSection, Component, Position);
+									PlaySound(AudioSection, Component, UpdateData.Position);
 								}
 								bComponentIsPlaying = true;
 							}
@@ -78,7 +78,7 @@ void FMovieSceneAudioTrackInstance::Update( float Position, float LastPosition,
 	else if (Player.GetPlaybackStatus() == EMovieScenePlayerStatus::Scrubbing)
 	{
 		// handle scrubbing
-		if (!FMath::IsNearlyEqual(Position, LastPosition))
+		if (!FMath::IsNearlyEqual(UpdateData.Position, UpdateData.LastPosition))
 		{
 			for (int32 i = 0; i < AudioSections.Num(); ++i)
 			{
@@ -92,9 +92,9 @@ void FMovieSceneAudioTrackInstance::Update( float Position, float LastPosition,
 						TWeakObjectPtr<UAudioComponent> Component = GetAudioComponent(Actors[ActorIndex], RowIndex);
 						if (Component.IsValid())
 						{
-							if (AudioSection->IsTimeWithinAudioRange(Position) && !Component->IsPlaying())
+							if (AudioSection->IsTimeWithinAudioRange(UpdateData.Position) && !Component->IsPlaying())
 							{
-								PlaySound(AudioSection, Component, Position);
+								PlaySound(AudioSection, Component, UpdateData.Position);
 								// Fade out the sound at the same volume in order to simply
 								// set a short duration on the sound, far from ideal soln
 								Component->FadeOut(AudioTrackConstants::ScrubDuration, 1.f);
diff --git a/Engine/Source/Runtime/MovieSceneTracks/Private/TrackInstances/MovieSceneBoolTrackInstance.cpp b/Engine/Source/Runtime/MovieSceneTracks/Private/TrackInstances/MovieSceneBoolTrackInstance.cpp
--- a/Engine/Source/Runtime/MovieSceneTracks/Private/TrackInstances/MovieSceneBoolTrackInstance.cpp
+++ b/Engine/Source/Runtime/MovieSceneTracks/Private/TrackInstances/MovieSceneBoolTrackInstance.cpp
@@ -47,11 +47,11 @@ void FMovieSceneBoolTrackInstance::RestoreState(const TArray<UObject*>& RuntimeO
 }
 
 
-void FMovieSceneBoolTrackInstance::Update( float Position, float LastPosition, const TArray<UObject*>& RuntimeObjects, class IMovieScenePlayer& Player, FMovieSceneSequenceInstance& SequenceInstance, EMovieSceneUpdatePass UpdatePass ) 
+void FMovieSceneBoolTrackInstance::Update(EMovieSceneUpdateData& UpdateData, const TArray<UObject*>& RuntimeObjects, class IMovieScenePlayer& Player, FMovieSceneSequenceInstance& SequenceInstance) 
 {
 	bool BoolValue = false;
 
-	if( BoolTrack->Eval( Position, LastPosition, BoolValue ) )
+	if( BoolTrack->Eval( UpdateData.Position, UpdateData.LastPosition, BoolValue ) )
 	{
 		for( UObject* Object : RuntimeObjects )
 		{
diff --git a/Engine/Source/Runtime/MovieSceneTracks/Private/TrackInstances/MovieSceneByteTrackInstance.cpp b/Engine/Source/Runtime/MovieSceneTracks/Private/TrackInstances/MovieSceneByteTrackInstance.cpp
--- a/Engine/Source/Runtime/MovieSceneTracks/Private/TrackInstances/MovieSceneByteTrackInstance.cpp
+++ b/Engine/Source/Runtime/MovieSceneTracks/Private/TrackInstances/MovieSceneByteTrackInstance.cpp
@@ -46,10 +46,10 @@ void FMovieSceneByteTrackInstance::RestoreState(const TArray<UObject*>& RuntimeO
 }
 
 
-void FMovieSceneByteTrackInstance::Update( float Position, float LastPosition, const TArray<UObject*>& RuntimeObjects, class IMovieScenePlayer& Player, FMovieSceneSequenceInstance& SequenceInstance, EMovieSceneUpdatePass UpdatePass ) 
+void FMovieSceneByteTrackInstance::Update(EMovieSceneUpdateData& UpdateData, const TArray<UObject*>& RuntimeObjects, class IMovieScenePlayer& Player, FMovieSceneSequenceInstance& SequenceInstance) 
 {
 	uint8 ByteValue = 0;
-	if( ByteTrack->Eval( Position, LastPosition, ByteValue ) )
+	if( ByteTrack->Eval( UpdateData.Position, UpdateData.LastPosition, ByteValue ) )
 	{
 		for( UObject* Object : RuntimeObjects )
 		{
diff --git a/Engine/Source/Runtime/MovieSceneTracks/Private/TrackInstances/MovieSceneCameraCutTrackInstance.cpp b/Engine/Source/Runtime/MovieSceneTracks/Private/TrackInstances/MovieSceneCameraCutTrackInstance.cpp
new file mode 100644
--- /dev/null
+++ b/Engine/Source/Runtime/MovieSceneTracks/Private/TrackInstances/MovieSceneCameraCutTrackInstance.cpp
@@ -0,0 +1,101 @@
+// Copyright 1998-2016 Epic Games, Inc. All Rights Reserved.
+
+#include "MovieSceneTracksPrivatePCH.h"
+#include "MovieSceneSequence.h"
+#include "MovieSceneSequenceInstance.h"
+#include "MovieSceneCameraCutTrack.h"
+#include "MovieSceneCameraCutTrackInstance.h"
+#include "MovieSceneCameraCutSection.h"
+#include "IMovieScenePlayer.h"
+
+
+/* FMovieSceneCameraCutTrackInstance structors
+ *****************************************************************************/
+
+FMovieSceneCameraCutTrackInstance::FMovieSceneCameraCutTrackInstance(UMovieSceneCameraCutTrack& InCameraCutTrack)
+	: CameraCutTrack(&InCameraCutTrack)
+{ }
+
+
+/* IMovieSceneTrackInstance interface
+ *****************************************************************************/
+
+void FMovieSceneCameraCutTrackInstance::ClearInstance(IMovieScenePlayer& Player, FMovieSceneSequenceInstance& SequenceInstance)
+{
+	Player.UpdateCameraCut(nullptr, nullptr);
+}
+
+
+void FMovieSceneCameraCutTrackInstance::RefreshInstance(const TArray<UObject*>& RuntimeObjects, IMovieScenePlayer& Player, FMovieSceneSequenceInstance& SequenceInstance)
+{
+	const TArray<UMovieSceneSection*>& CameraCutSections = CameraCutTrack->GetAllSections();
+
+	CachedCameraObjects.Empty();
+	CachedCameraObjects.SetNumZeroed(CameraCutSections.Num());
+
+	for (int32 CameraCutIndex = 0; CameraCutIndex < CameraCutSections.Num(); ++CameraCutIndex)
+	{
+		UMovieSceneCameraCutSection* CameraCutSection = CastChecked<UMovieSceneCameraCutSection>(CameraCutSections[CameraCutIndex]);
+		AcquireCameraForCameraCut(CameraCutIndex, CameraCutSection->GetCameraGuid(), SequenceInstance, Player);
+	}
+}
+
+	
+void FMovieSceneCameraCutTrackInstance::RestoreState(const TArray<UObject*>& RuntimeObjects, IMovieScenePlayer& Player, FMovieSceneSequenceInstance& SequenceInstance)
+{
+	Player.UpdateCameraCut(nullptr, nullptr);
+}
+	
+void FMovieSceneCameraCutTrackInstance::Update(EMovieSceneUpdateData& UpdateData, const TArray<UObject*>& RuntimeObjects, class IMovieScenePlayer& Player, FMovieSceneSequenceInstance& SequenceInstance) 
+{
+	// Don't switch CameraCuts when in preroll
+	if (UpdateData.bPreroll)
+	{
+		return;
+	}
+
+	UObject* CameraObject = nullptr;
+	const TArray<UMovieSceneSection*>& CameraCutSections = CameraCutTrack->GetAllSections();
+
+	UMovieSceneCameraCutSection* NearestCameraCutSection = Cast<UMovieSceneCameraCutSection>(MovieSceneHelpers::FindNearestSectionAtTime(CameraCutSections, UpdateData.Position));
+	if (NearestCameraCutSection != nullptr)
+	{
+		for (int32 CameraCutIndex = 0; CameraCutIndex < CameraCutSections.Num(); ++CameraCutIndex)
+		{
+			UMovieSceneCameraCutSection* CameraCutSection = CastChecked<UMovieSceneCameraCutSection>(CameraCutSections[CameraCutIndex]);
+
+			if (CameraCutSection == NearestCameraCutSection)
+			{
+				CameraObject = AcquireCameraForCameraCut(CameraCutIndex, CameraCutSection->GetCameraGuid(), SequenceInstance, Player);
+
+				break;
+			}
+		}
+	}
+
+	if (CameraObject != LastCameraObject)
+	{
+		Player.UpdateCameraCut(CameraObject, LastCameraObject.Get());
+		LastCameraObject = CameraObject;
+	}
+	else if (CameraObject != nullptr)
+	{
+		Player.UpdateCameraCut(CameraObject, nullptr);
+	}
+}
+
+
+/* FMovieSceneCameraCutTrackInstance implementation
+ *****************************************************************************/
+
+UObject* FMovieSceneCameraCutTrackInstance::AcquireCameraForCameraCut(int32 CameraCutIndex, const FGuid& CameraGuid, FMovieSceneSequenceInstance& SequenceInstance, const IMovieScenePlayer& Player)
+{
+	auto& CachedCamera = CachedCameraObjects[CameraCutIndex];
+
+	if (!CachedCamera.IsValid())
+	{
+		CachedCamera = SequenceInstance.FindObject(CameraGuid, Player);
+	}
+
+	return CachedCamera.Get();
+}
diff --git a/Engine/Source/Runtime/MovieSceneTracks/Private/TrackInstances/MovieSceneCinematicShotTrackInstance.cpp b/Engine/Source/Runtime/MovieSceneTracks/Private/TrackInstances/MovieSceneCinematicShotTrackInstance.cpp
new file mode 100644
--- /dev/null
+++ b/Engine/Source/Runtime/MovieSceneTracks/Private/TrackInstances/MovieSceneCinematicShotTrackInstance.cpp
@@ -0,0 +1,114 @@
+// Copyright 1998-2016 Epic Games, Inc. All Rights Reserved.
+
+#include "MovieSceneTracksPrivatePCH.h"
+#include "MovieSceneCinematicShotTrack.h"
+#include "MovieSceneCinematicShotTrackInstance.h"
+#include "MovieSceneSubTrackInstance.h"
+#include "MovieSceneSubSection.h"
+
+
+/* FMovieSceneCinematicShotTrackInstance structors
+ *****************************************************************************/
+
+FMovieSceneCinematicShotTrackInstance::FMovieSceneCinematicShotTrackInstance(UMovieSceneCinematicShotTrack& InTrack)
+	: FMovieSceneSubTrackInstance(InTrack)
+{ }
+
+/* IMovieSceneTrackInstance interface
+ *****************************************************************************/
+
+void FMovieSceneCinematicShotTrackInstance::ClearInstance(IMovieScenePlayer& Player, FMovieSceneSequenceInstance& SequenceInstance)
+{
+	FMovieSceneSubTrackInstance::ClearInstance(Player, SequenceInstance);
+
+	Player.UpdateCameraCut(nullptr, nullptr);
+}
+
+
+void FMovieSceneCinematicShotTrackInstance::RefreshInstance(const TArray<UObject*>& RuntimeObjects, IMovieScenePlayer& Player, FMovieSceneSequenceInstance& SequenceInstance)
+{
+	FMovieSceneSubTrackInstance::RefreshInstance(RuntimeObjects, Player, SequenceInstance);
+}
+
+	
+void FMovieSceneCinematicShotTrackInstance::RestoreState(const TArray<UObject*>& RuntimeObjects, IMovieScenePlayer& Player, FMovieSceneSequenceInstance& SequenceInstance)
+{
+	FMovieSceneSubTrackInstance::RestoreState(RuntimeObjects, Player, SequenceInstance);
+
+	Player.UpdateCameraCut(nullptr, nullptr);
+
+	//@todo Need to manage how this interacts with the fade track
+	/*
+	// Reset editor preview/fade
+	EMovieSceneViewportParams ViewportParams;
+	ViewportParams.SetWhichViewportParam = (EMovieSceneViewportParams::SetViewportParam)(EMovieSceneViewportParams::SVP_FadeAmount | EMovieSceneViewportParams::SVP_FadeColor);
+	ViewportParams.FadeAmount = 0.f;
+	ViewportParams.FadeColor = FLinearColor::Black;
+
+	TMap<FViewportClient*, EMovieSceneViewportParams> ViewportParamsMap;
+	Player.GetViewportSettings(ViewportParamsMap);
+	for (auto ViewportParamsPair : ViewportParamsMap)
+	{
+		ViewportParamsMap[ViewportParamsPair.Key] = ViewportParams;
+	}
+	Player.SetViewportSettings(ViewportParamsMap);
+	*/
+}
+	
+void FMovieSceneCinematicShotTrackInstance::Update(EMovieSceneUpdateData& UpdateData, const TArray<UObject*>& RuntimeObjects, class IMovieScenePlayer& Player, FMovieSceneSequenceInstance& SequenceInstance) 
+{
+	FMovieSceneSubTrackInstance::Update(UpdateData, RuntimeObjects, Player, SequenceInstance);
+
+	const TArray<UMovieSceneSection*>& AllSections = SubTrack->GetAllSections();
+	TArray<UMovieSceneSection*> TraversedSections = GetTraversedSectionsWithPreroll(AllSections, UpdateData.Position, UpdateData.LastPosition);
+
+	if (TraversedSections.Num() == 0)
+	{
+		Player.UpdateCameraCut(nullptr, nullptr);
+	}
+	else
+	{
+		//@todo Need to manage how this interacts with the fade track
+		/*
+		for (auto Section : TraversedSections)
+		{
+			UMovieSceneSubSection* SubSection = Cast<UMovieSceneSubSection>(Section);
+			if (SubSection->GetSequence() == nullptr)
+			{
+				// Set editor preview/fade
+				EMovieSceneViewportParams ViewportParams;
+				ViewportParams.SetWhichViewportParam = (EMovieSceneViewportParams::SetViewportParam)(EMovieSceneViewportParams::SVP_FadeAmount | EMovieSceneViewportParams::SVP_FadeColor);
+				ViewportParams.FadeAmount = 1.f;
+				ViewportParams.FadeColor = FLinearColor::Black;
+
+				TMap<FViewportClient*, EMovieSceneViewportParams> ViewportParamsMap;
+				Player.GetViewportSettings(ViewportParamsMap);
+				for (auto ViewportParamsPair : ViewportParamsMap)
+				{
+					ViewportParamsMap[ViewportParamsPair.Key] = ViewportParams;
+				}
+				Player.SetViewportSettings(ViewportParamsMap);
+
+				// Set runtime fade
+				for (const FWorldContext& Context : GEngine->GetWorldContexts())
+				{
+					if (Context.WorldType == EWorldType::Game || Context.WorldType == EWorldType::PIE)
+					{
+						UWorld* World = Context.World();
+						if (World != nullptr)
+						{
+							APlayerController* PlayerController = World->GetGameInstance()->GetFirstLocalPlayerController();
+							if (PlayerController != nullptr && PlayerController->PlayerCameraManager && !PlayerController->PlayerCameraManager->IsPendingKill())
+							{
+								PlayerController->PlayerCameraManager->SetManualCameraFade(1.f, FLinearColor::Black, false);
+							}
+						}
+					}
+				}
+
+				break;
+			}
+		}
+		*/
+	}
+}
\ No newline at end of file
diff --git a/Engine/Source/Runtime/MovieSceneTracks/Private/TrackInstances/MovieSceneColorTrackInstance.cpp b/Engine/Source/Runtime/MovieSceneTracks/Private/TrackInstances/MovieSceneColorTrackInstance.cpp
--- a/Engine/Source/Runtime/MovieSceneTracks/Private/TrackInstances/MovieSceneColorTrackInstance.cpp
+++ b/Engine/Source/Runtime/MovieSceneTracks/Private/TrackInstances/MovieSceneColorTrackInstance.cpp
@@ -91,7 +91,7 @@ void FMovieSceneColorTrackInstance::RestoreState(const TArray<UObject*>& Runtime
 }
 
 
-void FMovieSceneColorTrackInstance::Update( float Position, float LastPosition, const TArray<UObject*>& RuntimeObjects, class IMovieScenePlayer& Player, FMovieSceneSequenceInstance& SequenceInstance, EMovieSceneUpdatePass UpdatePass ) 
+void FMovieSceneColorTrackInstance::Update(EMovieSceneUpdateData& UpdateData, const TArray<UObject*>& RuntimeObjects, class IMovieScenePlayer& Player, FMovieSceneSequenceInstance& SequenceInstance) 
 {
 	for(UObject* Object : RuntimeObjects)
 	{
@@ -99,7 +99,7 @@ void FMovieSceneColorTrackInstance::Update( float Position, float LastPosition,
 		{
 			FSlateColor ColorValue = PropertyBindings->GetCurrentValue<FSlateColor>(Object);
 			FLinearColor LinearColor = ColorValue.GetSpecifiedColor();
-			if(ColorTrack->Eval(Position, LastPosition, LinearColor))
+			if(ColorTrack->Eval(UpdateData.Position, UpdateData.LastPosition, LinearColor))
 			{
 				FSlateColor NewColor(LinearColor);
 				PropertyBindings->CallFunction<FSlateColor>(Object, &NewColor);
@@ -108,7 +108,7 @@ void FMovieSceneColorTrackInstance::Update( float Position, float LastPosition,
 		else
 		{
 			FLinearColor ColorValue = ColorType == EColorType::Linear ? PropertyBindings->GetCurrentValue<FLinearColor>(Object) : PropertyBindings->GetCurrentValue<FColor>(Object).ReinterpretAsLinear();
-			if(ColorTrack->Eval(Position, LastPosition, ColorValue))
+			if(ColorTrack->Eval(UpdateData.Position, UpdateData.LastPosition, ColorValue))
 			{
 				// LightComponent's SetLightColor applies an sRGB conversion by default, so invert it here before applying the value 
 				if (ColorType == EColorType::RegularColor)
diff --git a/Engine/Source/Runtime/MovieSceneTracks/Private/TrackInstances/MovieSceneEventTrackInstance.cpp b/Engine/Source/Runtime/MovieSceneTracks/Private/TrackInstances/MovieSceneEventTrackInstance.cpp
--- a/Engine/Source/Runtime/MovieSceneTracks/Private/TrackInstances/MovieSceneEventTrackInstance.cpp
+++ b/Engine/Source/Runtime/MovieSceneTracks/Private/TrackInstances/MovieSceneEventTrackInstance.cpp
@@ -17,7 +17,7 @@ FMovieSceneEventTrackInstance::FMovieSceneEventTrackInstance(UMovieSceneEventTra
 /* IMovieSceneTrackInstance interface
  *****************************************************************************/
 
-void FMovieSceneEventTrackInstance::Update(float Position, float LastPosition, const TArray<UObject*>& RuntimeObjects, IMovieScenePlayer& Player, FMovieSceneSequenceInstance& SequenceInstance, EMovieSceneUpdatePass UpdatePass)
+void FMovieSceneEventTrackInstance::Update(EMovieSceneUpdateData& UpdateData, const TArray<UObject*>& RuntimeObjects, IMovieScenePlayer& Player, FMovieSceneSequenceInstance& SequenceInstance)
 {
-	EventTrack->TriggerEvents(Position, LastPosition);
+	EventTrack->TriggerEvents(UpdateData.Position, UpdateData.LastPosition);
 }
diff --git a/Engine/Source/Runtime/MovieSceneTracks/Private/TrackInstances/MovieSceneFadeTrackInstance.cpp b/Engine/Source/Runtime/MovieSceneTracks/Private/TrackInstances/MovieSceneFadeTrackInstance.cpp
--- a/Engine/Source/Runtime/MovieSceneTracks/Private/TrackInstances/MovieSceneFadeTrackInstance.cpp
+++ b/Engine/Source/Runtime/MovieSceneTracks/Private/TrackInstances/MovieSceneFadeTrackInstance.cpp
@@ -34,16 +34,16 @@ void FMovieSceneFadeTrackInstance::RestoreState(const TArray<UObject*>& RuntimeO
 	Player.SetViewportSettings(ViewportParamsMap);
 }
 
-void FMovieSceneFadeTrackInstance::Update(float Position, float LastPosition, const TArray<UObject*>& RuntimeObjects, IMovieScenePlayer& Player, FMovieSceneSequenceInstance& SequenceInstance, EMovieSceneUpdatePass UpdatePass)
+void FMovieSceneFadeTrackInstance::Update(EMovieSceneUpdateData& UpdateData, const TArray<UObject*>& RuntimeObjects, IMovieScenePlayer& Player, FMovieSceneSequenceInstance& SequenceInstance)
 {
 	FLinearColor FadeColor = FLinearColor::Black;
 	bool bFadeAudio = false;
 
 	float FloatValue = 0.0f;
 
-	if (FadeTrack->Eval(Position, LastPosition, FloatValue))
+	if (FadeTrack->Eval(UpdateData.Position, UpdateData.LastPosition, FloatValue))
 	{
-		const UMovieSceneSection* NearestSection = MovieSceneHelpers::FindNearestSectionAtTime(FadeTrack->GetAllSections(), Position);
+		const UMovieSceneSection* NearestSection = MovieSceneHelpers::FindNearestSectionAtTime(FadeTrack->GetAllSections(), UpdateData.Position);
 		const UMovieSceneFadeSection* FadeSection = CastChecked<const UMovieSceneFadeSection>(NearestSection);
 		if (FadeSection)
 		{
diff --git a/Engine/Source/Runtime/MovieSceneTracks/Private/TrackInstances/MovieSceneFloatTrackInstance.cpp b/Engine/Source/Runtime/MovieSceneTracks/Private/TrackInstances/MovieSceneFloatTrackInstance.cpp
--- a/Engine/Source/Runtime/MovieSceneTracks/Private/TrackInstances/MovieSceneFloatTrackInstance.cpp
+++ b/Engine/Source/Runtime/MovieSceneTracks/Private/TrackInstances/MovieSceneFloatTrackInstance.cpp
@@ -46,10 +46,10 @@ void FMovieSceneFloatTrackInstance::RestoreState(const TArray<UObject*>& Runtime
 }
 
 
-void FMovieSceneFloatTrackInstance::Update( float Position, float LastPosition, const TArray<UObject*>& RuntimeObjects, class IMovieScenePlayer& Player, FMovieSceneSequenceInstance& SequenceInstance, EMovieSceneUpdatePass UpdatePass ) 
+void FMovieSceneFloatTrackInstance::Update(EMovieSceneUpdateData& UpdateData, const TArray<UObject*>& RuntimeObjects, class IMovieScenePlayer& Player, FMovieSceneSequenceInstance& SequenceInstance ) 
 {
 	float FloatValue = 0.0f;
-	if( FloatTrack->Eval( Position, LastPosition, FloatValue ) )
+	if( FloatTrack->Eval( UpdateData.Position, UpdateData.LastPosition, FloatValue ) )
 	{
 		for(UObject* Object : RuntimeObjects)
 		{
diff --git a/Engine/Source/Runtime/MovieSceneTracks/Private/TrackInstances/MovieSceneMaterialTrackInstance.cpp b/Engine/Source/Runtime/MovieSceneTracks/Private/TrackInstances/MovieSceneMaterialTrackInstance.cpp
--- a/Engine/Source/Runtime/MovieSceneTracks/Private/TrackInstances/MovieSceneMaterialTrackInstance.cpp
+++ b/Engine/Source/Runtime/MovieSceneTracks/Private/TrackInstances/MovieSceneMaterialTrackInstance.cpp
@@ -53,11 +53,11 @@ void FMovieSceneMaterialTrackInstance::RestoreState(const TArray<UObject*>& Runt
 	}
 }
 
-void FMovieSceneMaterialTrackInstance::Update( float Position, float LastPosition, const TArray<UObject*>& RuntimeObjects, class IMovieScenePlayer& Player, FMovieSceneSequenceInstance& SequenceInstance, EMovieSceneUpdatePass UpdatePass ) 
+void FMovieSceneMaterialTrackInstance::Update(EMovieSceneUpdateData& UpdateData, const TArray<UObject*>& RuntimeObjects, class IMovieScenePlayer& Player, FMovieSceneSequenceInstance& SequenceInstance ) 
 {
 	TArray<FScalarParameterNameAndValue> ScalarValues;
 	TArray<FColorParameterNameAndValue> VectorValues;
-	MaterialTrack->Eval( Position, ScalarValues, VectorValues );
+	MaterialTrack->Eval( UpdateData.Position, ScalarValues, VectorValues );
 
 	// Iterate from back to front to allow for fast remove of invalid weak pointers.
 	for ( int32 i = DynamicMaterialInstances.Num() - 1; i >= 0; i--)
diff --git a/Engine/Source/Runtime/MovieSceneTracks/Private/TrackInstances/MovieSceneParticleParameterTrackInstance.cpp b/Engine/Source/Runtime/MovieSceneTracks/Private/TrackInstances/MovieSceneParticleParameterTrackInstance.cpp
--- a/Engine/Source/Runtime/MovieSceneTracks/Private/TrackInstances/MovieSceneParticleParameterTrackInstance.cpp
+++ b/Engine/Source/Runtime/MovieSceneTracks/Private/TrackInstances/MovieSceneParticleParameterTrackInstance.cpp
@@ -110,12 +110,12 @@ void FMovieSceneParticleParameterTrackInstance::RestoreState(const TArray<UObjec
 	ObjectToInitialParameterValuesMap.Empty();
 }
 
-void FMovieSceneParticleParameterTrackInstance::Update( float Position, float LastPosition, const TArray<UObject*>& RuntimeObjects, class IMovieScenePlayer& Player, FMovieSceneSequenceInstance& SequenceInstance, EMovieSceneUpdatePass UpdatePass )
+void FMovieSceneParticleParameterTrackInstance::Update(EMovieSceneUpdateData& UpdateData, const TArray<UObject*>& RuntimeObjects, class IMovieScenePlayer& Player, FMovieSceneSequenceInstance& SequenceInstance)
 {
 	TArray<FScalarParameterNameAndValue> ScalarNamesAndValues;
 	TArray<FVectorParameterNameAndValue> VectorNamesAndValues;
 	TArray<FColorParameterNameAndValue> ColorNamesAndValues;
-	ParticleParameterTrack->Eval( Position, ScalarNamesAndValues, VectorNamesAndValues, ColorNamesAndValues);
+	ParticleParameterTrack->Eval( UpdateData.Position, ScalarNamesAndValues, VectorNamesAndValues, ColorNamesAndValues);
 	for ( UParticleSystemComponent* ParticleSystemComponent : ParticleSystemComponents )
 	{
 		for ( const FScalarParameterNameAndValue& ScalarNameAndValue : ScalarNamesAndValues )
diff --git a/Engine/Source/Runtime/MovieSceneTracks/Private/TrackInstances/MovieSceneParticleTrackInstance.cpp b/Engine/Source/Runtime/MovieSceneTracks/Private/TrackInstances/MovieSceneParticleTrackInstance.cpp
--- a/Engine/Source/Runtime/MovieSceneTracks/Private/TrackInstances/MovieSceneParticleTrackInstance.cpp
+++ b/Engine/Source/Runtime/MovieSceneTracks/Private/TrackInstances/MovieSceneParticleTrackInstance.cpp
@@ -13,12 +13,12 @@ FMovieSceneParticleTrackInstance::~FMovieSceneParticleTrackInstance()
 }
 
 
-void FMovieSceneParticleTrackInstance::Update( float Position, float LastPosition, const TArray<UObject*>& RuntimeObjects, class IMovieScenePlayer& Player, FMovieSceneSequenceInstance& SequenceInstance, EMovieSceneUpdatePass UpdatePass ) 
+void FMovieSceneParticleTrackInstance::Update(EMovieSceneUpdateData& UpdateData, const TArray<UObject*>& RuntimeObjects, class IMovieScenePlayer& Player, FMovieSceneSequenceInstance& SequenceInstance) 
 {
 	// @todo Sequencer We need something analagous to Matinee 1's particle replay tracks
 	// What we have here is simple toggling/triggering
 
-	if (Position > LastPosition && Player.GetPlaybackStatus() == EMovieScenePlayerStatus::Playing)
+	if (UpdateData.Position > UpdateData.LastPosition && Player.GetPlaybackStatus() == EMovieScenePlayerStatus::Playing)
 	{
 		const TArray<UMovieSceneSection*> Sections = ParticleTrack->GetAllParticleSections();
 		EParticleKey::Type ParticleKey = EParticleKey::Deactivate;
@@ -30,11 +30,11 @@ void FMovieSceneParticleTrackInstance::Update( float Position, float LastPositio
 			if (Section->IsActive())
 			{
 				FIntegralCurve& ParticleKeyCurve = Section->GetParticleCurve();
-				FKeyHandle PreviousHandle = ParticleKeyCurve.FindKeyBeforeOrAt( Position );
+				FKeyHandle PreviousHandle = ParticleKeyCurve.FindKeyBeforeOrAt( UpdateData.Position );
 				if ( ParticleKeyCurve.IsKeyHandleValid( PreviousHandle ) )
 				{
 					FIntegralKey& PreviousKey = ParticleKeyCurve.GetKey( PreviousHandle );
-					if ( PreviousKey.Time > LastPosition )
+					if ( PreviousKey.Time > UpdateData.LastPosition )
 					{
 						ParticleKey = (EParticleKey::Type)PreviousKey.Value;
 						bKeyFound = true;
diff --git a/Engine/Source/Runtime/MovieSceneTracks/Private/TrackInstances/MovieSceneSkeletalAnimationTrackInstance.cpp b/Engine/Source/Runtime/MovieSceneTracks/Private/TrackInstances/MovieSceneSkeletalAnimationTrackInstance.cpp
--- a/Engine/Source/Runtime/MovieSceneTracks/Private/TrackInstances/MovieSceneSkeletalAnimationTrackInstance.cpp
+++ b/Engine/Source/Runtime/MovieSceneTracks/Private/TrackInstances/MovieSceneSkeletalAnimationTrackInstance.cpp
@@ -55,7 +55,7 @@ FMovieSceneSkeletalAnimationTrackInstance::~FMovieSceneSkeletalAnimationTrackIns
 	// @todo Sequencer Need to find some way to call PreviewFinishAnimControl (needs the runtime objects)
 }
 
-void FMovieSceneSkeletalAnimationTrackInstance::Update( float Position, float LastPosition, const TArray<UObject*>& RuntimeObjects, class IMovieScenePlayer& Player, FMovieSceneSequenceInstance& SequenceInstance, EMovieSceneUpdatePass UpdatePass ) 
+void FMovieSceneSkeletalAnimationTrackInstance::Update(EMovieSceneUpdateData& UpdateData, const TArray<UObject*>& RuntimeObjects, class IMovieScenePlayer& Player, FMovieSceneSequenceInstance& SequenceInstance) 
 {
 	// @todo Sequencer gameplay update has a different code path than editor update for animation
 
@@ -64,13 +64,13 @@ void FMovieSceneSkeletalAnimationTrackInstance::Update( float Position, float La
 		IMatineeAnimInterface* AnimInterface = Cast<IMatineeAnimInterface>(RuntimeObjects[i]);
 		if (AnimInterface) 
 		{
-			UMovieSceneSkeletalAnimationSection* AnimSection = Cast<UMovieSceneSkeletalAnimationSection>(AnimationTrack->GetAnimSectionAtTime(Position));
+			UMovieSceneSkeletalAnimationSection* AnimSection = Cast<UMovieSceneSkeletalAnimationSection>(AnimationTrack->GetAnimSectionAtTime(UpdateData.Position));
 			
 			// cbb: If there is no overlapping section, evaluate the closest section only if the current time is before it.
 			if (AnimSection == nullptr)
 			{
-				AnimSection = Cast<UMovieSceneSkeletalAnimationSection>(MovieSceneHelpers::FindNearestSectionAtTime(AnimationTrack->GetAllSections(), Position));
-				if (AnimSection != nullptr && (Position >= AnimSection->GetStartTime() || Position <= AnimSection->GetEndTime()))
+				AnimSection = Cast<UMovieSceneSkeletalAnimationSection>(MovieSceneHelpers::FindNearestSectionAtTime(AnimationTrack->GetAllSections(), UpdateData.Position));
+				if (AnimSection != nullptr && (UpdateData.Position >= AnimSection->GetStartTime() || UpdateData.Position <= AnimSection->GetEndTime()))
 				{
 					AnimSection = nullptr;
 				}
@@ -78,7 +78,7 @@ void FMovieSceneSkeletalAnimationTrackInstance::Update( float Position, float La
 
 			if (AnimSection && AnimSection->IsActive())
 			{
-				const FAnimEvalTimes EvalTimes = FAnimEvalTimes::MapTimesToAnimation(Position, LastPosition, AnimSection);
+				const FAnimEvalTimes EvalTimes = FAnimEvalTimes::MapTimesToAnimation(UpdateData.Position, UpdateData.LastPosition, AnimSection);
 
 				int32 ChannelIndex = 0;
 				UAnimSequence* AnimSequence = AnimSection->GetAnimSequence();
diff --git a/Engine/Source/Runtime/MovieSceneTracks/Private/TrackInstances/MovieSceneSlomoTrackInstance.cpp b/Engine/Source/Runtime/MovieSceneTracks/Private/TrackInstances/MovieSceneSlomoTrackInstance.cpp
--- a/Engine/Source/Runtime/MovieSceneTracks/Private/TrackInstances/MovieSceneSlomoTrackInstance.cpp
+++ b/Engine/Source/Runtime/MovieSceneTracks/Private/TrackInstances/MovieSceneSlomoTrackInstance.cpp
@@ -41,7 +41,7 @@ void FMovieSceneSlomoTrackInstance::SaveState(const TArray<UObject*>& RuntimeObj
 /* IMovieSceneTrackInstance interface
  *****************************************************************************/
 
-void FMovieSceneSlomoTrackInstance::Update(float Position, float LastPosition, const TArray<UObject*>& RuntimeObjects, IMovieScenePlayer& Player, FMovieSceneSequenceInstance& SequenceInstance, EMovieSceneUpdatePass UpdatePass)
+void FMovieSceneSlomoTrackInstance::Update(EMovieSceneUpdateData& UpdateData, const TArray<UObject*>& RuntimeObjects, IMovieScenePlayer& Player, FMovieSceneSequenceInstance& SequenceInstance)
 {
 	if (!ShouldBeApplied())
 	{
@@ -57,7 +57,7 @@ void FMovieSceneSlomoTrackInstance::Update(float Position, float LastPosition, c
 
 	float FloatValue = 0.0f;
 
-	if (SlomoTrack->Eval(Position, LastPosition, FloatValue))
+	if (SlomoTrack->Eval(UpdateData.Position, UpdateData.LastPosition, FloatValue))
 	{
 		WorldSettings->MatineeTimeDilation = FloatValue;
 		WorldSettings->ForceNetUpdate();
diff --git a/Engine/Source/Runtime/MovieSceneTracks/Private/TrackInstances/MovieSceneSpawnTrackInstance.cpp b/Engine/Source/Runtime/MovieSceneTracks/Private/TrackInstances/MovieSceneSpawnTrackInstance.cpp
--- a/Engine/Source/Runtime/MovieSceneTracks/Private/TrackInstances/MovieSceneSpawnTrackInstance.cpp
+++ b/Engine/Source/Runtime/MovieSceneTracks/Private/TrackInstances/MovieSceneSpawnTrackInstance.cpp
@@ -13,7 +13,7 @@ FMovieSceneSpawnTrackInstance::FMovieSceneSpawnTrackInstance(UMovieSceneSpawnTra
 	Track = &InTrack;
 }
 
-void FMovieSceneSpawnTrackInstance::Update(float Position, float LastPosition, const TArray<UObject*>& RuntimeObjects, IMovieScenePlayer& Player, FMovieSceneSequenceInstance& SequenceInstance, EMovieSceneUpdatePass UpdatePass)
+void FMovieSceneSpawnTrackInstance::Update(EMovieSceneUpdateData& UpdateData, const TArray<UObject*>& RuntimeObjects, IMovieScenePlayer& Player, FMovieSceneSequenceInstance& SequenceInstance)
 {
 	IMovieSceneSpawnRegister& SpawnRegister = Player.GetSpawnRegister();
 	FMovieSceneSpawnable* Spawnable = SequenceInstance.GetSequence()->GetMovieScene()->FindSpawnable(Track->GetObjectId());
@@ -21,14 +21,14 @@ void FMovieSceneSpawnTrackInstance::Update(float Position, float LastPosition, c
 	TRange<float> Range = SequenceInstance.GetTimeRange();
 
 	// If we're evaluating outside of the instance's time range, and the sequence owns the spawnable, there's no reason to evaluate - it should already be destroyed
-	if (Spawnable && Spawnable->GetSpawnOwnership() == ESpawnOwnership::InnerSequence && !Range.Contains(Position) && !Range.Contains(LastPosition))
+	if (Spawnable && Spawnable->GetSpawnOwnership() == ESpawnOwnership::InnerSequence && !Range.Contains(UpdateData.Position) && !Range.Contains(UpdateData.LastPosition))
 	{
 		SpawnRegister.DestroySpawnedObject(Track->GetObjectId(), SequenceInstance, Player);
 		return;
 	}
 
 	bool bIsSpawned = false;
-	if (Track->Eval(Position, LastPosition, bIsSpawned))
+	if (Track->Eval(UpdateData.Position, UpdateData.LastPosition, bIsSpawned))
 	{
 		// Spawn the object if needed
 		if (bIsSpawned && RuntimeObjects.Num() == 0)
diff --git a/Engine/Source/Runtime/MovieSceneTracks/Private/TrackInstances/MovieSceneSubTrackInstance.cpp b/Engine/Source/Runtime/MovieSceneTracks/Private/TrackInstances/MovieSceneSubTrackInstance.cpp
--- a/Engine/Source/Runtime/MovieSceneTracks/Private/TrackInstances/MovieSceneSubTrackInstance.cpp
+++ b/Engine/Source/Runtime/MovieSceneTracks/Private/TrackInstances/MovieSceneSubTrackInstance.cpp
@@ -104,10 +104,75 @@ void FMovieSceneSubTrackInstance::SaveState(const TArray<UObject*>& RuntimeObjec
 }
 
 
-void FMovieSceneSubTrackInstance::Update(float Position, float LastPosition, const TArray<UObject*>& RuntimeObjects, class IMovieScenePlayer& Player, FMovieSceneSequenceInstance& SequenceInstance, EMovieSceneUpdatePass UpdatePass) 
+TArray<UMovieSceneSection*> FMovieSceneSubTrackInstance::GetAllTraversedSectionsWithPreroll( const TArray<UMovieSceneSection*>& Sections, float CurrentTime, float PreviousTime )
+{
+	TArray<UMovieSceneSection*> TraversedSections;
+
+	bool bPlayingBackwards = CurrentTime - PreviousTime < 0.0f;
+	float MaxTime = bPlayingBackwards ? PreviousTime : CurrentTime;
+	float MinTime = bPlayingBackwards ? CurrentTime : PreviousTime;
+
+	TRange<float> TraversedRange(MinTime, TRangeBound<float>::Inclusive(MaxTime));
+
+	for (int32 SectionIndex = 0; SectionIndex < Sections.Num(); ++SectionIndex)
+	{
+		UMovieSceneSection* Section = Sections[SectionIndex];
+
+		float PrerollTime = 0.0f;
+		UMovieSceneSubSection* SubSection = CastChecked<UMovieSceneSubSection>(Section);
+		if (SubSection != nullptr)
+		{
+			PrerollTime = SubSection->PrerollTime;
+		}
+
+		if ((Section->GetStartTime()-PrerollTime == CurrentTime) || TraversedRange.Overlaps(TRange<float>(Section->GetRange().GetLowerBoundValue()-PrerollTime, Section->GetRange().GetUpperBoundValue())))
+		{
+			TraversedSections.Add(Section);
+		}
+	}
+
+	return TraversedSections;
+}
+
+TArray<UMovieSceneSection*> FMovieSceneSubTrackInstance::GetTraversedSectionsWithPreroll( const TArray<UMovieSceneSection*>& Sections, float CurrentTime, float PreviousTime )
+{
+	TArray<UMovieSceneSection*> TraversedSections = GetAllTraversedSectionsWithPreroll(Sections, CurrentTime, PreviousTime);
+
+	// Remove any overlaps that are underneath another
+	for (int32 RemoveAt = 0; RemoveAt < TraversedSections.Num(); )
+	{
+		UMovieSceneSection* Section = TraversedSections[RemoveAt];
+		
+		const bool bShouldRemove = TraversedSections.ContainsByPredicate([=](UMovieSceneSection* OtherSection){
+			if (Section->GetRowIndex() == OtherSection->GetRowIndex() &&
+				Section->GetRange().Overlaps(OtherSection->GetRange()) &&
+				Section->GetOverlapPriority() < OtherSection->GetOverlapPriority())
+			{
+				return true;
+			}
+			return false;
+		});
+		
+		if (bShouldRemove)
+		{
+			TraversedSections.RemoveAt(RemoveAt, 1, false);
+		}
+		else
+		{
+			++RemoveAt;
+		}
+	}
+
+	return TraversedSections;
+}
+
+
+void FMovieSceneSubTrackInstance::Update(EMovieSceneUpdateData& UpdateData, const TArray<UObject*>& RuntimeObjects, class IMovieScenePlayer& Player, FMovieSceneSequenceInstance& SequenceInstance) 
 {
 	const TArray<UMovieSceneSection*>& AllSections = SubTrack->GetAllSections();
-	TArray<UMovieSceneSection*> TraversedSections = MovieSceneHelpers::GetTraversedSections(AllSections, Position, LastPosition);
+	TArray<UMovieSceneSection*> TraversedSections = GetTraversedSectionsWithPreroll(AllSections, UpdateData.Position, UpdateData.LastPosition);
+
+	const float InitialUpdatePosition = UpdateData.Position;
 
 	for (const auto Section : TraversedSections)
 	{
@@ -128,11 +193,15 @@ void FMovieSceneSubTrackInstance::Update(float Position, float LastPosition, con
 		}
 
 		// calculate section's local time
-		const float InstanceOffset = SubSection->StartOffset + Instance->GetTimeRange().GetLowerBoundValue();
-		const float InstanceLastPosition = InstanceOffset + (LastPosition - SubSection->GetStartTime()) / SubSection->TimeScale;
-		const float InstancePosition = InstanceOffset + (Position - SubSection->GetStartTime()) / SubSection->TimeScale;
+		const float InstanceOffset = SubSection->StartOffset + Instance->GetTimeRange().GetLowerBoundValue() - SubSection->PrerollTime;
+		const float InstanceLastPosition = InstanceOffset + (UpdateData.LastPosition - (SubSection->GetStartTime()- SubSection->PrerollTime)) / SubSection->TimeScale;
+		const float InstancePosition = InstanceOffset + (UpdateData.Position - (SubSection->GetStartTime()- SubSection->PrerollTime)) / SubSection->TimeScale;
+
+		EMovieSceneUpdateData SubUpdateData(InstancePosition, InstanceLastPosition);
+		SubUpdateData.UpdatePass = UpdateData.UpdatePass;
+		SubUpdateData.bPreroll = InitialUpdatePosition < SubSection->GetStartTime();
 
 		// update section
-		Instance->Update(InstancePosition, InstanceLastPosition, Player);
+		Instance->Update(SubUpdateData, Player);
 	}
 }
diff --git a/Engine/Source/Runtime/MovieSceneTracks/Private/TrackInstances/MovieSceneVectorTrackInstance.cpp b/Engine/Source/Runtime/MovieSceneTracks/Private/TrackInstances/MovieSceneVectorTrackInstance.cpp
--- a/Engine/Source/Runtime/MovieSceneTracks/Private/TrackInstances/MovieSceneVectorTrackInstance.cpp
+++ b/Engine/Source/Runtime/MovieSceneTracks/Private/TrackInstances/MovieSceneVectorTrackInstance.cpp
@@ -97,10 +97,10 @@ void FMovieSceneVectorTrackInstance::RestoreState(const TArray<UObject*>& Runtim
 }
 
 
-void FMovieSceneVectorTrackInstance::Update( float Position, float LastPosition, const TArray<UObject*>& RuntimeObjects, class IMovieScenePlayer& Player, FMovieSceneSequenceInstance& SequenceInstance, EMovieSceneUpdatePass UpdatePass ) 
+void FMovieSceneVectorTrackInstance::Update( EMovieSceneUpdateData& UpdateData, const TArray<UObject*>& RuntimeObjects, class IMovieScenePlayer& Player, FMovieSceneSequenceInstance& SequenceInstance ) 
 {
 	FVector4 Vector;
-	if( VectorTrack->Eval( Position, LastPosition, Vector ) )
+	if( VectorTrack->Eval( UpdateData.Position, UpdateData.LastPosition, Vector ) )
 	{
 		int32 NumChannelsUsed = VectorTrack->GetNumChannelsUsed();
 		switch( NumChannelsUsed )
diff --git a/Engine/Source/Runtime/MovieSceneTracks/Private/TrackInstances/MovieSceneVisibilityTrackInstance.cpp b/Engine/Source/Runtime/MovieSceneTracks/Private/TrackInstances/MovieSceneVisibilityTrackInstance.cpp
--- a/Engine/Source/Runtime/MovieSceneTracks/Private/TrackInstances/MovieSceneVisibilityTrackInstance.cpp
+++ b/Engine/Source/Runtime/MovieSceneTracks/Private/TrackInstances/MovieSceneVisibilityTrackInstance.cpp
@@ -22,7 +22,7 @@ void FMovieSceneVisibilityTrackInstance::SaveState(const TArray<UObject*>& Runti
 		if (Actor != nullptr)
 		{
 #if WITH_EDITOR
-			if (GIsEditor && !Actor->GetWorld()->IsPlayInEditor())
+			if (GIsEditor && Actor->GetWorld() != nullptr && !Actor->GetWorld()->IsPlayInEditor())
 			{
 				if (InitHiddenInEditorMap.Find(Actor) == nullptr)
 				{
@@ -68,7 +68,7 @@ void FMovieSceneVisibilityTrackInstance::RestoreState(const TArray<UObject*>& Ru
 			}
 
 #if WITH_EDITOR
-			if (GIsEditor && !Actor->GetWorld()->IsPlayInEditor())
+			if (GIsEditor && Actor->GetWorld() != nullptr && !Actor->GetWorld()->IsPlayInEditor())
 			{
 				bool *HiddenInEditorValue = InitHiddenInEditorMap.Find(Object);
 				if (HiddenInEditorValue != nullptr)
@@ -98,10 +98,10 @@ void FMovieSceneVisibilityTrackInstance::RestoreState(const TArray<UObject*>& Ru
 }
 
 
-void FMovieSceneVisibilityTrackInstance::Update( float Position, float LastPosition, const TArray<UObject*>& RuntimeObjects, class IMovieScenePlayer& Player, FMovieSceneSequenceInstance& SequenceInstance, EMovieSceneUpdatePass UpdatePass ) 
+void FMovieSceneVisibilityTrackInstance::Update( EMovieSceneUpdateData& UpdateData, const TArray<UObject*>& RuntimeObjects, class IMovieScenePlayer& Player, FMovieSceneSequenceInstance& SequenceInstance ) 
 {
  	bool Visible = false;
- 	if( VisibilityTrack->Eval( Position, LastPosition, Visible ) )
+ 	if( VisibilityTrack->Eval( UpdateData.Position, UpdateData.LastPosition, Visible ) )
  	{
 		// Invert this evaluation since the property is "bHiddenInGame" and we want the visualization to be the inverse of that. Green means visible.
 		Visible = !Visible;
@@ -115,7 +115,7 @@ void FMovieSceneVisibilityTrackInstance::Update( float Position, float LastPosit
 				Actor->SetActorHiddenInGame(Visible);
 
 #if WITH_EDITOR
-				if (GIsEditor && !Actor->GetWorld()->IsPlayInEditor())
+				if (GIsEditor && Actor->GetWorld() != nullptr && !Actor->GetWorld()->IsPlayInEditor())
 				{
 					// In editor HiddenGame flag is not respected so set bHiddenEdTemporary too.
 					// It will be restored just like HiddenGame flag when Matinee is closed.
diff --git a/Engine/Source/Runtime/MovieSceneTracks/Private/Tracks/MovieScene3DConstraintTrack.cpp b/Engine/Source/Runtime/MovieSceneTracks/Private/Tracks/MovieScene3DConstraintTrack.cpp
--- a/Engine/Source/Runtime/MovieSceneTracks/Private/Tracks/MovieScene3DConstraintTrack.cpp
+++ b/Engine/Source/Runtime/MovieSceneTracks/Private/Tracks/MovieScene3DConstraintTrack.cpp
@@ -65,7 +65,7 @@ TRange<float> UMovieScene3DConstraintTrack::GetSectionBoundaries() const
 
 
 #if WITH_EDITORONLY_DATA
-FText UMovieScene3DConstraintTrack::GetDisplayName() const
+FText UMovieScene3DConstraintTrack::GetDefaultDisplayName() const
 {
 	return LOCTEXT("TrackName", "Constraint");
 }
diff --git a/Engine/Source/Runtime/MovieSceneTracks/Private/Tracks/MovieSceneCameraCutTrack.cpp b/Engine/Source/Runtime/MovieSceneTracks/Private/Tracks/MovieSceneCameraCutTrack.cpp
new file mode 100644
--- /dev/null
+++ b/Engine/Source/Runtime/MovieSceneTracks/Private/Tracks/MovieSceneCameraCutTrack.cpp
@@ -0,0 +1,127 @@
+// Copyright 1998-2016 Epic Games, Inc. All Rights Reserved.
+
+#include "MovieSceneTracksPrivatePCH.h"
+#include "MovieSceneCameraCutSection.h"
+#include "MovieSceneCameraCutTrack.h"
+#include "MovieSceneCommonHelpers.h"
+#include "IMovieScenePlayer.h"
+#include "MovieSceneCameraCutTrackInstance.h"
+
+
+#define LOCTEXT_NAMESPACE "MovieSceneCameraCutTrack"
+
+
+/* UMovieSceneCameraCutTrack interface
+ *****************************************************************************/
+
+void UMovieSceneCameraCutTrack::AddNewCameraCut(FGuid CameraHandle, float StartTime)
+{
+	Modify();
+
+	UMovieSceneCameraCutSection* NewSection = NewObject<UMovieSceneCameraCutSection>(this, NAME_None, RF_Transactional);
+	{
+		NewSection->SetStartTime(StartTime);
+		NewSection->SetEndTime(FindEndTimeForCameraCut(StartTime));
+		NewSection->SetCameraGuid(CameraHandle);
+	}
+
+	Sections.Add(NewSection);
+
+	// When a new CameraCut is added, sort all CameraCuts to ensure they are in the correct order
+	MovieSceneHelpers::SortConsecutiveSections(Sections);
+
+	// Once CameraCuts are sorted fixup the surrounding CameraCuts to fix any gaps
+	MovieSceneHelpers::FixupConsecutiveSections(Sections, *NewSection, false);
+}
+
+
+/* UMovieSceneTrack interface
+ *****************************************************************************/
+
+void UMovieSceneCameraCutTrack::AddSection(UMovieSceneSection& Section)
+{
+	if (Section.IsA<UMovieSceneCameraCutSection>())
+	{
+		Sections.Add(&Section);
+	}
+}
+
+
+TSharedPtr<IMovieSceneTrackInstance> UMovieSceneCameraCutTrack::CreateInstance()
+{
+	return MakeShareable(new FMovieSceneCameraCutTrackInstance(*this)); 
+}
+
+
+UMovieSceneSection* UMovieSceneCameraCutTrack::CreateNewSection()
+{
+	return NewObject<UMovieSceneCameraCutSection>(this, NAME_None, RF_Transactional);
+}
+
+
+const TArray<UMovieSceneSection*>& UMovieSceneCameraCutTrack::GetAllSections() const
+{
+	return Sections;
+}
+
+
+void UMovieSceneCameraCutTrack::RemoveSection(UMovieSceneSection& Section)
+{
+	Sections.Remove(&Section);
+	MovieSceneHelpers::FixupConsecutiveSections(Sections, Section, true);
+	MovieSceneHelpers::SortConsecutiveSections(Sections);
+
+	// @todo Sequencer: The movie scene owned by the section is now abandoned.  Should we offer to delete it?  
+}
+
+
+#if WITH_EDITORONLY_DATA
+FText UMovieSceneCameraCutTrack::GetDefaultDisplayName() const
+{
+	return LOCTEXT("TrackName", "Camera Cuts");
+}
+#endif
+
+
+#if WITH_EDITOR
+void UMovieSceneCameraCutTrack::OnSectionMoved(UMovieSceneSection& Section)
+{
+	MovieSceneHelpers::FixupConsecutiveSections(Sections, Section, false);
+}
+#endif
+
+float UMovieSceneCameraCutTrack::FindEndTimeForCameraCut( float StartTime )
+{
+	float EndTime = 0;
+	bool bFoundEndTime = false;
+
+	for( UMovieSceneSection* Section : Sections )
+	{
+		if( Section->GetStartTime() >= StartTime )
+		{
+			EndTime = Section->GetStartTime();
+			bFoundEndTime = true;
+			break;
+		}
+	}
+
+	if( !bFoundEndTime )
+	{
+		UMovieScene* OwnerScene = GetTypedOuter<UMovieScene>();
+
+		// End time should just end where the movie scene ends.  Ensure it is at least the same as start time (this should only happen when the movie scene has an initial time range smaller than the start time
+		EndTime = FMath::Max( OwnerScene->GetPlaybackRange().GetUpperBoundValue(), StartTime );
+	}
+
+			
+	if( StartTime == EndTime )
+	{
+		// Give the CameraCut a reasonable length of time to start out with.  A 0 time CameraCut is not usable
+		EndTime = StartTime + .5f;
+	}
+
+	return EndTime;
+}
+
+
+#undef LOCTEXT_NAMESPACE
\ No newline at end of file
diff --git a/Engine/Source/Runtime/MovieSceneTracks/Private/Tracks/MovieSceneCinematicShotTrack.cpp b/Engine/Source/Runtime/MovieSceneTracks/Private/Tracks/MovieSceneCinematicShotTrack.cpp
new file mode 100644
--- /dev/null
+++ b/Engine/Source/Runtime/MovieSceneTracks/Private/Tracks/MovieSceneCinematicShotTrack.cpp
@@ -0,0 +1,97 @@
+// Copyright 1998-2016 Epic Games, Inc. All Rights Reserved.
+
+#include "MovieSceneTracksPrivatePCH.h"
+#include "MovieSceneSequence.h"
+#include "MovieSceneCinematicShotSection.h"
+#include "MovieSceneCinematicShotTrack.h"
+#include "MovieSceneCinematicShotTrackInstance.h"
+#include "MovieSceneCommonHelpers.h"
+
+
+#define LOCTEXT_NAMESPACE "MovieSceneCinematicShotTrack"
+
+
+/* UMovieSceneSubTrack interface
+ *****************************************************************************/
+
+UMovieSceneSubSection* UMovieSceneCinematicShotTrack::AddSequence(UMovieSceneSequence* Sequence, float StartTime, float Duration, const bool& bInsertSequence)
+{
+	UMovieSceneSubSection* NewSection = UMovieSceneSubTrack::AddSequence(Sequence, StartTime, Duration, bInsertSequence);
+
+	UMovieSceneCinematicShotSection* NewShotSection = Cast<UMovieSceneCinematicShotSection>(NewSection);
+
+#if WITH_EDITOR
+
+	if (Sequence != nullptr)
+	{
+		NewShotSection->SetShotDisplayName(Sequence->GetDisplayName());	
+	}
+
+#endif
+
+	// When a new sequence is added, sort all sequences to ensure they are in the correct order
+	MovieSceneHelpers::SortConsecutiveSections(Sections);
+
+	// Once sequences are sorted fixup the surrounding sequences to fix any gaps
+	MovieSceneHelpers::FixupConsecutiveSections(Sections, *NewSection, false);
+
+	return NewSection;
+}
+
+/* UMovieSceneTrack interface
+ *****************************************************************************/
+
+void UMovieSceneCinematicShotTrack::AddSection(UMovieSceneSection& Section)
+{
+	if (Section.IsA<UMovieSceneCinematicShotSection>())
+	{
+		Sections.Add(&Section);
+	}
+}
+
+
+TSharedPtr<IMovieSceneTrackInstance> UMovieSceneCinematicShotTrack::CreateInstance()
+{
+	return MakeShareable(new FMovieSceneCinematicShotTrackInstance(*this)); 
+}
+
+
+UMovieSceneSection* UMovieSceneCinematicShotTrack::CreateNewSection()
+{
+	return NewObject<UMovieSceneCinematicShotSection>(this, NAME_None, RF_Transactional);
+}
+
+void UMovieSceneCinematicShotTrack::RemoveSection(UMovieSceneSection& Section)
+{
+	Sections.Remove(&Section);
+	MovieSceneHelpers::FixupConsecutiveSections(Sections, Section, true);
+	MovieSceneHelpers::SortConsecutiveSections(Sections);
+
+	// @todo Sequencer: The movie scene owned by the section is now abandoned.  Should we offer to delete it?  
+}
+
+bool UMovieSceneCinematicShotTrack::SupportsMultipleRows() const
+{
+	return false;
+}
+
+bool UMovieSceneCinematicShotTrack::AddsSectionBoundsToPlayRange() const
+{
+	return true;
+}
+
+#if WITH_EDITOR
+void UMovieSceneCinematicShotTrack::OnSectionMoved(UMovieSceneSection& Section)
+{
+	MovieSceneHelpers::FixupConsecutiveSections(Sections, Section, false);
+}
+#endif
+
+#if WITH_EDITORONLY_DATA
+FText UMovieSceneCinematicShotTrack::GetDefaultDisplayName() const
+{
+	return LOCTEXT("TrackName", "Shots");
+}
+#endif
+
+#undef LOCTEXT_NAMESPACE
diff --git a/Engine/Source/Runtime/MovieSceneTracks/Private/Tracks/MovieSceneEventTrack.cpp b/Engine/Source/Runtime/MovieSceneTracks/Private/Tracks/MovieSceneEventTrack.cpp
--- a/Engine/Source/Runtime/MovieSceneTracks/Private/Tracks/MovieSceneEventTrack.cpp
+++ b/Engine/Source/Runtime/MovieSceneTracks/Private/Tracks/MovieSceneEventTrack.cpp
@@ -6,6 +6,8 @@
 #include "MovieSceneEventTrackInstance.h"
 
 
+#define LOCTEXT_NAMESPACE "MovieSceneEventTrack"
+
 /* UMovieSceneEventTrack interface
  *****************************************************************************/
 
@@ -140,3 +142,15 @@ void UMovieSceneEventTrack::RemoveSection(UMovieSceneSection& Section)
 {
 	Sections.Remove(&Section);
 }
+
+#if WITH_EDITORONLY_DATA
+
+FText UMovieSceneEventTrack::GetDefaultDisplayName() const
+{ 
+	return LOCTEXT("TrackName", "Events"); 
+}
+
+#endif
+
+
+#undef LOCTEXT_NAMESPACE
diff --git a/Engine/Source/Runtime/MovieSceneTracks/Private/Tracks/MovieSceneMaterialTrack.cpp b/Engine/Source/Runtime/MovieSceneTracks/Private/Tracks/MovieSceneMaterialTrack.cpp
--- a/Engine/Source/Runtime/MovieSceneTracks/Private/Tracks/MovieSceneMaterialTrack.cpp
+++ b/Engine/Source/Runtime/MovieSceneTracks/Private/Tracks/MovieSceneMaterialTrack.cpp
@@ -6,67 +6,77 @@
 #include "IMovieScenePlayer.h"
 #include "MovieSceneMaterialTrackInstance.h"
 
-UMovieSceneMaterialTrack::UMovieSceneMaterialTrack( const FObjectInitializer& ObjectInitializer )
-	: Super( ObjectInitializer )
+
+UMovieSceneMaterialTrack::UMovieSceneMaterialTrack(const FObjectInitializer& ObjectInitializer)
+	: Super(ObjectInitializer)
 {
 }
 
+
 UMovieSceneSection* UMovieSceneMaterialTrack::CreateNewSection()
 {
-	return NewObject<UMovieSceneParameterSection>( this, UMovieSceneParameterSection::StaticClass(), NAME_None, RF_Transactional );
+	return NewObject<UMovieSceneParameterSection>(this, UMovieSceneParameterSection::StaticClass(), NAME_None, RF_Transactional);
 }
 
+
 void UMovieSceneMaterialTrack::RemoveAllAnimationData()
 {
 	Sections.Empty();
 }
 
+
 bool UMovieSceneMaterialTrack::HasSection(const UMovieSceneSection& Section) const
 {
 	return Sections.Contains(&Section);
 }
 
+
 void UMovieSceneMaterialTrack::AddSection(UMovieSceneSection& Section)
 {
 	Sections.Add(&Section);
 }
 
+
 void UMovieSceneMaterialTrack::RemoveSection(UMovieSceneSection& Section)
 {
 	Sections.Remove(&Section);
 }
 
+
 bool UMovieSceneMaterialTrack::IsEmpty() const
 {
 	return Sections.Num() == 0;
 }
 
+
 TRange<float> UMovieSceneMaterialTrack::GetSectionBoundaries() const
 {
 	TArray< TRange<float> > Bounds;
 
-	for ( int32 SectionIndex = 0; SectionIndex < Sections.Num(); ++SectionIndex )
+	for (int32 SectionIndex = 0; SectionIndex < Sections.Num(); ++SectionIndex)
 	{
-		Bounds.Add( Sections[SectionIndex]->GetRange() );
+		Bounds.Add(Sections[SectionIndex]->GetRange());
 	}
 
-	return TRange<float>::Hull( Bounds );
+	return TRange<float>::Hull(Bounds);
 }
 
+
 const TArray<UMovieSceneSection*>& UMovieSceneMaterialTrack::GetAllSections() const
 {
 	return Sections;
 }
 
-void UMovieSceneMaterialTrack::AddScalarParameterKey( FName ParameterName, float Time, float Value )
+
+void UMovieSceneMaterialTrack::AddScalarParameterKey(FName ParameterName, float Time, float Value)
 {
 	UMovieSceneParameterSection* NearestSection = Cast<UMovieSceneParameterSection>(MovieSceneHelpers::FindNearestSectionAtTime(Sections, Time));
-	if ( NearestSection == nullptr )
+	if (NearestSection == nullptr)
 	{
 		NearestSection = Cast<UMovieSceneParameterSection>(CreateNewSection());
-		NearestSection->SetStartTime( Time );
-		NearestSection->SetEndTime( Time );
-		Sections.Add( NearestSection );
+		NearestSection->SetStartTime(Time);
+		NearestSection->SetEndTime(Time);
+		Sections.Add(NearestSection);
 	}
 	if (NearestSection->TryModify())
 	{
@@ -74,46 +84,50 @@ void UMovieSceneMaterialTrack::AddScalarParameterKey( FName ParameterName, float
 	}
 }
 
-void UMovieSceneMaterialTrack::AddColorParameterKey( FName ParameterName, float Time, FLinearColor Value )
+
+void UMovieSceneMaterialTrack::AddColorParameterKey(FName ParameterName, float Time, FLinearColor Value)
 {
-	UMovieSceneParameterSection* NearestSection = Cast<UMovieSceneParameterSection>( MovieSceneHelpers::FindNearestSectionAtTime( Sections, Time ) );
-	if ( NearestSection == nullptr )
+	UMovieSceneParameterSection* NearestSection = Cast<UMovieSceneParameterSection>(MovieSceneHelpers::FindNearestSectionAtTime(Sections, Time));
+	if (NearestSection == nullptr)
 	{
-		NearestSection = Cast<UMovieSceneParameterSection>( CreateNewSection() );
-		NearestSection->SetStartTime( Time );
-		NearestSection->SetEndTime( Time );
-		Sections.Add( NearestSection );
+		NearestSection = Cast<UMovieSceneParameterSection>(CreateNewSection());
+		NearestSection->SetStartTime(Time);
+		NearestSection->SetEndTime(Time);
+		Sections.Add(NearestSection);
 	}
 	if (NearestSection->TryModify())
 	{
-		NearestSection->AddColorParameterKey( ParameterName, Time, Value );
+		NearestSection->AddColorParameterKey(ParameterName, Time, Value);
 	}
 }
 
-void UMovieSceneMaterialTrack::Eval( float Position, TArray<FScalarParameterNameAndValue>& OutScalarValues, TArray<FColorParameterNameAndValue>& OutColorValues ) const
+
+void UMovieSceneMaterialTrack::Eval(float Position, TArray<FScalarParameterNameAndValue>& OutScalarValues, TArray<FColorParameterNameAndValue>& OutColorValues) const
 {
 	TArray<FVectorParameterNameAndValue> OutVectorValues;
-	for ( UMovieSceneSection* Section : Sections )
+	for (UMovieSceneSection* Section : Sections)
 	{
-		UMovieSceneParameterSection* ParameterSection = Cast<UMovieSceneParameterSection>( Section );
+		UMovieSceneParameterSection* ParameterSection = Cast<UMovieSceneParameterSection>(Section);
 		ParameterSection->Eval(Position, OutScalarValues, OutVectorValues, OutColorValues);
 	}
 }
 
-UMovieSceneComponentMaterialTrack::UMovieSceneComponentMaterialTrack( const FObjectInitializer& ObjectInitializer )
-	: Super( ObjectInitializer )
+
+UMovieSceneComponentMaterialTrack::UMovieSceneComponentMaterialTrack(const FObjectInitializer& ObjectInitializer)
+	: Super(ObjectInitializer)
 {
 	
 }
 
+
 TSharedPtr<IMovieSceneTrackInstance> UMovieSceneComponentMaterialTrack::CreateInstance()
 {
-	return MakeShareable( new FMovieSceneComponentMaterialTrackInstance( *this ) );
+	return MakeShareable(new FMovieSceneComponentMaterialTrackInstance(*this));
 }
 
 
 #if WITH_EDITORONLY_DATA
-FText UMovieSceneComponentMaterialTrack::GetDisplayName() const
+FText UMovieSceneComponentMaterialTrack::GetDefaultDisplayName() const
 {
 	return FText::FromString(FString::Printf(TEXT("Material Element %i"), MaterialIndex));
 }
diff --git a/Engine/Source/Runtime/MovieSceneTracks/Private/Tracks/MovieSceneParticleParameterTrack.cpp b/Engine/Source/Runtime/MovieSceneTracks/Private/Tracks/MovieSceneParticleParameterTrack.cpp
--- a/Engine/Source/Runtime/MovieSceneTracks/Private/Tracks/MovieSceneParticleParameterTrack.cpp
+++ b/Engine/Source/Runtime/MovieSceneTracks/Private/Tracks/MovieSceneParticleParameterTrack.cpp
@@ -67,7 +67,7 @@ const TArray<UMovieSceneSection*>& UMovieSceneParticleParameterTrack::GetAllSect
 
 
 #if WITH_EDITORONLY_DATA
-FText UMovieSceneParticleParameterTrack::GetDisplayName() const
+FText UMovieSceneParticleParameterTrack::GetDefaultDisplayName() const
 {
 	return LOCTEXT("DisplayName", "Particle Parameter");
 }
diff --git a/Engine/Source/Runtime/MovieSceneTracks/Private/Tracks/MovieSceneParticleTrack.cpp b/Engine/Source/Runtime/MovieSceneTracks/Private/Tracks/MovieSceneParticleTrack.cpp
--- a/Engine/Source/Runtime/MovieSceneTracks/Private/Tracks/MovieSceneParticleTrack.cpp
+++ b/Engine/Source/Runtime/MovieSceneTracks/Private/Tracks/MovieSceneParticleTrack.cpp
@@ -87,7 +87,7 @@ UMovieSceneSection* UMovieSceneParticleTrack::CreateNewSection()
 }
 
 #if WITH_EDITORONLY_DATA
-FText UMovieSceneParticleTrack::GetDisplayName() const
+FText UMovieSceneParticleTrack::GetDefaultDisplayName() const
 {
 	return LOCTEXT("DisplayName", "Particle System");
 }
diff --git a/Engine/Source/Runtime/MovieSceneTracks/Private/Tracks/MovieScenePropertyTrack.cpp b/Engine/Source/Runtime/MovieSceneTracks/Private/Tracks/MovieScenePropertyTrack.cpp
--- a/Engine/Source/Runtime/MovieSceneTracks/Private/Tracks/MovieScenePropertyTrack.cpp
+++ b/Engine/Source/Runtime/MovieSceneTracks/Private/Tracks/MovieScenePropertyTrack.cpp
@@ -4,13 +4,13 @@
 #include "MovieScenePropertyTrack.h"
 
 
-UMovieScenePropertyTrack::UMovieScenePropertyTrack( const FObjectInitializer& ObjectInitializer )
-	: Super( ObjectInitializer )
+UMovieScenePropertyTrack::UMovieScenePropertyTrack(const FObjectInitializer& ObjectInitializer)
+	: Super(ObjectInitializer)
 {
 }
 
 
-void UMovieScenePropertyTrack::SetPropertyNameAndPath( FName InPropertyName, const FString& InPropertyPath )
+void UMovieScenePropertyTrack::SetPropertyNameAndPath(FName InPropertyName, const FString& InPropertyPath)
 {
 	check((InPropertyName != NAME_None) && !InPropertyPath.IsEmpty());
 
@@ -26,7 +26,7 @@ const TArray<UMovieSceneSection*>& UMovieScenePropertyTrack::GetAllSections() co
 
 
 #if WITH_EDITORONLY_DATA
-FText UMovieScenePropertyTrack::GetDisplayName() const
+FText UMovieScenePropertyTrack::GetDefaultDisplayName() const
 {
 	return FText::FromName(PropertyName);
 }
@@ -69,51 +69,51 @@ bool UMovieScenePropertyTrack::IsEmpty() const
 }
 
 
-UMovieSceneSection* UMovieScenePropertyTrack::FindOrAddSection( float Time )
+UMovieSceneSection* UMovieScenePropertyTrack::FindOrAddSection(float Time)
 {
 	// Find a spot for the section so that they are sorted by start time
-	for( int32 SectionIndex = 0; SectionIndex < Sections.Num(); ++SectionIndex )
+	for(int32 SectionIndex = 0; SectionIndex < Sections.Num(); ++SectionIndex)
 	{
 		UMovieSceneSection* Section = Sections[SectionIndex];
 
-		if( Section->IsTimeWithinSection( Time ) )
+		if(Section->IsTimeWithinSection(Time))
 		{
 			return Section;
 		}
 
 		// Check if there are no more sections that would overlap the time 
-		if( !Sections.IsValidIndex( SectionIndex+1 ) || Sections[SectionIndex+1]->GetStartTime() > Time )
+		if(!Sections.IsValidIndex(SectionIndex+1) || Sections[SectionIndex+1]->GetStartTime() > Time)
 		{
 			// No sections overlap the time
 		
-			if( SectionIndex > 0 )
+			if(SectionIndex > 0)
 			{
 				// Append and grow the previous section
 				UMovieSceneSection* PreviousSection = Sections[ SectionIndex ? SectionIndex-1 : 0 ];
 		
-				PreviousSection->SetEndTime( Time );
+				PreviousSection->SetEndTime(Time);
 				return PreviousSection;
 			}
-			else if( Sections.IsValidIndex( SectionIndex+1 ) )
+			else if(Sections.IsValidIndex(SectionIndex+1))
 			{
 				// Prepend and grow the next section because there are no sections before this one
 				UMovieSceneSection* NextSection = Sections[SectionIndex+1];
-				NextSection->SetStartTime( Time );
+				NextSection->SetStartTime(Time);
 				return NextSection;
 			}	
 			else
 			{
 				// SectionIndex == 0 
 				UMovieSceneSection* PreviousSection = Sections[0];
-				if( PreviousSection->GetEndTime() < Time )
+				if(PreviousSection->GetEndTime() < Time)
 				{
 					// Append and grow the section
-					PreviousSection->SetEndTime( Time );
+					PreviousSection->SetEndTime(Time);
 				}
 				else
 				{
 					// Prepend and grow the section
-					PreviousSection->SetStartTime( Time );
+					PreviousSection->SetStartTime(Time);
 				}
 				return PreviousSection;
 			}
@@ -121,15 +121,15 @@ UMovieSceneSection* UMovieScenePropertyTrack::FindOrAddSection( float Time )
 
 	}
 
-	check( Sections.Num() == 0 );
+	check(Sections.Num() == 0);
 
 	// Add a new section that starts and ends at the same time
 	UMovieSceneSection* NewSection = CreateNewSection();
 	NewSection->SetFlags(RF_Transactional);
-	NewSection->SetStartTime( Time );
-	NewSection->SetEndTime( Time );
+	NewSection->SetStartTime(Time);
+	NewSection->SetEndTime(Time);
 
-	Sections.Add( NewSection );
+	Sections.Add(NewSection);
 
 	return NewSection;
 }
diff --git a/Engine/Source/Runtime/MovieSceneTracks/Private/Tracks/MovieSceneSkeletalAnimationTrack.cpp b/Engine/Source/Runtime/MovieSceneTracks/Private/Tracks/MovieSceneSkeletalAnimationTrack.cpp
--- a/Engine/Source/Runtime/MovieSceneTracks/Private/Tracks/MovieSceneSkeletalAnimationTrack.cpp
+++ b/Engine/Source/Runtime/MovieSceneTracks/Private/Tracks/MovieSceneSkeletalAnimationTrack.cpp
@@ -13,8 +13,8 @@
 /* UMovieSceneSkeletalAnimationTrack structors
  *****************************************************************************/
 
-UMovieSceneSkeletalAnimationTrack::UMovieSceneSkeletalAnimationTrack( const FObjectInitializer& ObjectInitializer )
-	: Super( ObjectInitializer )
+UMovieSceneSkeletalAnimationTrack::UMovieSceneSkeletalAnimationTrack(const FObjectInitializer& ObjectInitializer)
+	: Super(ObjectInitializer)
 { }
 
 
@@ -52,7 +52,7 @@ UMovieSceneSection* UMovieSceneSkeletalAnimationTrack::GetAnimSectionAtTime(floa
 
 TSharedPtr<IMovieSceneTrackInstance> UMovieSceneSkeletalAnimationTrack::CreateInstance()
 {
-	return MakeShareable( new FMovieSceneSkeletalAnimationTrackInstance( *this ) ); 
+	return MakeShareable(new FMovieSceneSkeletalAnimationTrackInstance(*this)); 
 }
 
 
@@ -64,7 +64,7 @@ const TArray<UMovieSceneSection*>& UMovieSceneSkeletalAnimationTrack::GetAllSect
 
 UMovieSceneSection* UMovieSceneSkeletalAnimationTrack::CreateNewSection()
 {
-	return NewObject<UMovieSceneSkeletalAnimationSection>( this );
+	return NewObject<UMovieSceneSkeletalAnimationSection>(this);
 }
 
 
@@ -74,7 +74,7 @@ void UMovieSceneSkeletalAnimationTrack::RemoveAllAnimationData()
 }
 
 
-bool UMovieSceneSkeletalAnimationTrack::HasSection(const UMovieSceneSection& Section ) const
+bool UMovieSceneSkeletalAnimationTrack::HasSection(const UMovieSceneSection& Section) const
 {
 	return AnimationSections.Contains(&Section);
 }
@@ -113,7 +113,7 @@ TRange<float> UMovieSceneSkeletalAnimationTrack::GetSectionBoundaries() const
 
 #if WITH_EDITORONLY_DATA
 
-FText UMovieSceneSkeletalAnimationTrack::GetDisplayName() const
+FText UMovieSceneSkeletalAnimationTrack::GetDefaultDisplayName() const
 {
 	return LOCTEXT("TrackName", "Animation");
 }
diff --git a/Engine/Source/Runtime/MovieSceneTracks/Private/Tracks/MovieSceneSlomoTrack.cpp b/Engine/Source/Runtime/MovieSceneTracks/Private/Tracks/MovieSceneSlomoTrack.cpp
--- a/Engine/Source/Runtime/MovieSceneTracks/Private/Tracks/MovieSceneSlomoTrack.cpp
+++ b/Engine/Source/Runtime/MovieSceneTracks/Private/Tracks/MovieSceneSlomoTrack.cpp
@@ -26,7 +26,7 @@ UMovieSceneSection* UMovieSceneSlomoTrack::CreateNewSection()
 
 #if WITH_EDITORONLY_DATA
 
-FText UMovieSceneSlomoTrack::GetDisplayName() const
+FText UMovieSceneSlomoTrack::GetDefaultDisplayName() const
 {
 	return LOCTEXT("TrackName", "Play Rate");
 }
diff --git a/Engine/Source/Runtime/MovieSceneTracks/Private/Tracks/MovieSceneSpawnTrack.cpp b/Engine/Source/Runtime/MovieSceneTracks/Private/Tracks/MovieSceneSpawnTrack.cpp
--- a/Engine/Source/Runtime/MovieSceneTracks/Private/Tracks/MovieSceneSpawnTrack.cpp
+++ b/Engine/Source/Runtime/MovieSceneTracks/Private/Tracks/MovieSceneSpawnTrack.cpp
@@ -51,11 +51,11 @@ UMovieSceneSection* UMovieSceneSpawnTrack::CreateNewSection()
 
 TSharedPtr<IMovieSceneTrackInstance> UMovieSceneSpawnTrack::CreateInstance()
 {
-	return MakeShareable( new FMovieSceneSpawnTrackInstance( *this ) );
+	return MakeShareable(new FMovieSceneSpawnTrackInstance(*this));
 }
 
 
-bool UMovieSceneSpawnTrack::HasSection(const UMovieSceneSection& Section ) const
+bool UMovieSceneSpawnTrack::HasSection(const UMovieSceneSection& Section) const
 {
 	return Sections.ContainsByPredicate([&](const UMovieSceneSection* In){ return In == &Section; });
 }
@@ -93,7 +93,7 @@ const TArray<UMovieSceneSection*>& UMovieSceneSpawnTrack::GetAllSections() const
 
 #if WITH_EDITORONLY_DATA
 
-FText UMovieSceneSpawnTrack::GetDisplayName() const
+FText UMovieSceneSpawnTrack::GetDefaultDisplayName() const
 {
 	return LOCTEXT("TrackName", "Spawned");
 }
diff --git a/Engine/Source/Runtime/MovieSceneTracks/Private/Tracks/MovieSceneSubTrack.cpp b/Engine/Source/Runtime/MovieSceneTracks/Private/Tracks/MovieSceneSubTrack.cpp
--- a/Engine/Source/Runtime/MovieSceneTracks/Private/Tracks/MovieSceneSubTrack.cpp
+++ b/Engine/Source/Runtime/MovieSceneTracks/Private/Tracks/MovieSceneSubTrack.cpp
@@ -14,18 +14,82 @@
 /* UMovieSceneSubTrack interface
  *****************************************************************************/
 
-void UMovieSceneSubTrack::AddSequence(UMovieSceneSequence& Sequence, float Time)
+UMovieSceneSubSection* UMovieSceneSubTrack::AddSequence(UMovieSceneSequence* Sequence, float StartTime, float Duration, const bool& bInsertSequence)
 {
 	Modify();
 
+	// If inserting, push all existing sections out to make space for the new one
+	if (bInsertSequence)
+	{
+		// If there's a shot that starts at the same time as this new shot, push the shots forward to make space for this one
+		bool bPushShotsForward = false;
+		for (auto Section : Sections)
+		{
+			float SectionStartTime = Section->GetStartTime();
+			if (FMath::IsNearlyEqual(SectionStartTime, StartTime))
+			{
+				bPushShotsForward = true;
+				break;
+			}
+		}
+		if (bPushShotsForward)
+		{
+			for (auto Section : Sections)
+			{
+				float SectionStartTime = Section->GetStartTime();
+				if (SectionStartTime >= StartTime)
+				{
+					Section->SetStartTime(SectionStartTime + Duration);
+					Section->SetEndTime(Section->GetEndTime() + Duration);
+				}
+			}
+		}
+		// Otherwise, see if there's a shot after the start time and clamp the duration to that next shot
+		else
+		{
+			bool bFoundAfterShot = false;
+			float MinDiff = FLT_MAX;
+			for (auto Section : Sections)
+			{
+				float SectionStartTime = Section->GetStartTime();
+				if (SectionStartTime > StartTime)
+				{
+					bFoundAfterShot = true;
+					float Diff = SectionStartTime - StartTime;
+					MinDiff = FMath::Min(MinDiff, Diff);
+				}
+			}
+
+			if (MinDiff != FLT_MAX)
+			{
+				Duration = MinDiff;
+			}
+		}
+	}
+	// Otherwise, append this to the end of the existing sequences
+	else
+	{
+		if (Sections.Num())
+		{
+			StartTime = Sections[0]->GetEndTime();
+
+			for (const auto& Section : Sections)
+			{
+				StartTime = FMath::Max(StartTime, Section->GetEndTime());
+			}
+		}
+	}
+
 	UMovieSceneSubSection* NewSection = CastChecked<UMovieSceneSubSection>(CreateNewSection());
 	{
 		NewSection->SetSequence(Sequence);
-		NewSection->SetStartTime(Time);
-		NewSection->SetEndTime(Time + Sequence.GetMovieScene()->GetPlaybackRange().Size<float>());
+		NewSection->SetStartTime(StartTime);
+		NewSection->SetEndTime(StartTime + Duration);
 	}
 
 	Sections.Add(NewSection);
+
+	return NewSection;
 }
 
 
@@ -141,7 +205,7 @@ bool UMovieSceneSubTrack::SupportsMultipleRows() const
 }
 
 #if WITH_EDITORONLY_DATA
-FText UMovieSceneSubTrack::GetDisplayName() const
+FText UMovieSceneSubTrack::GetDefaultDisplayName() const
 {
 	return LOCTEXT("TrackName", "Subscenes");
 }
diff --git a/Engine/Source/Runtime/MovieSceneTracks/Private/Tracks/MovieSceneVisibilityTrack.cpp b/Engine/Source/Runtime/MovieSceneTracks/Private/Tracks/MovieSceneVisibilityTrack.cpp
--- a/Engine/Source/Runtime/MovieSceneTracks/Private/Tracks/MovieSceneVisibilityTrack.cpp
+++ b/Engine/Source/Runtime/MovieSceneTracks/Private/Tracks/MovieSceneVisibilityTrack.cpp
@@ -7,8 +7,11 @@
 #include "MovieSceneVisibilityTrackInstance.h"
 
 
-UMovieSceneVisibilityTrack::UMovieSceneVisibilityTrack( const FObjectInitializer& ObjectInitializer )
-	: Super( ObjectInitializer )
+#define LOCTEXT_NAMESPACE "MovieSceneVisibilityTrack"
+
+
+UMovieSceneVisibilityTrack::UMovieSceneVisibilityTrack(const FObjectInitializer& ObjectInitializer)
+	: Super(ObjectInitializer)
 { }
 
 
@@ -20,5 +23,16 @@ UMovieSceneSection* UMovieSceneVisibilityTrack::CreateNewSection()
 
 TSharedPtr<IMovieSceneTrackInstance> UMovieSceneVisibilityTrack::CreateInstance()
 {
-	return MakeShareable( new FMovieSceneVisibilityTrackInstance( *this ) );
+	return MakeShareable(new FMovieSceneVisibilityTrackInstance(*this));
+}
+
+
+#if WITH_EDITORONLY_DATA
+FText UMovieSceneVisibilityTrack::GetDefaultDisplayName() const
+{
+	return LOCTEXT("DisplayName", "Visibility");
 }
+#endif
+
+
+#undef LOCTEXT_NAMESPACE
diff --git a/Engine/Source/Runtime/MovieSceneTracks/Public/Sections/MovieScene3DTransformSection.h b/Engine/Source/Runtime/MovieSceneTracks/Public/Sections/MovieScene3DTransformSection.h
--- a/Engine/Source/Runtime/MovieSceneTracks/Public/Sections/MovieScene3DTransformSection.h
+++ b/Engine/Source/Runtime/MovieSceneTracks/Public/Sections/MovieScene3DTransformSection.h
@@ -2,8 +2,9 @@
 
 #pragma once
 
-#include "MovieSceneSection.h"
 #include "IKeyframeSection.h"
+#include "MovieSceneKeyStruct.h"
+#include "MovieSceneSection.h"
 #include "MovieScene3DTransformSection.generated.h"
 
 
@@ -69,6 +70,35 @@ struct FTransformKey
 };
 
 
+/**
+ * Proxy structure for 3D transform section key data.
+ */
+USTRUCT()
+struct FMovieScene3DTransformKeyStruct
+	: public FMovieSceneKeyStruct
+{
+	GENERATED_BODY()
+
+	/** They key's translation value. */
+	UPROPERTY(EditAnywhere, Category=Key)
+	FVector4 Translation;
+
+	/** They key's rotation value. */
+	UPROPERTY(EditAnywhere, Category=Key)
+	FRotator Rotation;
+
+	/** They key's scale value. */
+	UPROPERTY(EditAnywhere, Category=Key)
+	FVector4 Scale;
+
+	FRichCurveKey* TranslationKeys[3];
+	FRichCurveKey* RotationKeys[3];
+	FRichCurveKey* ScaleKeys[3];
+
+	virtual void PropagateChanges(const FPropertyChangedEvent& ChangeEvent) override;
+};
+
+
 /**
  * A 3D transform section
  */
@@ -81,11 +111,6 @@ class UMovieScene3DTransformSection
 
 public:
 
-	/** MovieSceneSection interface */
-	virtual void MoveSection( float DeltaPosition, TSet<FKeyHandle>& KeyHandles ) override;
-	virtual void DilateSection( float DilationFactor, float Origin, TSet<FKeyHandle>& KeyHandles ) override;
-	virtual void GetKeyHandles(TSet<FKeyHandle>& KeyHandles) const override;
-
 	/**
 	 * Evaluates the translation component of the transform
 	 *
@@ -134,6 +159,17 @@ public:
 	 */
 	MOVIESCENETRACKS_API FRichCurve& GetScaleCurve( EAxis::Type Axis );
 
+public:
+
+	// UMovieSceneSection interface
+
+	virtual void MoveSection( float DeltaPosition, TSet<FKeyHandle>& KeyHandles ) override;
+	virtual void DilateSection( float DilationFactor, float Origin, TSet<FKeyHandle>& KeyHandles ) override;
+	virtual void GetKeyHandles(TSet<FKeyHandle>& KeyHandles) const override;
+	virtual TSharedPtr<FStructOnScope> GetKeyStruct(const TArray<FKeyHandle>& KeyHandles) override;
+
+public:
+
 	// IKeyframeSection interface.
 
 	virtual bool NewKeyIsNewData( float Time, const FTransformKey& KeyData ) const override;
@@ -142,6 +178,7 @@ public:
 	virtual void SetDefault( const FTransformKey& KeyData ) override;
 
 private:
+
 	/** Translation curves */
 	UPROPERTY()
 	FRichCurve Translation[3];
diff --git a/Engine/Source/Runtime/MovieSceneTracks/Public/Sections/MovieSceneBoolSection.h b/Engine/Source/Runtime/MovieSceneTracks/Public/Sections/MovieSceneBoolSection.h
--- a/Engine/Source/Runtime/MovieSceneTracks/Public/Sections/MovieSceneBoolSection.h
+++ b/Engine/Source/Runtime/MovieSceneTracks/Public/Sections/MovieSceneBoolSection.h
@@ -49,6 +49,6 @@ private:
 	/** Ordered curve data */
 	// @todo Sequencer This could be optimized by packing the bools separately
 	// but that may not be worth the effort
-	UPROPERTY(EditAnywhere, Category="Curve")
+	UPROPERTY()
 	FIntegralCurve BoolCurve;
 };
diff --git a/Engine/Source/Runtime/MovieSceneTracks/Public/Sections/MovieSceneByteSection.h b/Engine/Source/Runtime/MovieSceneTracks/Public/Sections/MovieSceneByteSection.h
--- a/Engine/Source/Runtime/MovieSceneTracks/Public/Sections/MovieSceneByteSection.h
+++ b/Engine/Source/Runtime/MovieSceneTracks/Public/Sections/MovieSceneByteSection.h
@@ -41,6 +41,6 @@ private:
 	/** Ordered curve data */
 	// @todo Sequencer This could be optimized by packing the bytes separately
 	// but that may not be worth the effort
-	UPROPERTY(EditAnywhere, Category="Curve")
+	UPROPERTY()
 	FIntegralCurve ByteCurve;
 };
diff --git a/Engine/Source/Runtime/MovieSceneTracks/Public/Sections/MovieSceneCameraCutSection.h b/Engine/Source/Runtime/MovieSceneTracks/Public/Sections/MovieSceneCameraCutSection.h
new file mode 100644
--- /dev/null
+++ b/Engine/Source/Runtime/MovieSceneTracks/Public/Sections/MovieSceneCameraCutSection.h
@@ -0,0 +1,37 @@
+// Copyright 1998-2016 Epic Games, Inc. All Rights Reserved.
+
+#pragma once
+
+#include "MovieSceneSection.h"
+#include "MovieSceneCameraCutSection.generated.h"
+
+
+/**
+ * Movie CameraCuts are sections on the CameraCuts track, that show what the viewer "sees"
+ */
+UCLASS(MinimalAPI)
+class UMovieSceneCameraCutSection 
+	: public UMovieSceneSection
+{
+	GENERATED_BODY()
+
+public:
+
+	/** Gets the camera guid for this CameraCut section */
+	FGuid GetCameraGuid() const
+	{
+		return CameraGuid;
+	}
+
+	/** Sets the camera guid for this CameraCut section */
+	void SetCameraGuid(const FGuid& InGuid)
+	{
+		CameraGuid = InGuid;
+	}
+
+private:
+
+	/** The camera possessable or spawnable that this movie CameraCut uses */
+	UPROPERTY()
+	FGuid CameraGuid;
+};
diff --git a/Engine/Source/Runtime/MovieSceneTracks/Public/Sections/MovieSceneCinematicShotSection.h b/Engine/Source/Runtime/MovieSceneTracks/Public/Sections/MovieSceneCinematicShotSection.h
new file mode 100644
--- /dev/null
+++ b/Engine/Source/Runtime/MovieSceneTracks/Public/Sections/MovieSceneCinematicShotSection.h
@@ -0,0 +1,45 @@
+// Copyright 1998-2016 Epic Games, Inc. All Rights Reserved.
+
+#pragma once
+
+#include "MovieSceneSubSection.h"
+#include "MovieSceneCinematicShotSection.generated.h"
+
+class UMovieSceneSequence;
+
+
+/**
+ * Implements a cinematic shot section.
+ */
+UCLASS()
+class MOVIESCENETRACKS_API UMovieSceneCinematicShotSection
+	: public UMovieSceneSubSection
+{
+	GENERATED_BODY()
+
+	/** Default constructor. */
+	UMovieSceneCinematicShotSection();
+
+public:
+
+	/** @return The shot display name */
+	FText GetShotDisplayName() const
+	{
+		return DisplayName;
+	}
+
+	/** Set the shot display name */
+	void SetShotDisplayName(const FText& InDisplayName)
+	{
+		if (TryModify())
+		{
+			DisplayName = InDisplayName;
+		}
+	}
+
+private:
+
+	/** The Shot's display name */
+	UPROPERTY()
+	FText DisplayName;
+};
diff --git a/Engine/Source/Runtime/MovieSceneTracks/Public/Sections/MovieSceneColorSection.h b/Engine/Source/Runtime/MovieSceneTracks/Public/Sections/MovieSceneColorSection.h
--- a/Engine/Source/Runtime/MovieSceneTracks/Public/Sections/MovieSceneColorSection.h
+++ b/Engine/Source/Runtime/MovieSceneTracks/Public/Sections/MovieSceneColorSection.h
@@ -1,10 +1,13 @@
 // Copyright 1998-2016 Epic Games, Inc. All Rights Reserved.
 
 #pragma once
-#include "MovieSceneSection.h"
+
 #include "IKeyframeSection.h"
+#include "MovieSceneKeyStruct.h"
+#include "MovieSceneSection.h"
 #include "MovieSceneColorSection.generated.h"
 
+
 enum class EKeyColorChannel
 {
 	Red,
@@ -13,9 +16,10 @@ enum class EKeyColorChannel
 	Alpha,
 };
 
+
 struct FColorKey
 {
-	FColorKey( EKeyColorChannel InChannel, float InChannelValue, bool InbIsSlateColor )
+	FColorKey(EKeyColorChannel InChannel, float InChannelValue, bool InbIsSlateColor)
 	{
 		Channel = InChannel;
 		ChannelValue = InChannelValue;
@@ -27,33 +31,44 @@ struct FColorKey
 	bool bIsSlateColor;
 };
 
+
+/**
+ * Proxy structure for color section key data.
+ */
+USTRUCT()
+struct FMovieSceneColorKeyStruct
+	: public FMovieSceneKeyStruct
+{
+	GENERATED_BODY()
+
+	/** They key's color value. */
+	UPROPERTY(EditAnywhere, Category=Key, meta=(InlineColorPicker))
+	FLinearColor Color;
+
+	FRichCurveKey* Keys[4];
+
+	virtual void PropagateChanges(const FPropertyChangedEvent& ChangeEvent) override;
+};
+
+
 /**
  * A single floating point section
  */
-UCLASS( MinimalAPI )
+UCLASS(MinimalAPI)
 class UMovieSceneColorSection 
 	: public UMovieSceneSection
 	, public IKeyframeSection<FColorKey>
 {
 	GENERATED_UCLASS_BODY()
+
 public:
-	/** MovieSceneSection interface */
-	virtual void MoveSection(float DeltaPosition, TSet<FKeyHandle>& KeyHandles) override;
-	virtual void DilateSection(float DilationFactor, float Origin, TSet<FKeyHandle>& KeyHandles) override;
-	virtual void GetKeyHandles(TSet<FKeyHandle>& KeyHandles) const override;
 
 	/**
 	 * Updates this section
 	 *
 	 * @param Position	The position in time within the movie scene
 	 */
-	virtual FLinearColor Eval( float Position, const FLinearColor& DefaultColor ) const;
-
-	// IKeyframeSection interface.
-	virtual void AddKey( float Time, const FColorKey& Key, EMovieSceneKeyInterpolation KeyInterpolation ) override;
-	virtual bool NewKeyIsNewData(float Time, const FColorKey& Key) const override;
-	virtual bool HasKeys(const FColorKey& Key) const override;
-	virtual void SetDefault(const FColorKey& Key ) override;
+	virtual FLinearColor Eval(float Position, const FLinearColor& DefaultColor) const;
 
 	/**
 	 * Gets the red color curve
@@ -70,6 +85,7 @@ public:
 	 */
 	FRichCurve& GetGreenCurve() { return GreenCurve; }
 	const FRichCurve& GetGreenCurve() const { return GreenCurve; }
+
 	/**
 	 * Gets the blue color curve
 	 *
@@ -86,6 +102,24 @@ public:
 	FRichCurve& GetAlphaCurve() { return AlphaCurve; }
 	const FRichCurve& GetAlphaCurve() const { return AlphaCurve; }
 
+public:
+
+	// UMovieSceneSection interface
+
+	virtual void MoveSection(float DeltaPosition, TSet<FKeyHandle>& KeyHandles) override;
+	virtual void DilateSection(float DilationFactor, float Origin, TSet<FKeyHandle>& KeyHandles) override;
+	virtual void GetKeyHandles(TSet<FKeyHandle>& KeyHandles) const override;
+	virtual TSharedPtr<FStructOnScope> GetKeyStruct(const TArray<FKeyHandle>& KeyHandles) override;
+
+public:
+
+	// IKeyframeSection interface
+
+	virtual void AddKey(float Time, const FColorKey& Key, EMovieSceneKeyInterpolation KeyInterpolation) override;
+	virtual bool NewKeyIsNewData(float Time, const FColorKey& Key) const override;
+	virtual bool HasKeys(const FColorKey& Key) const override;
+	virtual void SetDefault(const FColorKey& Key) override;
+
 private:
 
 	/** Red curve data */
diff --git a/Engine/Source/Runtime/MovieSceneTracks/Public/Sections/MovieSceneEventSection.h b/Engine/Source/Runtime/MovieSceneTracks/Public/Sections/MovieSceneEventSection.h
--- a/Engine/Source/Runtime/MovieSceneTracks/Public/Sections/MovieSceneEventSection.h
+++ b/Engine/Source/Runtime/MovieSceneTracks/Public/Sections/MovieSceneEventSection.h
@@ -8,42 +8,6 @@
 class ALevelScriptActor;
 
 
-/**
- * Structure for event section keys.
- */
-USTRUCT()
-struct FMovieSceneEventSectionKey
-{
-	GENERATED_USTRUCT_BODY()
-
-	/** The names of the events to be triggered. */
-	UPROPERTY(EditAnywhere, Category=EventTrackKey)
-	TArray<FName> EventNames;
-
-	/** The time at which the event should be triggered. */
-	UPROPERTY()
-	float Time;
-
-	/** Default constructor. */
-	FMovieSceneEventSectionKey()
-		: Time(0.0f)
-	{ }
-
-	/** Creates and initializes a new instance. */
-	FMovieSceneEventSectionKey(const FName& InEventName, float InTime)
-		: Time(InTime)
-	{
-		EventNames.Add(InEventName);
-	}
-
-	/** Operator less, used to sort the heap based on time until execution. */
-	bool operator<(const FMovieSceneEventSectionKey& Other) const
-	{
-		return Time < Other.Time;
-	}
-};
-
-
 /**
  * Implements a section in movie scene event tracks.
  */
@@ -105,6 +69,6 @@ protected:
 private:
 
 	/** The section's keys. */
-	UPROPERTY(EditAnywhere, Category="Events")
+	UPROPERTY()
 	FNameCurve Events;
 };
diff --git a/Engine/Source/Runtime/MovieSceneTracks/Public/Sections/MovieSceneParticleSection.h b/Engine/Source/Runtime/MovieSceneTracks/Public/Sections/MovieSceneParticleSection.h
--- a/Engine/Source/Runtime/MovieSceneTracks/Public/Sections/MovieSceneParticleSection.h
+++ b/Engine/Source/Runtime/MovieSceneTracks/Public/Sections/MovieSceneParticleSection.h
@@ -43,6 +43,6 @@ class UMovieSceneParticleSection
 
 private:
 	/** Curve containing the particle keys. */
-	UPROPERTY(EditAnywhere, Category="Particles")
+	UPROPERTY()
 	FIntegralCurve ParticleKeys;
 };
