Copying //UE4/Dev-Sequencer to //UE4/Main

diff --git a/Engine/Source/Runtime/LevelSequence/Public/LevelSequencePlayer.h b/Engine/Source/Runtime/LevelSequence/Public/LevelSequencePlayer.h
--- a/Engine/Source/Runtime/LevelSequence/Public/LevelSequencePlayer.h
+++ b/Engine/Source/Runtime/LevelSequence/Public/LevelSequencePlayer.h
@@ -138,6 +138,7 @@ protected:
 	virtual void SetViewportSettings(const TMap<FViewportClient*, EMovieSceneViewportParams>& ViewportParamsMap) override;
 	virtual void GetViewportSettings(TMap<FViewportClient*, EMovieSceneViewportParams>& ViewportParamsMap) const override;
 	virtual EMovieScenePlayerStatus::Type GetPlaybackStatus() const override;
+	virtual void SetPlaybackStatus(EMovieScenePlayerStatus::Type InPlaybackStatus) override;
 	virtual void AddOrUpdateMovieSceneInstance(UMovieSceneSection& MovieSceneSection, TSharedRef<FMovieSceneSequenceInstance> InstanceToAdd) override;
 	virtual void RemoveMovieSceneInstance(UMovieSceneSection& MovieSceneSection, TSharedRef<FMovieSceneSequenceInstance> InstanceToRemove) override;
 	virtual TSharedRef<FMovieSceneSequenceInstance> GetRootMovieSceneSequenceInstance() const override;
diff --git a/Engine/Source/Runtime/MovieScene/Private/MovieScene.cpp b/Engine/Source/Runtime/MovieScene/Private/MovieScene.cpp
--- a/Engine/Source/Runtime/MovieScene/Private/MovieScene.cpp
+++ b/Engine/Source/Runtime/MovieScene/Private/MovieScene.cpp
@@ -190,6 +190,10 @@ FMovieScenePossessable* UMovieScene::FindPossessable( const FGuid& Guid )
 	return nullptr;
 }
 
+FMovieScenePossessable* UMovieScene::FindPossessable( const TFunctionRef<bool(FMovieScenePossessable&)>& InPredicate )
+{
+	return Possessables.FindByPredicate(InPredicate);
+}
 
 int32 UMovieScene::GetPossessableCount() const
 {
@@ -258,9 +262,16 @@ void UMovieScene::SetPlaybackRange(float Start, float End)
 {
 	if (ensure(End >= Start))
 	{
+		const auto NewRange = TRange<float>(Start, TRangeBound<float>::Inclusive(End));
+
+		if (PlaybackRange == NewRange)
+		{
+			return;
+		}
+
 		Modify();
 		
-		PlaybackRange = TRange<float>(Start, TRangeBound<float>::Inclusive(End));
+		PlaybackRange = NewRange;
 
 #if WITH_EDITORONLY_DATA
 		if (EditorData.WorkingRange.IsEmpty())
@@ -338,7 +349,7 @@ class UMovieSceneTrack* UMovieScene::AddTrack( TSubclassOf<UMovieSceneTrack> Tra
 
 			CreatedType = NewObject<UMovieSceneTrack>(this, TrackClass, NAME_None, RF_Transactional);
 			ensure(CreatedType);
-
+			
 			Binding.AddTrack( *CreatedType );
 		}
 	}
@@ -351,7 +362,7 @@ bool UMovieScene::RemoveTrack(UMovieSceneTrack& Track)
 {
 	Modify();
 
-	bool bAnythingRemoved = false;
+bool bAnythingRemoved = false;
 
 	for (auto& Binding : ObjectBindings)
 	{
diff --git a/Engine/Source/Runtime/MovieScene/Private/MovieSceneCommonHelpers.cpp b/Engine/Source/Runtime/MovieScene/Private/MovieSceneCommonHelpers.cpp
--- a/Engine/Source/Runtime/MovieScene/Private/MovieSceneCommonHelpers.cpp
+++ b/Engine/Source/Runtime/MovieScene/Private/MovieSceneCommonHelpers.cpp
@@ -182,6 +182,36 @@ USceneComponent* MovieSceneHelpers::SceneComponentFromRuntimeObject(UObject* Obj
 	return SceneComponent;
 }
 
+UCameraComponent* MovieSceneHelpers::CameraComponentFromActor(const AActor* InActor)
+{
+	TArray<UCameraComponent*> CameraComponents;
+	InActor->GetComponents<UCameraComponent>(CameraComponents);
+
+	for (UCameraComponent* CameraComponent : CameraComponents)
+	{
+		if (CameraComponent->bIsActive)
+		{
+			return CameraComponent;
+		}
+	}
+
+	// now see if any actors are attached to us, directly or indirectly, that have an active camera component we might want to use
+	// we will just return the first one.
+	// #note: assumption here that attachment cannot be circular
+	TArray<AActor*> AttachedActors;
+	InActor->GetAttachedActors(AttachedActors);
+	for (AActor* AttachedActor : AttachedActors)
+	{
+		UCameraComponent* const Comp = CameraComponentFromActor(AttachedActor);
+		if (Comp)
+		{
+			return Comp;
+		}
+	}
+
+	return nullptr;
+}
+
 void MovieSceneHelpers::SetKeyInterpolation(FRichCurve& InCurve, FKeyHandle InKeyHandle, EMovieSceneKeyInterpolation InKeyInterpolation)
 {
 	switch (InKeyInterpolation)
diff --git a/Engine/Source/Runtime/MovieScene/Public/IMovieScenePlayer.h b/Engine/Source/Runtime/MovieScene/Public/IMovieScenePlayer.h
--- a/Engine/Source/Runtime/MovieScene/Public/IMovieScenePlayer.h
+++ b/Engine/Source/Runtime/MovieScene/Public/IMovieScenePlayer.h
@@ -12,6 +12,7 @@ namespace EMovieScenePlayerStatus
 		Playing,
 		Recording,
 		Scrubbing,
+		Jumping,
 		MAX
 	};
 }
@@ -106,6 +107,11 @@ public:
 	/** @return whether the player is currently playing, scrubbing, etc. */
 	virtual EMovieScenePlayerStatus::Type GetPlaybackStatus() const = 0;
 
+	/** 
+	* @param PlaybackStatus The playback status to set
+	*/
+	virtual void SetPlaybackStatus(EMovieScenePlayerStatus::Type InPlaybackStatus) = 0;
+
 	/**
 	 * Obtain an object responsible for managing movie scene spawnables
 	 */
diff --git a/Engine/Source/Runtime/MovieScene/Public/IMovieSceneTrackInstance.h b/Engine/Source/Runtime/MovieScene/Public/IMovieSceneTrackInstance.h
--- a/Engine/Source/Runtime/MovieScene/Public/IMovieSceneTrackInstance.h
+++ b/Engine/Source/Runtime/MovieScene/Public/IMovieSceneTrackInstance.h
@@ -18,6 +18,7 @@ struct EMovieSceneUpdateData
 	float Position;
 	float LastPosition;
 	bool bPreroll;
+	bool bJumpCut;
 	EMovieSceneUpdatePass UpdatePass;
 
 	EMovieSceneUpdateData()
@@ -25,6 +26,7 @@ struct EMovieSceneUpdateData
 		Position = 0.0f;
 		LastPosition = 0.0f;
 		bPreroll = false;
+		bJumpCut = false;
 		UpdatePass = MSUP_PreUpdate;
 	}
 	EMovieSceneUpdateData(float InPosition, float InLastPosition)
@@ -32,6 +34,7 @@ struct EMovieSceneUpdateData
 		Position = InPosition;
 		LastPosition = InLastPosition;
 		bPreroll = false;
+		bJumpCut = false;
 		UpdatePass = MSUP_PreUpdate;
 	}
 };
diff --git a/Engine/Source/Runtime/MovieScene/Public/MovieScene.h b/Engine/Source/Runtime/MovieScene/Public/MovieScene.h
--- a/Engine/Source/Runtime/MovieScene/Public/MovieScene.h
+++ b/Engine/Source/Runtime/MovieScene/Public/MovieScene.h
@@ -169,6 +169,14 @@ public:
 	 */
 	struct FMovieScenePossessable* FindPossessable(const FGuid& Guid);
 
+	/**
+	 * Attempt to find a possessable using some custom prdeicate
+	 *
+	 * @param InPredicate A predicate to test each possessable against
+	 * @return Possessable object that was found (or nullptr if not found).
+	 */
+	FMovieScenePossessable* FindPossessable( const TFunctionRef<bool(FMovieScenePossessable&)>& InPredicate );
+
 	/**
 	 * Grabs a reference to a specific possessable by index.
 	 *
diff --git a/Engine/Source/Runtime/MovieScene/Public/MovieSceneCommonHelpers.h b/Engine/Source/Runtime/MovieScene/Public/MovieSceneCommonHelpers.h
--- a/Engine/Source/Runtime/MovieScene/Public/MovieSceneCommonHelpers.h
+++ b/Engine/Source/Runtime/MovieScene/Public/MovieSceneCommonHelpers.h
@@ -56,6 +56,14 @@ public:
 	 */	
 	static USceneComponent* SceneComponentFromRuntimeObject(UObject* Object);
 
+	/**
+	 * Get the active camera component from the actor 
+	 *
+	 * @param InActor The actor to look for the camera component on
+	 * @return The active camera component
+	 */
+	static UCameraComponent* CameraComponentFromActor(const AActor* InActor);
+
 	/*
 	 * Set the runtime object movable
 	 *
diff --git a/Engine/Source/Runtime/MovieScene/Public/MovieSceneTrack.h b/Engine/Source/Runtime/MovieScene/Public/MovieSceneTrack.h
--- a/Engine/Source/Runtime/MovieScene/Public/MovieSceneTrack.h
+++ b/Engine/Source/Runtime/MovieScene/Public/MovieSceneTrack.h
@@ -20,6 +20,14 @@ class UMovieSceneTrack
 
 public:
 
+	UMovieSceneTrack(const FObjectInitializer& InInitializer)
+		: Super(InInitializer)
+	{
+#if WITH_EDITORONLY_DATA
+		TrackTint = FColor(127, 127, 127, 0);
+#endif
+	}
+
 	/**
 	 * Creates a runtime instance of this class.
 	 * 
@@ -111,6 +119,34 @@ public:
 	 * @return Display name text.
 	 */
 	virtual FText GetDisplayName() const PURE_VIRTUAL(UMovieSceneTrack::GetDisplayName, return FText::FromString(TEXT("Unnamed Track")););
+
+	/**
+	 * Get this track's color tint.
+	 *
+	 * @return Color Tint.
+	 */
+	const FColor& GetColorTint() const
+	{
+		return TrackTint;
+	}
+
+	/**
+	 * Set this track's color tint.
+	 *
+	 * @param InTrackTint The color to tint this track.
+	 */
+	void SetColorTint(const FColor& InTrackTint)
+	{
+		TrackTint = InTrackTint;
+	}
+
+protected:
+
+	/** This track's tint color */
+	UPROPERTY(EditAnywhere, Category=General, DisplayName=Color)
+	FColor TrackTint;
+
+public:
 #endif
 
 #if WITH_EDITOR
diff --git a/Engine/Source/Runtime/MovieSceneCapture/Private/CompositionGraphCaptureProtocol.cpp b/Engine/Source/Runtime/MovieSceneCapture/Private/CompositionGraphCaptureProtocol.cpp
--- a/Engine/Source/Runtime/MovieSceneCapture/Private/CompositionGraphCaptureProtocol.cpp
+++ b/Engine/Source/Runtime/MovieSceneCapture/Private/CompositionGraphCaptureProtocol.cpp
@@ -5,6 +5,8 @@
 #include "SceneViewExtension.h"
 #include "BufferVisualizationData.h"
 #include "SceneViewport.h"
+#include "MovieSceneCaptureSettings.h"
+#include "MovieSceneCaptureModule.h"
 
 struct FSceneViewExtension : ISceneViewExtension
 {
@@ -70,9 +72,9 @@ struct FSceneViewExtension : ISceneViewExtension
 				: FinalPostProcessSettings(InFinalPostProcessSettings), RenderPasses(InRenderPasses)
 			{}
 
-			void ProcessValue(const FString& InName, UMaterial* Material, const FText&)
+			void ProcessValue(const FString& InName, UMaterial* Material, const FText& InText)
 			{
-				if (!RenderPasses.Num() || RenderPasses.Contains(InName))
+				if (!RenderPasses.Num() || RenderPasses.Contains(InName) || RenderPasses.Contains(InText.ToString()))
 				{
 					FinalPostProcessSettings.BufferVisualizationOverviewMaterials.Add(Material);
 				}
@@ -114,6 +116,48 @@ private:
 	int32 RestoreDumpHDR;
 };
 
+void UCompositionGraphCaptureSettings::OnReleaseConfig(FMovieSceneCaptureSettings& InSettings)
+{
+	// Remove {material} if it exists
+	InSettings.OutputFormat = InSettings.OutputFormat.Replace(TEXT("{material}"), TEXT(""));
+
+	// Remove .{frame} if it exists
+	InSettings.OutputFormat = InSettings.OutputFormat.Replace(TEXT(".{frame}"), TEXT(""));
+
+	Super::OnReleaseConfig(InSettings);
+}
+
+void UCompositionGraphCaptureSettings::OnLoadConfig(FMovieSceneCaptureSettings& InSettings)
+{
+	// Add .{frame} if it doesn't already exist
+	FString OutputFormat = InSettings.OutputFormat;
+
+	if (!OutputFormat.Contains(TEXT("{frame}")))
+	{
+		OutputFormat.Append(TEXT(".{frame}"));
+
+		InSettings.OutputFormat = OutputFormat;
+	}
+
+	// Add {material} if it doesn't already exist
+	if (!OutputFormat.Contains(TEXT("{material}")))
+	{
+		int32 FramePosition = OutputFormat.Find(TEXT(".{frame}"));
+		if (FramePosition != INDEX_NONE)
+		{
+			OutputFormat.InsertAt(FramePosition, TEXT("{material}"));
+		}
+		else
+		{ 
+			OutputFormat.Append(TEXT("{material}"));
+		}
+
+		InSettings.OutputFormat = OutputFormat;
+	}
+
+	Super::OnLoadConfig(InSettings);
+}
+
 bool FCompositionGraphCaptureProtocol::Initialize(const FCaptureProtocolInitSettings& InSettings, const ICaptureProtocolHost& Host)
 {
 	SceneViewport = InSettings.SceneViewport;
diff --git a/Engine/Source/Runtime/MovieSceneCapture/Private/ImageSequenceProtocol.cpp b/Engine/Source/Runtime/MovieSceneCapture/Private/ImageSequenceProtocol.cpp
--- a/Engine/Source/Runtime/MovieSceneCapture/Private/ImageSequenceProtocol.cpp
+++ b/Engine/Source/Runtime/MovieSceneCapture/Private/ImageSequenceProtocol.cpp
@@ -4,6 +4,52 @@
 
 #include "ImageSequenceProtocol.h"
 
+void UBmpImageCaptureSettings::OnReleaseConfig(FMovieSceneCaptureSettings& InSettings)
+{
+	// Remove .{frame} if it exists. The "." before the {frame} is intentional because some media players denote frame numbers separated by "."
+	InSettings.OutputFormat = InSettings.OutputFormat.Replace(TEXT(".{frame}"), TEXT(""));
+
+	Super::OnReleaseConfig(InSettings);
+}
+
+void UBmpImageCaptureSettings::OnLoadConfig(FMovieSceneCaptureSettings& InSettings)
+{
+	// Add .{frame} if it doesn't already exist
+	FString OutputFormat = InSettings.OutputFormat;
+
+	if (!OutputFormat.Contains(TEXT("{frame}")))
+	{
+		OutputFormat.Append(TEXT(".{frame}"));
+
+		InSettings.OutputFormat = OutputFormat;
+	}
+
+	Super::OnLoadConfig(InSettings);
+}
+
+void UImageCaptureSettings::OnReleaseConfig(FMovieSceneCaptureSettings& InSettings)
+{
+	// Remove .{frame} if it exists. The "." before the {frame} is intentional because some media players denote frame numbers separated by "."
+	InSettings.OutputFormat = InSettings.OutputFormat.Replace(TEXT(".{frame}"), TEXT(""));
+
+	Super::OnReleaseConfig(InSettings);
+}
+
+void UImageCaptureSettings::OnLoadConfig(FMovieSceneCaptureSettings& InSettings)
+{
+	// Add .{frame} if it doesn't already exist
+	FString OutputFormat = InSettings.OutputFormat;
+
+	if (!OutputFormat.Contains(TEXT("{frame}")))
+	{
+		OutputFormat.Append(TEXT(".{frame}"));
+
+		InSettings.OutputFormat = OutputFormat;
+	}
+
+	Super::OnLoadConfig(InSettings);
+}
+
 #if WITH_EDITOR
 #include "IImageWrapperModule.h"
 
diff --git a/Engine/Source/Runtime/MovieSceneCapture/Private/MovieSceneCapture.cpp b/Engine/Source/Runtime/MovieSceneCapture/Private/MovieSceneCapture.cpp
--- a/Engine/Source/Runtime/MovieSceneCapture/Private/MovieSceneCapture.cpp
+++ b/Engine/Source/Runtime/MovieSceneCapture/Private/MovieSceneCapture.cpp
@@ -34,7 +34,7 @@ FMovieSceneCaptureSettings::FMovieSceneCaptureSettings()
 	bCreateTemporaryCopiesOfLevels = false;
 	bUseRelativeFrameNumbers = false;
 	GameModeOverride = nullptr;
-	OutputFormat = TEXT("{world}_{frame}");
+	OutputFormat = TEXT("{world}");
 	FrameRate = 24;
 	bEnableTextureStreaming = true;
 	bCinematicMode = true;
@@ -70,10 +70,15 @@ void UMovieSceneCapture::PostInitProperties()
 {
 	if (!HasAnyFlags(RF_ClassDefaultObject))
 	{
+		if (ProtocolSettings)
+		{
+			ProtocolSettings->OnReleaseConfig(Settings);
+		}
 		ProtocolSettings = IMovieSceneCaptureModule::Get().GetProtocolRegistry().FactorySettingsType(CaptureType, this);
 		if (ProtocolSettings)
 		{
 			ProtocolSettings->LoadConfig();
+			ProtocolSettings->OnLoadConfig(Settings);
 		}
 	}
 
@@ -304,10 +309,16 @@ void UMovieSceneCapture::PostEditChangeProperty( struct FPropertyChangedEvent& P
 
 	if (PropertyName == GET_MEMBER_NAME_CHECKED(UMovieSceneCapture, CaptureType))
 	{
+		if (ProtocolSettings)
+		{
+			ProtocolSettings->OnReleaseConfig(Settings);
+		}
+
 		ProtocolSettings = IMovieSceneCaptureModule::Get().GetProtocolRegistry().FactorySettingsType(CaptureType, this);
 		if (ProtocolSettings)
 		{
 			ProtocolSettings->LoadConfig();
+			ProtocolSettings->OnLoadConfig(Settings);
 		}
 	}
 
diff --git a/Engine/Source/Runtime/MovieSceneCapture/Private/MovieSceneCaptureModule.cpp b/Engine/Source/Runtime/MovieSceneCapture/Private/MovieSceneCaptureModule.cpp
--- a/Engine/Source/Runtime/MovieSceneCapture/Private/MovieSceneCaptureModule.cpp
+++ b/Engine/Source/Runtime/MovieSceneCapture/Private/MovieSceneCaptureModule.cpp
@@ -67,7 +67,7 @@ private:
 		}
 		{
 			Info.DisplayName = LOCTEXT("BMPDescription", "Image Sequence (bmp)");
-			Info.SettingsClassType = nullptr;	// Bitmaps don't have any options
+			Info.SettingsClassType = UBmpImageCaptureSettings::StaticClass();
 			Info.Factory = []() -> TSharedRef<IMovieSceneCaptureProtocol> {
 				return MakeShareable(new FImageSequenceProtocol(EImageFormat::BMP));
 			};
@@ -169,7 +169,7 @@ private:
 						UClass* ProtocolTypeClass = FindObject<UClass>(nullptr, *ProtocolTypeField->AsString());
 						if (ProtocolTypeClass)
 						{
-							Capture->ProtocolSettings = NewObject<UMovieSceneCapture>(Capture, ProtocolTypeClass);
+							Capture->ProtocolSettings = NewObject<UMovieSceneCaptureProtocolSettings>(Capture, ProtocolTypeClass);
 							if (Capture->ProtocolSettings)
 							{
 								auto ProtocolDataField = RootObject->TryGetField(TEXT("ProtocolData"));
diff --git a/Engine/Source/Runtime/MovieSceneCapture/Public/MovieSceneCapture.h b/Engine/Source/Runtime/MovieSceneCapture/Public/MovieSceneCapture.h
--- a/Engine/Source/Runtime/MovieSceneCapture/Public/MovieSceneCapture.h
+++ b/Engine/Source/Runtime/MovieSceneCapture/Public/MovieSceneCapture.h
@@ -9,6 +9,7 @@
 #include "RenderingThread.h"
 #include "IMovieSceneCaptureProtocol.h"
 #include "MovieSceneCaptureProtocolRegistry.h"
+#include "MovieSceneCaptureProtocolSettings.h"
 #include "MovieSceneCapture.generated.h"
 
 /** Structure used to cache various metrics for our capture */
@@ -53,7 +54,7 @@ public:
 
 	/** Settings specific to the capture protocol */
 	UPROPERTY(EditAnywhere, Category=CaptureSettings)
-	UObject* ProtocolSettings;
+	UMovieSceneCaptureProtocolSettings* ProtocolSettings;
 
 	/** Settings that define how to capture */
 	UPROPERTY(EditAnywhere, config, Category=CaptureSettings, meta=(ShowOnlyInnerProperties))
diff --git a/Engine/Source/Runtime/MovieSceneCapture/Public/MovieSceneCaptureProtocolRegistry.h b/Engine/Source/Runtime/MovieSceneCapture/Public/MovieSceneCaptureProtocolRegistry.h
--- a/Engine/Source/Runtime/MovieSceneCapture/Public/MovieSceneCaptureProtocolRegistry.h
+++ b/Engine/Source/Runtime/MovieSceneCapture/Public/MovieSceneCaptureProtocolRegistry.h
@@ -2,6 +2,7 @@
 
 #pragma once
 
+#include "MovieSceneCaptureProtocolSettings.h"
 #include "MovieSceneCaptureProtocolRegistry.generated.h"
 
 struct IMovieSceneCaptureProtocol;
@@ -55,13 +56,13 @@ public:
 	}
 
 	/** Create a new settings type for the specified protocol ID */
-	UObject* FactorySettingsType(FCaptureProtocolID InProtocolID, UObject* Outer)
+	UMovieSceneCaptureProtocolSettings* FactorySettingsType(FCaptureProtocolID InProtocolID, UObject* Outer)
 	{
 		if (const FMovieSceneCaptureProtocolInfo* Info = Register.Find(InProtocolID))
 		{
 			if (Info->SettingsClassType)
 			{
-				return NewObject<UObject>(Outer, Info->SettingsClassType);
+				return NewObject<UMovieSceneCaptureProtocolSettings>(Outer, Info->SettingsClassType);
 			}
 		}
 		return nullptr;
diff --git a/Engine/Source/Runtime/MovieSceneCapture/Public/MovieSceneCaptureProtocolSettings.h b/Engine/Source/Runtime/MovieSceneCapture/Public/MovieSceneCaptureProtocolSettings.h
new file mode 100644
--- /dev/null
+++ b/Engine/Source/Runtime/MovieSceneCapture/Public/MovieSceneCaptureProtocolSettings.h
@@ -0,0 +1,25 @@
+// Copyright 1998-2016 Epic Games, Inc. All Rights Reserved.
+
+#pragma once
+
+#include "MovieSceneCaptureSettings.h"
+#include "MovieSceneCaptureProtocolSettings.generated.h"
+
+UCLASS()
+class MOVIESCENECAPTURE_API UMovieSceneCaptureProtocolSettings : public UObject
+{
+public:
+
+	GENERATED_BODY()
+
+	/**
+	 * Called when this protocol has been released
+	 */
+	virtual void OnReleaseConfig(FMovieSceneCaptureSettings& InSettings) {}
+
+	/**
+	 * Called when this protocol has been loaded
+	 */
+	virtual void OnLoadConfig(FMovieSceneCaptureSettings& InSettings) {}
+};
+
diff --git a/Engine/Source/Runtime/MovieSceneCapture/Public/Protocols/CompositionGraphCaptureProtocol.h b/Engine/Source/Runtime/MovieSceneCapture/Public/Protocols/CompositionGraphCaptureProtocol.h
--- a/Engine/Source/Runtime/MovieSceneCapture/Public/Protocols/CompositionGraphCaptureProtocol.h
+++ b/Engine/Source/Runtime/MovieSceneCapture/Public/Protocols/CompositionGraphCaptureProtocol.h
@@ -3,6 +3,7 @@
 #pragma once
 
 #include "IMovieSceneCaptureProtocol.h"
+#include "MovieSceneCaptureProtocolSettings.h"
 #include "CompositionGraphCaptureProtocol.generated.h"
 
 USTRUCT()
@@ -15,11 +16,15 @@ struct MOVIESCENECAPTURE_API FCompositionGraphCapturePasses
 };
 
 UCLASS(config=EditorPerProjectUserSettings, DisplayName="Composition Graph Options")
-class MOVIESCENECAPTURE_API UCompositionGraphCaptureSettings : public UObject
+class MOVIESCENECAPTURE_API UCompositionGraphCaptureSettings : public UMovieSceneCaptureProtocolSettings
 {
 public:
 	GENERATED_BODY()
 	
+	/**~ UMovieSceneCaptureProtocolSettings implementation */
+	virtual void OnReleaseConfig(FMovieSceneCaptureSettings& InSettings) override;
+	virtual void OnLoadConfig(FMovieSceneCaptureSettings& InSettings) override;
+
 	/** A list of render passes to include in the capture. Leave empty to export all available passes. */
 	UPROPERTY(config, EditAnywhere, Category="Composition Graph Options")
 	FCompositionGraphCapturePasses IncludeRenderPasses;
diff --git a/Engine/Source/Runtime/MovieSceneCapture/Public/Protocols/FrameGrabberProtocol.h b/Engine/Source/Runtime/MovieSceneCapture/Public/Protocols/FrameGrabberProtocol.h
--- a/Engine/Source/Runtime/MovieSceneCapture/Public/Protocols/FrameGrabberProtocol.h
+++ b/Engine/Source/Runtime/MovieSceneCapture/Public/Protocols/FrameGrabberProtocol.h
@@ -4,11 +4,12 @@
 
 #include "FrameGrabber.h"
 #include "IMovieSceneCaptureProtocol.h"
+#include "MovieSceneCaptureProtocolSettings.h"
 
 #include "FrameGrabberProtocol.generated.h"
 
 UCLASS()
-class MOVIESCENECAPTURE_API UFrameGrabberProtocolSettings : public UObject
+class MOVIESCENECAPTURE_API UFrameGrabberProtocolSettings : public UMovieSceneCaptureProtocolSettings
 {
 public:
 	UFrameGrabberProtocolSettings(const FObjectInitializer&) : DesiredPixelFormat(PF_B8G8R8A8), RingBufferSize(3) {}
diff --git a/Engine/Source/Runtime/MovieSceneCapture/Public/Protocols/ImageSequenceProtocol.h b/Engine/Source/Runtime/MovieSceneCapture/Public/Protocols/ImageSequenceProtocol.h
--- a/Engine/Source/Runtime/MovieSceneCapture/Public/Protocols/ImageSequenceProtocol.h
+++ b/Engine/Source/Runtime/MovieSceneCapture/Public/Protocols/ImageSequenceProtocol.h
@@ -7,8 +7,22 @@
 #endif
 
 #include "FrameGrabberProtocol.h"
+#include "MovieSceneCaptureProtocolSettings.h"
 #include "ImageSequenceProtocol.generated.h"
 
+UCLASS(config=EditorPerProjectUserSettings, DisplayName="Image Encoding")
+class MOVIESCENECAPTURE_API UBmpImageCaptureSettings : public UMovieSceneCaptureProtocolSettings
+{
+public:
+	UBmpImageCaptureSettings(const FObjectInitializer& Init) : UMovieSceneCaptureProtocolSettings(Init) {}
+
+	GENERATED_BODY()
+
+	/**~ UMovieSceneCaptureProtocolSettings implementation */
+	virtual void OnReleaseConfig(FMovieSceneCaptureSettings& InSettings) override;
+	virtual void OnLoadConfig(FMovieSceneCaptureSettings& InSettings) override;
+};
+
 UCLASS(config=EditorPerProjectUserSettings, DisplayName="Image Encoding")
 class MOVIESCENECAPTURE_API UImageCaptureSettings : public UFrameGrabberProtocolSettings
 {
@@ -17,6 +31,10 @@ public:
 
 	GENERATED_BODY()
 
+	/**~ UMovieSceneCaptureProtocolSettings implementation */
+	virtual void OnReleaseConfig(FMovieSceneCaptureSettings& InSettings) override;
+	virtual void OnLoadConfig(FMovieSceneCaptureSettings& InSettings) override;
+
 	/** Level of compression to apply to the image, between 1 (worst quality, best compression) and 100 (best quality, worst compression)*/
 	UPROPERTY(config, EditAnywhere, Category=ImageSettings, meta=(ClampMin=1, ClampMax=100))
 	int32 CompressionQuality;
diff --git a/Engine/Source/Runtime/MovieSceneTracks/Private/Sections/MovieScene3DAttachSection.cpp b/Engine/Source/Runtime/MovieSceneTracks/Private/Sections/MovieScene3DAttachSection.cpp
--- a/Engine/Source/Runtime/MovieSceneTracks/Private/Sections/MovieScene3DAttachSection.cpp
+++ b/Engine/Source/Runtime/MovieSceneTracks/Private/Sections/MovieScene3DAttachSection.cpp
@@ -9,6 +9,7 @@ UMovieScene3DAttachSection::UMovieScene3DAttachSection( const FObjectInitializer
 	: Super( ObjectInitializer )
 {
 	AttachSocketName = NAME_None;
+	AttachComponentName = NAME_None;
 	bConstrainTx = true;
 	bConstrainTy = true;
 	bConstrainTz = true;
@@ -20,13 +21,49 @@ UMovieScene3DAttachSection::UMovieScene3DAttachSection( const FObjectInitializer
 
 void UMovieScene3DAttachSection::Eval( USceneComponent* SceneComponent, float Position, AActor* Actor, FVector& OutTranslation, FRotator& OutRotation ) const
 {
-	if (Actor->GetRootComponent()->DoesSocketExist(AttachSocketName))
+	bool bFoundAttachment = false;
+	if(AttachSocketName != NAME_None)
 	{
-		FTransform SocketTransform = Actor->GetRootComponent()->GetSocketTransform(AttachSocketName);
-		OutTranslation = SocketTransform.GetLocation();
-		OutRotation = SocketTransform.GetRotation().Rotator();
+		if(AttachComponentName != NAME_None )
+		{
+			TInlineComponentArray<USceneComponent*> PotentialAttachComponents(Actor);
+			for(USceneComponent* PotentialAttachComponent : PotentialAttachComponents)
+			{
+				if(PotentialAttachComponent->GetFName() == AttachComponentName && PotentialAttachComponent->DoesSocketExist(AttachSocketName))
+				{
+					FTransform SocketTransform = PotentialAttachComponent->GetSocketTransform(AttachSocketName);
+					OutTranslation = SocketTransform.GetLocation();
+					OutRotation = SocketTransform.GetRotation().Rotator();
+					bFoundAttachment = true;
+					break;
+				}
+			}
+		}
+		else if (Actor->GetRootComponent()->DoesSocketExist(AttachSocketName))
+		{
+			FTransform SocketTransform = Actor->GetRootComponent()->GetSocketTransform(AttachSocketName);
+			OutTranslation = SocketTransform.GetLocation();
+			OutRotation = SocketTransform.GetRotation().Rotator();
+			bFoundAttachment = true;
+		}
 	}
-	else
+	else if(AttachComponentName != NAME_None )
+	{
+		TInlineComponentArray<USceneComponent*> PotentialAttachComponents(Actor);
+		for(USceneComponent* PotentialAttachComponent : PotentialAttachComponents)
+		{
+			if(PotentialAttachComponent->GetFName() == AttachComponentName)
+			{
+				FTransform ComponentTransform = PotentialAttachComponent->GetComponentToWorld();
+				OutTranslation = ComponentTransform.GetLocation();
+				OutRotation = ComponentTransform.GetRotation().Rotator();
+				bFoundAttachment = true;
+				break;
+			}
+		}
+	}
+
+	if(!bFoundAttachment)
 	{
 		OutTranslation = Actor->GetActorLocation();
 		OutRotation = Actor->GetActorRotation();
diff --git a/Engine/Source/Runtime/MovieSceneTracks/Private/Sections/MovieScene3DTransformSection.cpp b/Engine/Source/Runtime/MovieSceneTracks/Private/Sections/MovieScene3DTransformSection.cpp
--- a/Engine/Source/Runtime/MovieSceneTracks/Private/Sections/MovieScene3DTransformSection.cpp
+++ b/Engine/Source/Runtime/MovieSceneTracks/Private/Sections/MovieScene3DTransformSection.cpp
@@ -5,6 +5,41 @@
 #include "MovieScene3DTransformTrack.h"
 
 
+/* FMovieScene3DLocationKeyStruct interface
+ *****************************************************************************/
+
+void FMovieScene3DLocationKeyStruct::PropagateChanges(const FPropertyChangedEvent& ChangeEvent)
+{
+	for (int32 Index = 0; Index <= 2; ++Index)
+	{
+		LocationKeys[Index]->Value = Location[Index];
+	}
+}
+
+
+/* FMovieScene3DRotationKeyStruct interface
+ *****************************************************************************/
+
+void FMovieScene3DRotationKeyStruct::PropagateChanges(const FPropertyChangedEvent& ChangeEvent)
+{
+	RotationKeys[0]->Value = Rotation.Pitch;
+	RotationKeys[1]->Value = Rotation.Yaw;
+	RotationKeys[2]->Value = Rotation.Roll;
+}
+
+
+/* FMovieScene3DScaleKeyStruct interface
+ *****************************************************************************/
+
+void FMovieScene3DScaleKeyStruct::PropagateChanges(const FPropertyChangedEvent& ChangeEvent)
+{
+	for (int32 Index = 0; Index <= 2; ++Index)
+	{
+		ScaleKeys[Index]->Value = Scale[Index];
+	}
+}
+
+
 /* FMovieScene3DTransformKeyStruct interface
  *****************************************************************************/
 
@@ -12,7 +47,7 @@ void FMovieScene3DTransformKeyStruct::PropagateChanges(const FPropertyChangedEve
 {
 	for (int32 Index = 0; Index <= 2; ++Index)
 	{
-		TranslationKeys[Index]->Value = Translation[Index];
+		LocationKeys[Index]->Value = Location[Index];
 		ScaleKeys[Index]->Value = Scale[Index];
 	}
 
@@ -163,24 +198,96 @@ void UMovieScene3DTransformSection::GetKeyHandles(TSet<FKeyHandle>& KeyHandles)
 
 TSharedPtr<FStructOnScope> UMovieScene3DTransformSection::GetKeyStruct(const TArray<FKeyHandle>& KeyHandles)
 {
-	TSharedRef<FStructOnScope> KeyStruct = MakeShareable(new FStructOnScope(FMovieScene3DTransformKeyStruct::StaticStruct()));
-	auto Struct = (FMovieScene3DTransformKeyStruct*)KeyStruct->GetStructMemory();
+	FRichCurveKey* TranslationKeys[3];
+	FRichCurveKey* RotationKeys[3];
+	FRichCurveKey* ScaleKeys[3];
+
+	for (int32 Index = 0; Index <= 2; ++Index)
 	{
-		for (int32 Index = 0; Index <= 2; ++Index)
+		TranslationKeys[Index] = Translation[Index].GetFirstMatchingKey(KeyHandles);
+		RotationKeys[Index] = Rotation[Index].GetFirstMatchingKey(KeyHandles);
+		ScaleKeys[Index] = Scale[Index].GetFirstMatchingKey(KeyHandles);
+	}
+
+	if ((TranslationKeys[0] != nullptr) && (RotationKeys[0] != nullptr) && (ScaleKeys[0] != nullptr))
+	{
+		TSharedRef<FStructOnScope> KeyStruct = MakeShareable(new FStructOnScope(FMovieScene3DTransformKeyStruct::StaticStruct()));
+		auto Struct = (FMovieScene3DTransformKeyStruct*)KeyStruct->GetStructMemory();
+		{
+			for (int32 Index = 0; Index <= 2; ++Index)
+			{
+				check(TranslationKeys[Index] != nullptr);
+				Struct->LocationKeys[Index] = TranslationKeys[Index];
+				Struct->Location[Index] = TranslationKeys[Index]->Value;
+
+				check(RotationKeys[Index] != nullptr);
+				Struct->RotationKeys[Index] = RotationKeys[Index];
+
+				check(ScaleKeys[Index] != nullptr);
+				Struct->ScaleKeys[Index] = ScaleKeys[Index];
+				Struct->Scale[Index] = ScaleKeys[Index]->Value;
+			}
+
+			Struct->Rotation.Pitch = RotationKeys[0]->Value;
+			Struct->Rotation.Yaw = RotationKeys[1]->Value;
+			Struct->Rotation.Roll = RotationKeys[2]->Value;
+		}
+
+		return KeyStruct;
+	}
+	
+	if (TranslationKeys[0] != nullptr)
+	{
+		TSharedRef<FStructOnScope> KeyStruct = MakeShareable(new FStructOnScope(FMovieScene3DLocationKeyStruct::StaticStruct()));
+		auto Struct = (FMovieScene3DLocationKeyStruct*)KeyStruct->GetStructMemory();
+		{
+			for (int32 Index = 0; Index <= 2; ++Index)
+			{
+				check(TranslationKeys[Index] != nullptr);
+				Struct->LocationKeys[Index] = TranslationKeys[Index];
+				Struct->Location[Index] = TranslationKeys[Index]->Value;
+			}
+		}
+
+		return KeyStruct;
+	}
+	
+	if (RotationKeys[0] != nullptr)
+	{
+		TSharedRef<FStructOnScope> KeyStruct = MakeShareable(new FStructOnScope(FMovieScene3DRotationKeyStruct::StaticStruct()));
+		auto Struct = (FMovieScene3DRotationKeyStruct*)KeyStruct->GetStructMemory();
 		{
-			Struct->TranslationKeys[Index] = &Translation[Index].GetFirstMatchingKey(KeyHandles);
-			Struct->Translation[Index] = Struct->TranslationKeys[Index]->Value;
-			Struct->RotationKeys[Index] = &Rotation[Index].GetFirstMatchingKey(KeyHandles);
-			Struct->ScaleKeys[Index] = &Scale[Index].GetFirstMatchingKey(KeyHandles);
-			Struct->Scale[Index] = Struct->ScaleKeys[Index]->Value;
+			for (int32 Index = 0; Index <= 2; ++Index)
+			{
+				check(RotationKeys[Index] != nullptr);
+				Struct->RotationKeys[Index] = RotationKeys[Index];
+			}
+
+			Struct->Rotation.Pitch = Struct->RotationKeys[0]->Value;
+			Struct->Rotation.Yaw = Struct->RotationKeys[1]->Value;
+			Struct->Rotation.Roll = Struct->RotationKeys[2]->Value;
 		}
 
-		Struct->Rotation.Pitch = Struct->RotationKeys[0]->Value;
-		Struct->Rotation.Yaw = Struct->RotationKeys[1]->Value;
-		Struct->Rotation.Roll = Struct->RotationKeys[2]->Value;
+		return KeyStruct;
 	}
+	
+	if (ScaleKeys[0] != nullptr)
+	{
+		TSharedRef<FStructOnScope> KeyStruct = MakeShareable(new FStructOnScope(FMovieScene3DScaleKeyStruct::StaticStruct()));
+		auto Struct = (FMovieScene3DScaleKeyStruct*)KeyStruct->GetStructMemory();
+		{
+			for (int32 Index = 0; Index <= 2; ++Index)
+			{
+				check(ScaleKeys[Index] != nullptr);
+				Struct->ScaleKeys[Index] = ScaleKeys[Index];
+				Struct->Scale[Index] = ScaleKeys[Index]->Value;
+			}
+		}
 
-	return KeyStruct;
+		return KeyStruct;
+	}
+
+	return nullptr;
 }
 
 
diff --git a/Engine/Source/Runtime/MovieSceneTracks/Private/Sections/MovieSceneAudioSection.cpp b/Engine/Source/Runtime/MovieSceneTracks/Private/Sections/MovieSceneAudioSection.cpp
--- a/Engine/Source/Runtime/MovieSceneTracks/Private/Sections/MovieSceneAudioSection.cpp
+++ b/Engine/Source/Runtime/MovieSceneTracks/Private/Sections/MovieSceneAudioSection.cpp
@@ -12,6 +12,9 @@ UMovieSceneAudioSection::UMovieSceneAudioSection( const FObjectInitializer& Obje
 	Sound = nullptr;
 	AudioStartTime = 0.f;
 	AudioDilationFactor = 1.f;
+#if WITH_EDITORONLY_DATA
+	bShowIntensity = false;
+#endif
 }
 
 
diff --git a/Engine/Source/Runtime/MovieSceneTracks/Private/Sections/MovieSceneColorSection.cpp b/Engine/Source/Runtime/MovieSceneTracks/Private/Sections/MovieSceneColorSection.cpp
--- a/Engine/Source/Runtime/MovieSceneTracks/Private/Sections/MovieSceneColorSection.cpp
+++ b/Engine/Source/Runtime/MovieSceneTracks/Private/Sections/MovieSceneColorSection.cpp
@@ -107,13 +107,14 @@ TSharedPtr<FStructOnScope> UMovieSceneColorSection::GetKeyStruct(const TArray<FK
 	TSharedRef<FStructOnScope> KeyStruct = MakeShareable(new FStructOnScope(FMovieSceneColorKeyStruct::StaticStruct()));
 	auto Struct = (FMovieSceneColorKeyStruct*)KeyStruct->GetStructMemory();
 	{
-		Struct->Keys[0] = &RedCurve.GetFirstMatchingKey(KeyHandles);
-		Struct->Keys[1] = &GreenCurve.GetFirstMatchingKey(KeyHandles);
-		Struct->Keys[2] = &BlueCurve.GetFirstMatchingKey(KeyHandles);
-		Struct->Keys[3] = &AlphaCurve.GetFirstMatchingKey(KeyHandles);
+		Struct->Keys[0] = RedCurve.GetFirstMatchingKey(KeyHandles);
+		Struct->Keys[1] = GreenCurve.GetFirstMatchingKey(KeyHandles);
+		Struct->Keys[2] = BlueCurve.GetFirstMatchingKey(KeyHandles);
+		Struct->Keys[3] = AlphaCurve.GetFirstMatchingKey(KeyHandles);
 
 		for (int32 Index = 0; Index <= 3; ++Index)
 		{
+			check(Struct->Keys[Index] != nullptr);
 			Struct->Color.Component(Index) = Struct->Keys[Index]->Value;
 		}
 	}
diff --git a/Engine/Source/Runtime/MovieSceneTracks/Private/Sections/MovieSceneLevelVisibilitySection.cpp b/Engine/Source/Runtime/MovieSceneTracks/Private/Sections/MovieSceneLevelVisibilitySection.cpp
new file mode 100644
--- /dev/null
+++ b/Engine/Source/Runtime/MovieSceneTracks/Private/Sections/MovieSceneLevelVisibilitySection.cpp
@@ -0,0 +1,29 @@
+// Copyright 1998-2015 Epic Games, Inc. All Rights Reserved.
+
+#include "MovieSceneTracksPrivatePCH.h"
+#include "MovieSceneLevelVisibilitySection.h"
+
+
+UMovieSceneLevelVisibilitySection::UMovieSceneLevelVisibilitySection( const FObjectInitializer& ObjectInitializer )
+	: Super( ObjectInitializer )
+{
+	Visibility = ELevelVisibility::Visible;
+}
+
+
+ELevelVisibility UMovieSceneLevelVisibilitySection::GetVisibility() const
+{
+	return Visibility;
+}
+
+
+void UMovieSceneLevelVisibilitySection::SetVisibility( ELevelVisibility InVisibility )
+{
+	Visibility = InVisibility;
+}
+
+
+TArray<FName>* UMovieSceneLevelVisibilitySection::GetLevelNames()
+{
+	return &LevelNames;
+}
diff --git a/Engine/Source/Runtime/MovieSceneTracks/Private/Sections/MovieSceneParameterSection.cpp b/Engine/Source/Runtime/MovieSceneTracks/Private/Sections/MovieSceneParameterSection.cpp
--- a/Engine/Source/Runtime/MovieSceneTracks/Private/Sections/MovieSceneParameterSection.cpp
+++ b/Engine/Source/Runtime/MovieSceneTracks/Private/Sections/MovieSceneParameterSection.cpp
@@ -210,4 +210,100 @@ void UMovieSceneParameterSection::UpdateParameterIndicesFromRemoval( int32 Remov
 			ColorParameterAndCurves.Index--;
 		}
 	}
-}
\ No newline at end of file
+}
+
+void UMovieSceneParameterSection::GatherCurves(TArray<const FRichCurve*>& OutCurves) const
+{
+	for ( const FScalarParameterNameAndCurve& ScalarParameterNameAndCurve : ScalarParameterNamesAndCurves )
+	{
+		OutCurves.Add(&ScalarParameterNameAndCurve.ParameterCurve);
+	}
+
+	for ( const FVectorParameterNameAndCurves& VectorParameterNameAndCurves : VectorParameterNamesAndCurves )
+	{
+		OutCurves.Add(&VectorParameterNameAndCurves.XCurve);
+		OutCurves.Add(&VectorParameterNameAndCurves.YCurve);
+		OutCurves.Add(&VectorParameterNameAndCurves.ZCurve);
+	}
+
+	for ( const FColorParameterNameAndCurves& ColorParameterNameAndCurve : ColorParameterNamesAndCurves )
+	{
+		OutCurves.Add(&ColorParameterNameAndCurve.RedCurve);
+		OutCurves.Add(&ColorParameterNameAndCurve.BlueCurve);
+		OutCurves.Add(&ColorParameterNameAndCurve.GreenCurve);
+		OutCurves.Add(&ColorParameterNameAndCurve.AlphaCurve);
+	}
+}
+
+void UMovieSceneParameterSection::GatherCurves(TArray<FRichCurve*>& OutCurves)
+{
+	for ( FScalarParameterNameAndCurve& ScalarParameterNameAndCurve : ScalarParameterNamesAndCurves )
+	{
+		OutCurves.Add(&ScalarParameterNameAndCurve.ParameterCurve);
+	}
+
+	for ( FVectorParameterNameAndCurves& VectorParameterNameAndCurves : VectorParameterNamesAndCurves )
+	{
+		OutCurves.Add(&VectorParameterNameAndCurves.XCurve);
+		OutCurves.Add(&VectorParameterNameAndCurves.YCurve);
+		OutCurves.Add(&VectorParameterNameAndCurves.ZCurve);
+	}
+
+	for ( FColorParameterNameAndCurves& ColorParameterNameAndCurve : ColorParameterNamesAndCurves )
+	{
+		OutCurves.Add(&ColorParameterNameAndCurve.RedCurve);
+		OutCurves.Add(&ColorParameterNameAndCurve.BlueCurve);
+		OutCurves.Add(&ColorParameterNameAndCurve.GreenCurve);
+		OutCurves.Add(&ColorParameterNameAndCurve.AlphaCurve);
+	}
+}
+
+
+
+/* UMovieSceneSection overrides
+ *****************************************************************************/
+
+void UMovieSceneParameterSection::DilateSection(float DilationFactor, float Origin, TSet<FKeyHandle>& KeyHandles)
+{
+	Super::DilateSection(DilationFactor, Origin, KeyHandles);
+
+	TArray<FRichCurve*> AllCurves;
+	GatherCurves(AllCurves);
+
+	for (auto Curve : AllCurves)
+	{
+		Curve->ScaleCurve(Origin, DilationFactor, KeyHandles);
+	}
+}
+
+void UMovieSceneParameterSection::GetKeyHandles(TSet<FKeyHandle>& KeyHandles) const
+{
+	TArray<const FRichCurve*> AllCurves;
+	GatherCurves(AllCurves);
+
+	for (auto Curve : AllCurves)
+	{
+		for (auto It(Curve->GetKeyHandleIterator()); It; ++It)
+		{
+			float Time = Curve->GetKeyTime(It.Key());
+			if (IsTimeWithinSection(Time))
+			{
+				KeyHandles.Add(It.Key());
+			}
+		}
+	}
+}
+
+
+void UMovieSceneParameterSection::MoveSection(float DeltaPosition, TSet<FKeyHandle>& KeyHandles)
+{
+	Super::MoveSection(DeltaPosition, KeyHandles);
+
+	TArray<FRichCurve*> AllCurves;
+	GatherCurves(AllCurves);
+
+	for (auto Curve : AllCurves)
+	{
+		Curve->ShiftCurve(DeltaPosition, KeyHandles);
+	}
+}
diff --git a/Engine/Source/Runtime/MovieSceneTracks/Private/Sections/MovieSceneVectorSection.cpp b/Engine/Source/Runtime/MovieSceneTracks/Private/Sections/MovieSceneVectorSection.cpp
--- a/Engine/Source/Runtime/MovieSceneTracks/Private/Sections/MovieSceneVectorSection.cpp
+++ b/Engine/Source/Runtime/MovieSceneTracks/Private/Sections/MovieSceneVectorSection.cpp
@@ -90,7 +90,8 @@ TSharedPtr<FStructOnScope> UMovieSceneVectorSection::GetKeyStruct(const TArray<F
 	{
 		for (int32 Index = 0; Index <= 3; ++Index)
 		{
-			Struct->Keys[Index] = &Curves[Index].GetFirstMatchingKey(KeyHandles);
+			Struct->Keys[Index] = Curves[Index].GetFirstMatchingKey(KeyHandles);
+			check(Struct->Keys[Index] != nullptr);
 			Struct->Vector[Index] = Struct->Keys[Index]->Value;
 		}
 	}
diff --git a/Engine/Source/Runtime/MovieSceneTracks/Private/TrackInstances/MovieScene3DConstraintTrackInstance.cpp b/Engine/Source/Runtime/MovieSceneTracks/Private/TrackInstances/MovieScene3DConstraintTrackInstance.cpp
--- a/Engine/Source/Runtime/MovieSceneTracks/Private/TrackInstances/MovieScene3DConstraintTrackInstance.cpp
+++ b/Engine/Source/Runtime/MovieSceneTracks/Private/TrackInstances/MovieScene3DConstraintTrackInstance.cpp
@@ -73,6 +73,13 @@ void FMovieScene3DConstraintTrackInstance::Update(EMovieSceneUpdateData& UpdateD
 			{
 				Player.GetRuntimeObjects( Player.GetRootMovieSceneSequenceInstance(), ConstraintId, ConstraintObjects);
 
+				// Also look in this current sequence instance for the constrained object
+				UObject* FoundObject = SequenceInstance.FindObject(ConstraintId, Player);
+				if (FoundObject && ConstraintObjects.Find(FoundObject) == INDEX_NONE)
+				{
+					ConstraintObjects.Add(FoundObject);
+				}
+
 				for (int32 ConstraintObjectIndex = 0; ConstraintObjectIndex < ConstraintObjects.Num(); ++ConstraintObjectIndex)
 				{
 					AActor* Actor = Cast<AActor>(ConstraintObjects[ConstraintObjectIndex]);
diff --git a/Engine/Source/Runtime/MovieSceneTracks/Private/TrackInstances/MovieSceneAudioTrackInstance.cpp b/Engine/Source/Runtime/MovieSceneTracks/Private/TrackInstances/MovieSceneAudioTrackInstance.cpp
--- a/Engine/Source/Runtime/MovieSceneTracks/Private/TrackInstances/MovieSceneAudioTrackInstance.cpp
+++ b/Engine/Source/Runtime/MovieSceneTracks/Private/TrackInstances/MovieSceneAudioTrackInstance.cpp
@@ -18,6 +18,12 @@ FMovieSceneAudioTrackInstance::FMovieSceneAudioTrackInstance( UMovieSceneAudioTr
 
 void FMovieSceneAudioTrackInstance::Update(EMovieSceneUpdateData& UpdateData, const TArray<UObject*>& RuntimeObjects, class IMovieScenePlayer& Player, FMovieSceneSequenceInstance& SequenceInstance) 
 {
+	// If the playback status is jumping, ie. one such occurrence is setting the time for thumbnail generation, disable audio updates
+	if (Player.GetPlaybackStatus() == EMovieScenePlayerStatus::Jumping)
+	{
+		return;
+	}
+
 	const TArray<UMovieSceneSection*>& AudioSections = AudioTrack->GetAudioSections();
 
 	TArray<AActor*> Actors = GetRuntimeActors(RuntimeObjects);
diff --git a/Engine/Source/Runtime/MovieSceneTracks/Private/TrackInstances/MovieSceneCameraCutTrackInstance.cpp b/Engine/Source/Runtime/MovieSceneTracks/Private/TrackInstances/MovieSceneCameraCutTrackInstance.cpp
--- a/Engine/Source/Runtime/MovieSceneTracks/Private/TrackInstances/MovieSceneCameraCutTrackInstance.cpp
+++ b/Engine/Source/Runtime/MovieSceneTracks/Private/TrackInstances/MovieSceneCameraCutTrackInstance.cpp
@@ -57,7 +57,7 @@ void FMovieSceneCameraCutTrackInstance::Update(EMovieSceneUpdateData& UpdateData
 	UObject* CameraObject = nullptr;
 	const TArray<UMovieSceneSection*>& CameraCutSections = CameraCutTrack->GetAllSections();
 
-	UMovieSceneCameraCutSection* NearestCameraCutSection = Cast<UMovieSceneCameraCutSection>(MovieSceneHelpers::FindNearestSectionAtTime(CameraCutSections, UpdateData.Position));
+	UMovieSceneCameraCutSection* NearestCameraCutSection = Cast<UMovieSceneCameraCutSection>(MovieSceneHelpers::FindSectionAtTime(CameraCutSections, UpdateData.Position));
 	if (NearestCameraCutSection != nullptr)
 	{
 		for (int32 CameraCutIndex = 0; CameraCutIndex < CameraCutSections.Num(); ++CameraCutIndex)
diff --git a/Engine/Source/Runtime/MovieSceneTracks/Private/TrackInstances/MovieSceneCinematicShotTrackInstance.cpp b/Engine/Source/Runtime/MovieSceneTracks/Private/TrackInstances/MovieSceneCinematicShotTrackInstance.cpp
--- a/Engine/Source/Runtime/MovieSceneTracks/Private/TrackInstances/MovieSceneCinematicShotTrackInstance.cpp
+++ b/Engine/Source/Runtime/MovieSceneTracks/Private/TrackInstances/MovieSceneCinematicShotTrackInstance.cpp
@@ -60,7 +60,12 @@ void FMovieSceneCinematicShotTrackInstance::Update(EMovieSceneUpdateData& Update
 	FMovieSceneSubTrackInstance::Update(UpdateData, RuntimeObjects, Player, SequenceInstance);
 
 	const TArray<UMovieSceneSection*>& AllSections = SubTrack->GetAllSections();
-	TArray<UMovieSceneSection*> TraversedSections = GetTraversedSectionsWithPreroll(AllSections, UpdateData.Position, UpdateData.LastPosition);
+
+	// Evaluate only the sections that intersect with the update time and not including the last position
+	float CurrentTime = UpdateData.Position;
+	float PreviousTime = CurrentTime;
+
+	TArray<UMovieSceneSection*> TraversedSections = GetTraversedSectionsWithPreroll(AllSections, CurrentTime, PreviousTime);
 
 	if (TraversedSections.Num() == 0)
 	{
diff --git a/Engine/Source/Runtime/MovieSceneTracks/Private/TrackInstances/MovieSceneLevelVisibilityTrackInstance.cpp b/Engine/Source/Runtime/MovieSceneTracks/Private/TrackInstances/MovieSceneLevelVisibilityTrackInstance.cpp
new file mode 100644
--- /dev/null
+++ b/Engine/Source/Runtime/MovieSceneTracks/Private/TrackInstances/MovieSceneLevelVisibilityTrackInstance.cpp
@@ -0,0 +1,198 @@
+// Copyright 1998-2015 Epic Games, Inc. All Rights Reserved.
+
+#include "MovieSceneTracksPrivatePCH.h"
+#include "MovieSceneLevelVisibilityTrack.h"
+#include "MovieSceneLevelVisibilityTrackInstance.h"
+
+
+FMovieSceneLevelVisibilityTrackInstance::FMovieSceneLevelVisibilityTrackInstance( UMovieSceneLevelVisibilityTrack& InLevelVisibilityTrack )
+{
+	LevelVisibilityTrack = &InLevelVisibilityTrack;
+}
+
+// TODO: This was copied from LevelStreaming.cpp, it should be in a set of shared utilities somewhere.
+FString MakeSafeLevelName( const FName& InLevelName )
+{
+	// Special case for PIE, the PackageName gets mangled.
+	if ( !GWorld->StreamingLevelsPrefix.IsEmpty() )
+	{
+		FString PackageName = GWorld->StreamingLevelsPrefix + FPackageName::GetShortName( InLevelName );
+		if ( !FPackageName::IsShortPackageName( InLevelName ) )
+		{
+			PackageName = FPackageName::GetLongPackagePath( InLevelName.ToString() ) + TEXT( "/" ) + PackageName;
+		}
+
+		return PackageName;
+	}
+
+	return InLevelName.ToString();
+}
+
+
+ULevelStreaming* GetStreamingLevel( FName LevelName )
+{
+	if ( LevelName != NAME_None )
+	{
+		FString SafeLevelName = MakeSafeLevelName( LevelName );
+		for ( ULevelStreaming* LevelStreaming : GWorld->StreamingLevels )
+		{
+			if ( LevelStreaming && LevelStreaming->GetWorldAssetPackageName().EndsWith( SafeLevelName, ESearchCase::IgnoreCase ) )
+			{
+				return LevelStreaming;
+			}
+		}
+	}
+
+	return nullptr;
+}
+
+
+void GetAnimatedStreamingLevels( UMovieSceneLevelVisibilityTrack* LevelVisibilityTrack, TArray<ULevelStreaming*>& StreamingLevels )
+{
+	if ( LevelVisibilityTrack != nullptr )
+	{
+		for ( UMovieSceneSection* Section : LevelVisibilityTrack->GetAllSections() )
+		{
+			UMovieSceneLevelVisibilitySection* LevelVisibilitySection = Cast<UMovieSceneLevelVisibilitySection>( Section );
+			if ( LevelVisibilitySection != nullptr )
+			{
+				for ( FName StreamingLevelName : *LevelVisibilitySection->GetLevelNames() )
+				{
+					ULevelStreaming* StreamingLevel = GetStreamingLevel( StreamingLevelName );
+					if ( StreamingLevel != nullptr )
+					{
+						StreamingLevels.Add( StreamingLevel );
+					}
+				}
+			}
+		}
+	}
+}
+
+
+void FMovieSceneLevelVisibilityTrackInstance::SaveState(const TArray<UObject*>& RuntimeObjects, IMovieScenePlayer& Player, FMovieSceneSequenceInstance& SequenceInstance)
+{
+	TArray<ULevelStreaming*> AnimatedStreamingLevels;
+	GetAnimatedStreamingLevels(LevelVisibilityTrack, AnimatedStreamingLevels);
+
+	for ( ULevelStreaming* StreamingLevel : AnimatedStreamingLevels )
+	{
+		SaveLevelVisibilityState( StreamingLevel );
+	}
+}
+
+
+void FMovieSceneLevelVisibilityTrackInstance::RestoreState(const TArray<UObject*>& RuntimeObjects, IMovieScenePlayer& Player, FMovieSceneSequenceInstance& SequenceInstance)
+{
+	TArray<ULevelStreaming*> AnimatedStreamingLevels;
+	GetAnimatedStreamingLevels( LevelVisibilityTrack, AnimatedStreamingLevels );
+
+	for ( ULevelStreaming* StreamingLevel : AnimatedStreamingLevels )
+	{
+		bool* SavedVisibility = SavedLevelVisibility.Find( FObjectKey( StreamingLevel ) );
+		if ( SavedVisibility != nullptr )
+		{
+			SetEditorLevelVisibilityState( StreamingLevel, *SavedVisibility );
+		}
+	}
+
+	SavedLevelVisibility.Empty();
+}
+
+
+void FMovieSceneLevelVisibilityTrackInstance::Update(EMovieSceneUpdateData& UpdateData, const TArray<UObject*>& RuntimeObjects, class IMovieScenePlayer& Player, FMovieSceneSequenceInstance& SequenceInstance ) 
+{
+	for ( UMovieSceneSection* Section : LevelVisibilityTrack->GetAllSections() )
+	{
+		UMovieSceneLevelVisibilitySection* LevelVisibilitySection = Cast<UMovieSceneLevelVisibilitySection>(Section);
+		if (LevelVisibilitySection != nullptr &&
+			LevelVisibilitySection->GetStartTime() <= UpdateData.Position &&
+			LevelVisibilitySection->GetEndTime() > UpdateData.Position )
+		{
+			bool bShouldBeVisible = LevelVisibilitySection->GetVisibility() == ELevelVisibility::Visible ? true : false;
+			bool bVisibilityModified = false;
+			for ( FName LevelName : *LevelVisibilitySection->GetLevelNames() )
+			{
+				FName SafeLevelName ( *MakeSafeLevelName( LevelName ) );
+				ULevelStreaming* StreamingLevel = nullptr;
+				TWeakObjectPtr<ULevelStreaming>* FoundStreamingLevel = NameToLevelMap.Find( SafeLevelName );
+
+				if ( FoundStreamingLevel != nullptr )
+				{
+					StreamingLevel = FoundStreamingLevel->Get();
+				}
+				else if ( GIsEditor && !GWorld->HasBegunPlay() )
+				{
+					// If we're in the editor, it possible the level list has been updated since the last refresh, so try to find the level.
+					StreamingLevel = GetStreamingLevel( LevelName );
+					if ( StreamingLevel != nullptr )
+					{
+						SaveLevelVisibilityState( StreamingLevel );
+						NameToLevelMap.Add( LevelName, TWeakObjectPtr<ULevelStreaming>( StreamingLevel ) );
+					}
+				}
+
+				if ( StreamingLevel != nullptr )
+				{
+					if ( GIsEditor && !GWorld->HasBegunPlay() )
+					{
+						SetEditorLevelVisibilityState( StreamingLevel, bShouldBeVisible );
+					}
+					else
+					{
+						if ( StreamingLevel->bShouldBeVisible != bShouldBeVisible )
+						{
+							StreamingLevel->bShouldBeVisible = bShouldBeVisible;
+							bVisibilityModified = true;
+						}
+					}
+					
+				}
+			}
+			if ( bVisibilityModified )
+			{
+				GWorld->FlushLevelStreaming( EFlushLevelStreamingType::Visibility );
+			}
+		}
+	}
+}
+
+
+void FMovieSceneLevelVisibilityTrackInstance::RefreshInstance( const TArray<UObject*>& RuntimeObjects, class IMovieScenePlayer& Player, FMovieSceneSequenceInstance& SequenceInstance )
+{
+	NameToLevelMap.Empty();
+
+	TArray<ULevelStreaming*> AnimatedStreamingLevels;
+	GetAnimatedStreamingLevels( LevelVisibilityTrack, AnimatedStreamingLevels );
+
+	for ( ULevelStreaming* StreamingLevel : AnimatedStreamingLevels )
+	{
+		NameToLevelMap.Add( FPackageName::GetShortFName( StreamingLevel->GetWorldAssetPackageFName() ), TWeakObjectPtr<ULevelStreaming>( StreamingLevel ) );
+	}
+}
+
+void FMovieSceneLevelVisibilityTrackInstance::SaveLevelVisibilityState( ULevelStreaming* StreamingLevel )
+{
+	SavedLevelVisibility.Add( FObjectKey( StreamingLevel ), StreamingLevel->bShouldBeVisibleInEditor );
+}
+
+void FMovieSceneLevelVisibilityTrackInstance::SetEditorLevelVisibilityState( ULevelStreaming* StreamingLevel, bool bShouldBeVisible )
+{
+	if ( StreamingLevel->bShouldBeVisibleInEditor != bShouldBeVisible )
+	{
+		StreamingLevel->bShouldBeVisibleInEditor = bShouldBeVisible;
+		ULevel* LoadedLevel = StreamingLevel->GetLoadedLevel();
+		if ( LoadedLevel != nullptr )
+		{
+			for ( AActor* Actor : LoadedLevel->Actors )
+			{
+				if ( Actor != nullptr )
+				{
+#if WITH_EDITOR
+					Actor->SetIsTemporarilyHiddenInEditor( !bShouldBeVisible );
+#endif
+				}
+			}
+		}
+	}
+}
diff --git a/Engine/Source/Runtime/MovieSceneTracks/Private/TrackInstances/MovieSceneParticleTrackInstance.cpp b/Engine/Source/Runtime/MovieSceneTracks/Private/TrackInstances/MovieSceneParticleTrackInstance.cpp
--- a/Engine/Source/Runtime/MovieSceneTracks/Private/TrackInstances/MovieSceneParticleTrackInstance.cpp
+++ b/Engine/Source/Runtime/MovieSceneTracks/Private/TrackInstances/MovieSceneParticleTrackInstance.cpp
@@ -12,6 +12,17 @@ FMovieSceneParticleTrackInstance::~FMovieSceneParticleTrackInstance()
 {
 }
 
+static UParticleSystemComponent* ParticleSystemComponentFromRuntimeObject(UObject* Object)
+{
+	if(AEmitter* Emitter = Cast<AEmitter>(Object))
+	{
+		return Emitter->GetParticleSystemComponent();
+	}
+	else
+	{
+		return Cast<UParticleSystemComponent>(Object);
+	}
+}
 
 void FMovieSceneParticleTrackInstance::Update(EMovieSceneUpdateData& UpdateData, const TArray<UObject*>& RuntimeObjects, class IMovieScenePlayer& Player, FMovieSceneSequenceInstance& SequenceInstance) 
 {
@@ -47,28 +58,25 @@ void FMovieSceneParticleTrackInstance::Update(EMovieSceneUpdateData& UpdateData,
 		{
 			for (int32 i = 0; i < RuntimeObjects.Num(); ++i)
 			{
-				AEmitter* Emitter = Cast<AEmitter>(RuntimeObjects[i]);
-				if (Emitter)
+				UParticleSystemComponent* ParticleSystemComponent = ParticleSystemComponentFromRuntimeObject(RuntimeObjects[i]);
+
+				if ( ParticleSystemComponent != nullptr )
 				{
-					UParticleSystemComponent* ParticleSystemComponent = Emitter->GetParticleSystemComponent();
-					if ( ParticleSystemComponent != nullptr )
+					if ( ParticleKey == EParticleKey::Activate)
 					{
-						if ( ParticleKey == EParticleKey::Activate)
-						{
-							if ( ParticleSystemComponent->IsActive() )
-							{
-								ParticleSystemComponent->SetActive(false, true);
-							}
-							ParticleSystemComponent->SetActive(true, true);
-						}
-						else if( ParticleKey == EParticleKey::Deactivate )
+						if ( ParticleSystemComponent->IsActive() )
 						{
 							ParticleSystemComponent->SetActive(false, true);
 						}
-						else if ( ParticleKey == EParticleKey::Trigger )
-						{
-							ParticleSystemComponent->ActivateSystem(true);
-						}
+						ParticleSystemComponent->SetActive(true, true);
+					}
+					else if( ParticleKey == EParticleKey::Deactivate )
+					{
+						ParticleSystemComponent->SetActive(false, true);
+					}
+					else if ( ParticleKey == EParticleKey::Trigger )
+					{
+						ParticleSystemComponent->ActivateSystem(true);
 					}
 				}
 			}
@@ -83,6 +91,10 @@ void FMovieSceneParticleTrackInstance::Update(EMovieSceneUpdateData& UpdateData,
 			{
 				Emitter->Deactivate();
 			}
+			else if(UParticleSystemComponent* Component =  Cast<UParticleSystemComponent>(RuntimeObjects[i]))
+			{
+				Component->SetActive(false, true);
+			}
 		}
 	}
 }
diff --git a/Engine/Source/Runtime/MovieSceneTracks/Private/TrackInstances/MovieSceneSkeletalAnimationTrackInstance.cpp b/Engine/Source/Runtime/MovieSceneTracks/Private/TrackInstances/MovieSceneSkeletalAnimationTrackInstance.cpp
--- a/Engine/Source/Runtime/MovieSceneTracks/Private/TrackInstances/MovieSceneSkeletalAnimationTrackInstance.cpp
+++ b/Engine/Source/Runtime/MovieSceneTracks/Private/TrackInstances/MovieSceneSkeletalAnimationTrackInstance.cpp
@@ -61,46 +61,164 @@ void FMovieSceneSkeletalAnimationTrackInstance::Update(EMovieSceneUpdateData& Up
 
 	for (int32 i = 0; i < RuntimeObjects.Num(); ++i)
 	{
-		IMatineeAnimInterface* AnimInterface = Cast<IMatineeAnimInterface>(RuntimeObjects[i]);
-		if (AnimInterface) 
+		// get the skeletal mesh component we want to control
+		USkeletalMeshComponent* Component = nullptr;
+
+		UObject* RuntimeObject = RuntimeObjects[i];
+		if(RuntimeObject != nullptr)
 		{
-			UMovieSceneSkeletalAnimationSection* AnimSection = Cast<UMovieSceneSkeletalAnimationSection>(AnimationTrack->GetAnimSectionAtTime(UpdateData.Position));
-			
-			// cbb: If there is no overlapping section, evaluate the closest section only if the current time is before it.
-			if (AnimSection == nullptr)
+			// first try to control the component directly
+			USkeletalMeshComponent* SkeletalMeshComponent = Cast<USkeletalMeshComponent>(RuntimeObject);
+			if(SkeletalMeshComponent == nullptr)
 			{
-				AnimSection = Cast<UMovieSceneSkeletalAnimationSection>(MovieSceneHelpers::FindNearestSectionAtTime(AnimationTrack->GetAllSections(), UpdateData.Position));
-				if (AnimSection != nullptr && (UpdateData.Position >= AnimSection->GetStartTime() || UpdateData.Position <= AnimSection->GetEndTime()))
+				// then check to see if we are controlling an actor & if so use its first USkeletalMeshComponent 
+				AActor* Actor = Cast<AActor>(RuntimeObject);
+				if(Actor != nullptr)
 				{
-					AnimSection = nullptr;
+					SkeletalMeshComponent = Actor->FindComponentByClass<USkeletalMeshComponent>();
 				}
 			}
 
-			if (AnimSection && AnimSection->IsActive())
+			if(SkeletalMeshComponent)
 			{
-				const FAnimEvalTimes EvalTimes = FAnimEvalTimes::MapTimesToAnimation(UpdateData.Position, UpdateData.LastPosition, AnimSection);
-
-				int32 ChannelIndex = 0;
-				UAnimSequence* AnimSequence = AnimSection->GetAnimSequence();
-
-				// NOTE: The order of the bLooping and bFireNotifiers parameters are swapped in the preview and non-preview SetAnimPosition functions.
-				const bool bLooping = false;
-				if (GIsEditor && !GWorld->HasBegunPlay())
+				UMovieSceneSkeletalAnimationSection* AnimSection = Cast<UMovieSceneSkeletalAnimationSection>(AnimationTrack->GetAnimSectionAtTime(UpdateData.Position));
+			
+				// cbb: If there is no overlapping section, evaluate the closest section only if the current time is before it.
+				if (AnimSection == nullptr)
 				{
-					// Pass an absolute delta time to the preview anim evaluation
-					float Delta = FMath::Abs(EvalTimes.Current - EvalTimes.Previous);
-					const bool bFireNotifies = true;
-					AnimInterface->PreviewBeginAnimControl(nullptr);
-					AnimInterface->PreviewSetAnimPosition(AnimSection->GetSlotName(), ChannelIndex, AnimSequence, EvalTimes.Current, bLooping, bFireNotifies, Delta);
+					AnimSection = Cast<UMovieSceneSkeletalAnimationSection>(MovieSceneHelpers::FindNearestSectionAtTime(AnimationTrack->GetAllSections(), UpdateData.Position));
+					if (AnimSection != nullptr && (UpdateData.Position >= AnimSection->GetStartTime() || UpdateData.Position <= AnimSection->GetEndTime()))
+					{
+						AnimSection = nullptr;
+					}
 				}
-				else
+
+				if (AnimSection && AnimSection->IsActive())
 				{
-					// Don't fire notifies at runtime since they will be fired through the standard animation tick path.
-					const bool bFireNotifies = false;
-					AnimInterface->BeginAnimControl(nullptr);
-					AnimInterface->SetAnimPosition(AnimSection->GetSlotName(), ChannelIndex, AnimSequence, EvalTimes.Current, bFireNotifies, bLooping);
+					UAnimSequence* AnimSequence = AnimSection->GetAnimSequence();
+
+					if (AnimSequence)
+					{
+						const FAnimEvalTimes EvalTimes = FAnimEvalTimes::MapTimesToAnimation(UpdateData.Position, UpdateData.LastPosition, AnimSection);
+
+						int32 ChannelIndex = 0;
+
+						// NOTE: The order of the bLooping and bFireNotifiers parameters are swapped in the preview and non-preview SetAnimPosition functions.
+						const bool bLooping = false;
+						if (GIsEditor && !GWorld->HasBegunPlay())
+						{
+							// Pass an absolute delta time to the preview anim evaluation
+							float DeltaTime = FMath::Abs(EvalTimes.Current - EvalTimes.Previous);
+						
+							// If the playback status is jumping, ie. one such occurrence is setting the time for thumbnail generation, disable anim notifies updates because it could fire audio
+							const bool bFireNotifies = Player.GetPlaybackStatus() == EMovieScenePlayerStatus::Jumping || Player.GetPlaybackStatus() == EMovieScenePlayerStatus::Stopped ? false : true;
+
+							// When jumping from one cut to another cut, the delta time should be 0 so that anim notifies before the current position are not evaluated. Note, anim notifies at the current time should still be evaluated.
+							if (UpdateData.bJumpCut)
+							{
+								DeltaTime = 0.f;
+							}
+						
+							PreviewBeginAnimControl(SkeletalMeshComponent);
+							PreviewSetAnimPosition(SkeletalMeshComponent, AnimSection->GetSlotName(), ChannelIndex, AnimSequence, EvalTimes.Current, bLooping, bFireNotifies, DeltaTime);
+						}
+						else
+						{
+							// Don't fire notifies at runtime since they will be fired through the standard animation tick path.
+							const bool bFireNotifies = false;
+							BeginAnimControl(SkeletalMeshComponent);
+							SetAnimPosition(SkeletalMeshComponent, AnimSection->GetSlotName(), ChannelIndex, AnimSequence, EvalTimes.Current, bFireNotifies, bLooping);
+						}
+					}
 				}
 			}
 		}
 	}
 }
+
+void FMovieSceneSkeletalAnimationTrackInstance::BeginAnimControl(USkeletalMeshComponent* SkeletalMeshComponent)
+{
+	if (CanPlayAnimation(SkeletalMeshComponent))
+	{
+		UAnimInstance* AnimInstance = SkeletalMeshComponent->GetAnimInstance();
+		if (!AnimInstance)
+		{
+			SkeletalMeshComponent->SetAnimationMode(EAnimationMode::Type::AnimationSingleNode);
+		}
+	}
+}
+
+bool FMovieSceneSkeletalAnimationTrackInstance::CanPlayAnimation(USkeletalMeshComponent* SkeletalMeshComponent, class UAnimSequenceBase* AnimAssetBase/*=nullptr*/) const
+{
+	return (SkeletalMeshComponent->SkeletalMesh && SkeletalMeshComponent->SkeletalMesh->Skeleton && 
+		(!AnimAssetBase || SkeletalMeshComponent->SkeletalMesh->Skeleton->IsCompatible(AnimAssetBase->GetSkeleton())));
+}
+
+void FMovieSceneSkeletalAnimationTrackInstance::SetAnimPosition(USkeletalMeshComponent* SkeletalMeshComponent, FName SlotName, int32 ChannelIndex, UAnimSequence* InAnimSequence, float InPosition, bool bFireNotifies, bool bLooping)
+{
+	if (CanPlayAnimation(SkeletalMeshComponent, InAnimSequence))
+	{
+		FAnimMontageInstance::SetMatineeAnimPositionInner(SlotName, SkeletalMeshComponent, InAnimSequence, CurrentlyPlayingMontage, InPosition, bLooping);
+	}
+}
+
+void FMovieSceneSkeletalAnimationTrackInstance::FinishAnimControl(USkeletalMeshComponent* SkeletalMeshComponent)
+{
+	if(SkeletalMeshComponent->GetAnimationMode() == EAnimationMode::Type::AnimationBlueprint)
+	{
+		UAnimInstance* AnimInstance = SkeletalMeshComponent->GetAnimInstance();
+		if(AnimInstance)
+		{
+			AnimInstance->Montage_Stop(0.f);
+			AnimInstance->UpdateAnimation(0.f, false);
+		}
+
+		// Update space bases to reset it back to ref pose
+		SkeletalMeshComponent->RefreshBoneTransforms();
+		SkeletalMeshComponent->RefreshSlaveComponents();
+		SkeletalMeshComponent->UpdateComponentToWorld();
+	}
+}
+
+
+void FMovieSceneSkeletalAnimationTrackInstance::PreviewBeginAnimControl(USkeletalMeshComponent* SkeletalMeshComponent)
+{
+	if (CanPlayAnimation(SkeletalMeshComponent))
+	{
+		UAnimInstance* AnimInstance = SkeletalMeshComponent->GetAnimInstance();
+		if (!AnimInstance)
+		{
+			SkeletalMeshComponent->SetAnimationMode(EAnimationMode::Type::AnimationSingleNode);
+		}
+	}
+}
+
+void FMovieSceneSkeletalAnimationTrackInstance::PreviewFinishAnimControl(USkeletalMeshComponent* SkeletalMeshComponent)
+{
+	if (CanPlayAnimation(SkeletalMeshComponent))
+	{
+		// if in editor, reset the Animations, makes easier for artist to see them visually and align them
+		// in game, we keep the last pose that matinee kept. If you'd like it to have animation, you'll need to have AnimTree or AnimGraph to handle correctly
+		if (SkeletalMeshComponent->GetAnimationMode() == EAnimationMode::Type::AnimationBlueprint)
+		{
+			UAnimInstance* AnimInstance = SkeletalMeshComponent->GetAnimInstance();
+			if(AnimInstance)
+			{
+				AnimInstance->Montage_Stop(0.f);
+				AnimInstance->UpdateAnimation(0.f, false);
+			}
+		}
+		// Update space bases to reset it back to ref pose
+		SkeletalMeshComponent->RefreshBoneTransforms();
+		SkeletalMeshComponent->RefreshSlaveComponents();
+		SkeletalMeshComponent->UpdateComponentToWorld();
+	}
+}
+
+void FMovieSceneSkeletalAnimationTrackInstance::PreviewSetAnimPosition(USkeletalMeshComponent* SkeletalMeshComponent, FName SlotName, int32 ChannelIndex, UAnimSequence* InAnimSequence, float InPosition, bool bLooping, bool bFireNotifies, float DeltaTime)
+{
+	if(CanPlayAnimation(SkeletalMeshComponent, InAnimSequence))
+	{
+		FAnimMontageInstance::PreviewMatineeSetAnimPositionInner(SlotName, SkeletalMeshComponent, InAnimSequence, CurrentlyPlayingMontage, InPosition, bLooping, bFireNotifies, DeltaTime);
+	}
+}
\ No newline at end of file
diff --git a/Engine/Source/Runtime/MovieSceneTracks/Private/TrackInstances/MovieSceneSubTrackInstance.cpp b/Engine/Source/Runtime/MovieSceneTracks/Private/TrackInstances/MovieSceneSubTrackInstance.cpp
--- a/Engine/Source/Runtime/MovieSceneTracks/Private/TrackInstances/MovieSceneSubTrackInstance.cpp
+++ b/Engine/Source/Runtime/MovieSceneTracks/Private/TrackInstances/MovieSceneSubTrackInstance.cpp
@@ -170,7 +170,12 @@ TArray<UMovieSceneSection*> FMovieSceneSubTrackInstance::GetTraversedSectionsWit
 void FMovieSceneSubTrackInstance::Update(EMovieSceneUpdateData& UpdateData, const TArray<UObject*>& RuntimeObjects, class IMovieScenePlayer& Player, FMovieSceneSequenceInstance& SequenceInstance) 
 {
 	const TArray<UMovieSceneSection*>& AllSections = SubTrack->GetAllSections();
-	TArray<UMovieSceneSection*> TraversedSections = GetTraversedSectionsWithPreroll(AllSections, UpdateData.Position, UpdateData.LastPosition);
+
+	// Evaluate only the sections that intersect with the update time and not including the last position
+	float CurrentTime = UpdateData.Position;
+	float PreviousTime = CurrentTime;
+
+	TArray<UMovieSceneSection*> TraversedSections = GetTraversedSectionsWithPreroll(AllSections, CurrentTime, PreviousTime);
 
 	const float InitialUpdatePosition = UpdateData.Position;
 
@@ -198,6 +203,7 @@ void FMovieSceneSubTrackInstance::Update(EMovieSceneUpdateData& UpdateData, cons
 		const float InstancePosition = InstanceOffset + (UpdateData.Position - (SubSection->GetStartTime()- SubSection->PrerollTime)) / SubSection->TimeScale;
 
 		EMovieSceneUpdateData SubUpdateData(InstancePosition, InstanceLastPosition);
+		SubUpdateData.bJumpCut = UpdateData.LastPosition < SubSection->GetStartTime();
 		SubUpdateData.UpdatePass = UpdateData.UpdatePass;
 		SubUpdateData.bPreroll = InitialUpdatePosition < SubSection->GetStartTime();
 
diff --git a/Engine/Source/Runtime/MovieSceneTracks/Private/TrackInstances/MovieSceneVisibilityTrackInstance.cpp b/Engine/Source/Runtime/MovieSceneTracks/Private/TrackInstances/MovieSceneVisibilityTrackInstance.cpp
--- a/Engine/Source/Runtime/MovieSceneTracks/Private/TrackInstances/MovieSceneVisibilityTrackInstance.cpp
+++ b/Engine/Source/Runtime/MovieSceneTracks/Private/TrackInstances/MovieSceneVisibilityTrackInstance.cpp
@@ -109,6 +109,7 @@ void FMovieSceneVisibilityTrackInstance::Update( EMovieSceneUpdateData& UpdateDa
  		for( UObject* Object : RuntimeObjects )
  		{
 			AActor* Actor = Cast<AActor>(Object);
+			USceneComponent* SceneComponent = Cast<USceneComponent>(Object);
 
 			if (Actor != nullptr)
 			{
@@ -124,17 +125,24 @@ void FMovieSceneVisibilityTrackInstance::Update( EMovieSceneUpdateData& UpdateDa
 						Actor->SetIsTemporarilyHiddenInEditor(Visible);
 					}
 
-					USceneComponent* SceneComponent = MovieSceneHelpers::SceneComponentFromRuntimeObject(Actor);
-					if (SceneComponent != nullptr)
+					USceneComponent* RuntimeSceneComponent = MovieSceneHelpers::SceneComponentFromRuntimeObject(Actor);
+					if (RuntimeSceneComponent != nullptr)
 					{
-						if (SceneComponent->IsVisibleInEditor() != !Visible)
+						if (RuntimeSceneComponent->IsVisibleInEditor() != !Visible)
 						{
-							SceneComponent->SetVisibility(!Visible);
+							RuntimeSceneComponent->SetVisibility(!Visible);
 						}
 					}
 				}
 #endif // WITH_EDITOR
 			}
+			else if(SceneComponent != nullptr)
+			{
+				if (SceneComponent->IsVisible() != !Visible)
+				{
+					SceneComponent->SetVisibility(!Visible);
+				}
+			}
  		}
 	}
 }
diff --git a/Engine/Source/Runtime/MovieSceneTracks/Private/Tracks/MovieScene3DAttachTrack.cpp b/Engine/Source/Runtime/MovieSceneTracks/Private/Tracks/MovieScene3DAttachTrack.cpp
--- a/Engine/Source/Runtime/MovieSceneTracks/Private/Tracks/MovieScene3DAttachTrack.cpp
+++ b/Engine/Source/Runtime/MovieSceneTracks/Private/Tracks/MovieScene3DAttachTrack.cpp
@@ -21,7 +21,7 @@ TSharedPtr<IMovieSceneTrackInstance> UMovieScene3DAttachTrack::CreateInstance()
 }
 
 
-void UMovieScene3DAttachTrack::AddConstraint(float KeyTime, float ConstraintEndTime, const FName SocketName, const FGuid& ConstraintId)
+void UMovieScene3DAttachTrack::AddConstraint(float KeyTime, float ConstraintEndTime, const FName SocketName, const FName ComponentName, const FGuid& ConstraintId)
 {
 	// add the section
 	UMovieScene3DAttachSection* NewSection = NewObject<UMovieScene3DAttachSection>(this);
@@ -29,10 +29,20 @@ void UMovieScene3DAttachTrack::AddConstraint(float KeyTime, float ConstraintEndT
 	NewSection->AddAttach(KeyTime, ConstraintEndTime, ConstraintId);
 	NewSection->InitialPlacement( ConstraintSections, KeyTime, ConstraintEndTime, SupportsMultipleRows() );
 	NewSection->AttachSocketName = SocketName;
+	NewSection->AttachComponentName = ComponentName;
 
 	ConstraintSections.Add(NewSection);
 }
 
+UMovieSceneSection* UMovieScene3DAttachTrack::CreateNewSection()
+{
+	UMovieScene3DAttachSection* NewSection = NewObject<UMovieScene3DAttachSection>(this);
+	NewSection->SetFlags(RF_Transactional);
+
+	ConstraintSections.Add(NewSection);
+
+	return NewSection;
+}
 
 #if WITH_EDITORONLY_DATA
 FText UMovieScene3DAttachTrack::GetDisplayName() const
diff --git a/Engine/Source/Runtime/MovieSceneTracks/Private/Tracks/MovieScene3DConstraintTrack.cpp b/Engine/Source/Runtime/MovieSceneTracks/Private/Tracks/MovieScene3DConstraintTrack.cpp
--- a/Engine/Source/Runtime/MovieSceneTracks/Private/Tracks/MovieScene3DConstraintTrack.cpp
+++ b/Engine/Source/Runtime/MovieSceneTracks/Private/Tracks/MovieScene3DConstraintTrack.cpp
@@ -12,7 +12,11 @@
 
 UMovieScene3DConstraintTrack::UMovieScene3DConstraintTrack( const FObjectInitializer& ObjectInitializer )
 	: Super( ObjectInitializer )
-{ }
+{
+#if WITH_EDITORONLY_DATA
+	TrackTint = FColor(128, 90, 0);
+#endif
+}
 
 
 const TArray<UMovieSceneSection*>& UMovieScene3DConstraintTrack::GetAllSections() const
diff --git a/Engine/Source/Runtime/MovieSceneTracks/Private/Tracks/MovieScene3DPathTrack.cpp b/Engine/Source/Runtime/MovieSceneTracks/Private/Tracks/MovieScene3DPathTrack.cpp
--- a/Engine/Source/Runtime/MovieSceneTracks/Private/Tracks/MovieScene3DPathTrack.cpp
+++ b/Engine/Source/Runtime/MovieSceneTracks/Private/Tracks/MovieScene3DPathTrack.cpp
@@ -21,7 +21,7 @@ TSharedPtr<IMovieSceneTrackInstance> UMovieScene3DPathTrack::CreateInstance()
 }
 
 
-void UMovieScene3DPathTrack::AddConstraint(float KeyTime, float ConstraintEndTime, const FName SocketName, const FGuid& ConstraintId)
+void UMovieScene3DPathTrack::AddConstraint(float KeyTime, float ConstraintEndTime, const FName SocketName, const FName ComponentName, const FGuid& ConstraintId)
 {
 	UMovieScene3DPathSection* NewSection = NewObject<UMovieScene3DPathSection>(this);
 	{
diff --git a/Engine/Source/Runtime/MovieSceneTracks/Private/Tracks/MovieScene3DTransformTrack.cpp b/Engine/Source/Runtime/MovieSceneTracks/Private/Tracks/MovieScene3DTransformTrack.cpp
--- a/Engine/Source/Runtime/MovieSceneTracks/Private/Tracks/MovieScene3DTransformTrack.cpp
+++ b/Engine/Source/Runtime/MovieSceneTracks/Private/Tracks/MovieScene3DTransformTrack.cpp
@@ -9,7 +9,14 @@
 
 UMovieScene3DTransformTrack::UMovieScene3DTransformTrack( const FObjectInitializer& ObjectInitializer )
 	: Super( ObjectInitializer )
-{ }
+{
+	static FName Transform("Transform");
+	SetPropertyNameAndPath(Transform, Transform.ToString());
+
+#if WITH_EDITORONLY_DATA
+	TrackTint = FColor(48, 227, 255);
+#endif
+}
 
 
 UMovieSceneSection* UMovieScene3DTransformTrack::CreateNewSection()
diff --git a/Engine/Source/Runtime/MovieSceneTracks/Private/Tracks/MovieSceneAudioTrack.cpp b/Engine/Source/Runtime/MovieSceneTracks/Private/Tracks/MovieSceneAudioTrack.cpp
--- a/Engine/Source/Runtime/MovieSceneTracks/Private/Tracks/MovieSceneAudioTrack.cpp
+++ b/Engine/Source/Runtime/MovieSceneTracks/Private/Tracks/MovieSceneAudioTrack.cpp
@@ -17,7 +17,11 @@
 
 UMovieSceneAudioTrack::UMovieSceneAudioTrack( const FObjectInitializer& ObjectInitializer )
 	: Super( ObjectInitializer )
-{ }
+{
+#if WITH_EDITORONLY_DATA
+	TrackTint = FColor(0, 0, 255);
+#endif
+}
 
 
 TSharedPtr<IMovieSceneTrackInstance> UMovieSceneAudioTrack::CreateInstance()
diff --git a/Engine/Source/Runtime/MovieSceneTracks/Private/Tracks/MovieSceneCameraCutTrack.cpp b/Engine/Source/Runtime/MovieSceneTracks/Private/Tracks/MovieSceneCameraCutTrack.cpp
--- a/Engine/Source/Runtime/MovieSceneTracks/Private/Tracks/MovieSceneCameraCutTrack.cpp
+++ b/Engine/Source/Runtime/MovieSceneTracks/Private/Tracks/MovieSceneCameraCutTrack.cpp
@@ -13,6 +13,14 @@
 
 /* UMovieSceneCameraCutTrack interface
  *****************************************************************************/
+UMovieSceneCameraCutTrack::UMovieSceneCameraCutTrack( const FObjectInitializer& ObjectInitializer )
+	: Super( ObjectInitializer )
+{
+#if WITH_EDITORONLY_DATA
+	TrackTint = FColor(0, 0, 0);
+#endif
+}
+
 
 void UMovieSceneCameraCutTrack::AddNewCameraCut(FGuid CameraHandle, float StartTime)
 {
diff --git a/Engine/Source/Runtime/MovieSceneTracks/Private/Tracks/MovieSceneCinematicShotTrack.cpp b/Engine/Source/Runtime/MovieSceneTracks/Private/Tracks/MovieSceneCinematicShotTrack.cpp
--- a/Engine/Source/Runtime/MovieSceneTracks/Private/Tracks/MovieSceneCinematicShotTrack.cpp
+++ b/Engine/Source/Runtime/MovieSceneTracks/Private/Tracks/MovieSceneCinematicShotTrack.cpp
@@ -13,6 +13,13 @@
 
 /* UMovieSceneSubTrack interface
  *****************************************************************************/
+UMovieSceneCinematicShotTrack::UMovieSceneCinematicShotTrack(const FObjectInitializer& ObjectInitializer)
+	: Super(ObjectInitializer)
+{
+#if WITH_EDITORONLY_DATA
+	TrackTint = FColor(0,0,0);
+#endif
+}
 
 UMovieSceneSubSection* UMovieSceneCinematicShotTrack::AddSequence(UMovieSceneSequence* Sequence, float StartTime, float Duration, const bool& bInsertSequence)
 {
diff --git a/Engine/Source/Runtime/MovieSceneTracks/Private/Tracks/MovieSceneFadeTrack.cpp b/Engine/Source/Runtime/MovieSceneTracks/Private/Tracks/MovieSceneFadeTrack.cpp
--- a/Engine/Source/Runtime/MovieSceneTracks/Private/Tracks/MovieSceneFadeTrack.cpp
+++ b/Engine/Source/Runtime/MovieSceneTracks/Private/Tracks/MovieSceneFadeTrack.cpp
@@ -25,7 +25,7 @@ UMovieSceneSection* UMovieSceneFadeTrack::CreateNewSection()
 
 
 #if WITH_EDITORONLY_DATA
-FText UMovieSceneFadeTrack::GetDisplayName() const
+FText UMovieSceneFadeTrack::GetDefaultDisplayName() const
 {
 	return LOCTEXT("TrackName", "Fade");
 }
diff --git a/Engine/Source/Runtime/MovieSceneTracks/Private/Tracks/MovieSceneLevelVisibilityTrack.cpp b/Engine/Source/Runtime/MovieSceneTracks/Private/Tracks/MovieSceneLevelVisibilityTrack.cpp
new file mode 100644
--- /dev/null
+++ b/Engine/Source/Runtime/MovieSceneTracks/Private/Tracks/MovieSceneLevelVisibilityTrack.cpp
@@ -0,0 +1,80 @@
+// Copyright 1998-2015 Epic Games, Inc. All Rights Reserved.
+
+#include "MovieSceneTracksPrivatePCH.h"
+#include "MovieSceneLevelVisibilitySection.h"
+#include "MovieSceneLevelVisibilityTrack.h"
+#include "IMovieScenePlayer.h"
+#include "MovieSceneLevelVisibilityTrackInstance.h"
+
+
+#define LOCTEXT_NAMESPACE "MovieSceneLevelVisibilityTrack"
+
+UMovieSceneLevelVisibilityTrack::UMovieSceneLevelVisibilityTrack( const FObjectInitializer& ObjectInitializer )
+	: Super( ObjectInitializer )
+{
+}
+
+
+TSharedPtr<IMovieSceneTrackInstance> UMovieSceneLevelVisibilityTrack::CreateInstance()
+{
+	return MakeShareable( new FMovieSceneLevelVisibilityTrackInstance( *this ) ); 
+}
+
+
+bool UMovieSceneLevelVisibilityTrack::IsEmpty() const
+{
+	return Sections.Num() == 0;
+}
+
+
+void UMovieSceneLevelVisibilityTrack::AddSection(UMovieSceneSection& Section)
+{
+	Sections.Add(&Section);
+}
+
+
+void UMovieSceneLevelVisibilityTrack::RemoveSection( UMovieSceneSection& Section )
+{
+	Sections.Remove(&Section);
+}
+
+
+UMovieSceneSection* UMovieSceneLevelVisibilityTrack::CreateNewSection()
+{
+	return NewObject<UMovieSceneLevelVisibilitySection>(this, UMovieSceneLevelVisibilitySection::StaticClass(), NAME_None, RF_Transactional);
+}
+
+
+const TArray<UMovieSceneSection*>& UMovieSceneLevelVisibilityTrack::GetAllSections() const
+{
+	return Sections;
+}
+
+
+TRange<float> UMovieSceneLevelVisibilityTrack::GetSectionBoundaries() const
+{
+	TArray< TRange<float> > Bounds;
+
+	for (int32 SectionIndex = 0; SectionIndex < Sections.Num(); ++SectionIndex)
+	{
+		Bounds.Add(Sections[SectionIndex]->GetRange());
+	}
+
+	return TRange<float>::Hull(Bounds);
+}
+
+
+bool UMovieSceneLevelVisibilityTrack::HasSection(const UMovieSceneSection& Section) const
+{
+	return Sections.Contains(&Section);
+}
+
+
+#if WITH_EDITORONLY_DATA
+FText UMovieSceneLevelVisibilityTrack::GetDefaultDisplayName() const
+{
+	return LOCTEXT("DisplayName", "Level Visibility");
+}
+#endif
+
+#undef LOCTEXT_NAMESPACE
diff --git a/Engine/Source/Runtime/MovieSceneTracks/Private/Tracks/MovieSceneMaterialTrack.cpp b/Engine/Source/Runtime/MovieSceneTracks/Private/Tracks/MovieSceneMaterialTrack.cpp
--- a/Engine/Source/Runtime/MovieSceneTracks/Private/Tracks/MovieSceneMaterialTrack.cpp
+++ b/Engine/Source/Runtime/MovieSceneTracks/Private/Tracks/MovieSceneMaterialTrack.cpp
@@ -10,6 +10,9 @@
 UMovieSceneMaterialTrack::UMovieSceneMaterialTrack(const FObjectInitializer& ObjectInitializer)
 	: Super(ObjectInitializer)
 {
+#if WITH_EDITORONLY_DATA
+	TrackTint = FColor(64,192,64);
+#endif
 }
 
 
diff --git a/Engine/Source/Runtime/MovieSceneTracks/Private/Tracks/MovieSceneParticleParameterTrack.cpp b/Engine/Source/Runtime/MovieSceneTracks/Private/Tracks/MovieSceneParticleParameterTrack.cpp
--- a/Engine/Source/Runtime/MovieSceneTracks/Private/Tracks/MovieSceneParticleParameterTrack.cpp
+++ b/Engine/Source/Runtime/MovieSceneTracks/Private/Tracks/MovieSceneParticleParameterTrack.cpp
@@ -11,6 +11,9 @@
 UMovieSceneParticleParameterTrack::UMovieSceneParticleParameterTrack( const FObjectInitializer& ObjectInitializer )
 	: Super( ObjectInitializer )
 {
+#if WITH_EDITORONLY_DATA
+	TrackTint = FColor(0, 170, 255);
+#endif
 }
 
 TSharedPtr<IMovieSceneTrackInstance> UMovieSceneParticleParameterTrack::CreateInstance()
diff --git a/Engine/Source/Runtime/MovieSceneTracks/Private/Tracks/MovieSceneParticleTrack.cpp b/Engine/Source/Runtime/MovieSceneTracks/Private/Tracks/MovieSceneParticleTrack.cpp
--- a/Engine/Source/Runtime/MovieSceneTracks/Private/Tracks/MovieSceneParticleTrack.cpp
+++ b/Engine/Source/Runtime/MovieSceneTracks/Private/Tracks/MovieSceneParticleTrack.cpp
@@ -12,7 +12,11 @@
 
 UMovieSceneParticleTrack::UMovieSceneParticleTrack( const FObjectInitializer& ObjectInitializer )
 	: Super( ObjectInitializer )
-{ }
+{
+#if WITH_EDITORONLY_DATA
+	TrackTint = FColor(255,255,255);
+#endif
+}
 
 
 TSharedPtr<IMovieSceneTrackInstance> UMovieSceneParticleTrack::CreateInstance()
diff --git a/Engine/Source/Runtime/MovieSceneTracks/Private/Tracks/MovieScenePropertyTrack.cpp b/Engine/Source/Runtime/MovieSceneTracks/Private/Tracks/MovieScenePropertyTrack.cpp
--- a/Engine/Source/Runtime/MovieSceneTracks/Private/Tracks/MovieScenePropertyTrack.cpp
+++ b/Engine/Source/Runtime/MovieSceneTracks/Private/Tracks/MovieScenePropertyTrack.cpp
@@ -69,8 +69,10 @@ bool UMovieScenePropertyTrack::IsEmpty() const
 }
 
 
-UMovieSceneSection* UMovieScenePropertyTrack::FindOrAddSection(float Time)
+UMovieSceneSection* UMovieScenePropertyTrack::FindOrAddSection(float Time, bool& bSectionAdded)
 {
+	bSectionAdded = false;
+
 	// Find a spot for the section so that they are sorted by start time
 	for(int32 SectionIndex = 0; SectionIndex < Sections.Num(); ++SectionIndex)
 	{
@@ -130,6 +132,8 @@ UMovieSceneSection* UMovieScenePropertyTrack::FindOrAddSection(float Time)
 	NewSection->SetEndTime(Time);
 
 	Sections.Add(NewSection);
+	
+	bSectionAdded = true;
 
 	return NewSection;
 }
diff --git a/Engine/Source/Runtime/MovieSceneTracks/Private/Tracks/MovieSceneSkeletalAnimationTrack.cpp b/Engine/Source/Runtime/MovieSceneTracks/Private/Tracks/MovieSceneSkeletalAnimationTrack.cpp
--- a/Engine/Source/Runtime/MovieSceneTracks/Private/Tracks/MovieSceneSkeletalAnimationTrack.cpp
+++ b/Engine/Source/Runtime/MovieSceneTracks/Private/Tracks/MovieSceneSkeletalAnimationTrack.cpp
@@ -15,7 +15,11 @@
 
 UMovieSceneSkeletalAnimationTrack::UMovieSceneSkeletalAnimationTrack(const FObjectInitializer& ObjectInitializer)
 	: Super(ObjectInitializer)
-{ }
+{
+#if WITH_EDITORONLY_DATA
+	TrackTint = FColor(194, 29, 194);
+#endif
+}
 
 
 /* UMovieSceneSkeletalAnimationTrack interface
diff --git a/Engine/Source/Runtime/MovieSceneTracks/Private/Tracks/MovieSceneSpawnTrack.cpp b/Engine/Source/Runtime/MovieSceneTracks/Private/Tracks/MovieSceneSpawnTrack.cpp
--- a/Engine/Source/Runtime/MovieSceneTracks/Private/Tracks/MovieSceneSpawnTrack.cpp
+++ b/Engine/Source/Runtime/MovieSceneTracks/Private/Tracks/MovieSceneSpawnTrack.cpp
@@ -14,7 +14,11 @@
 
 UMovieSceneSpawnTrack::UMovieSceneSpawnTrack(const FObjectInitializer& Obj)
 	: Super(Obj)
-{ }
+{
+#if WITH_EDITORONLY_DATA
+	TrackTint = FColor(255, 0, 234);
+#endif
+}
 
 
 /* UMovieSceneSpawnTrack interface
diff --git a/Engine/Source/Runtime/MovieSceneTracks/Private/Tracks/MovieSceneSubTrack.cpp b/Engine/Source/Runtime/MovieSceneTracks/Private/Tracks/MovieSceneSubTrack.cpp
--- a/Engine/Source/Runtime/MovieSceneTracks/Private/Tracks/MovieSceneSubTrack.cpp
+++ b/Engine/Source/Runtime/MovieSceneTracks/Private/Tracks/MovieSceneSubTrack.cpp
@@ -13,6 +13,13 @@
 
 /* UMovieSceneSubTrack interface
  *****************************************************************************/
+UMovieSceneSubTrack::UMovieSceneSubTrack( const FObjectInitializer& ObjectInitializer )
+	: Super(ObjectInitializer)
+{
+#if WITH_EDITORONLY_DATA
+	TrackTint = FColor(255, 0, 0);
+#endif
+}
 
 UMovieSceneSubSection* UMovieSceneSubTrack::AddSequence(UMovieSceneSequence* Sequence, float StartTime, float Duration, const bool& bInsertSequence)
 {
diff --git a/Engine/Source/Runtime/MovieSceneTracks/Public/Sections/MovieScene3DAttachSection.h b/Engine/Source/Runtime/MovieSceneTracks/Public/Sections/MovieScene3DAttachSection.h
--- a/Engine/Source/Runtime/MovieSceneTracks/Public/Sections/MovieScene3DAttachSection.h
+++ b/Engine/Source/Runtime/MovieSceneTracks/Public/Sections/MovieScene3DAttachSection.h
@@ -35,6 +35,9 @@ public:
 	UPROPERTY(EditAnywhere, Category="Attach")
 	FName AttachSocketName;
 
+	UPROPERTY(EditAnywhere, Category="Attach")
+	FName AttachComponentName;
+
 	UPROPERTY(EditAnywhere, Category="Attach")
 	uint32 bConstrainTx:1;
 
diff --git a/Engine/Source/Runtime/MovieSceneTracks/Public/Sections/MovieScene3DTransformSection.h b/Engine/Source/Runtime/MovieSceneTracks/Public/Sections/MovieScene3DTransformSection.h
--- a/Engine/Source/Runtime/MovieSceneTracks/Public/Sections/MovieScene3DTransformSection.h
+++ b/Engine/Source/Runtime/MovieSceneTracks/Public/Sections/MovieScene3DTransformSection.h
@@ -70,6 +70,63 @@ struct FTransformKey
 };
 
 
+/**
+ * Proxy structure for translation keys in 3D transform sections.
+ */
+USTRUCT()
+struct FMovieScene3DLocationKeyStruct
+	: public FMovieSceneKeyStruct
+{
+	GENERATED_BODY()
+
+	/** They key's translation value. */
+	UPROPERTY(EditAnywhere, Category=Key)
+	FVector Location;
+
+	FRichCurveKey* LocationKeys[3];
+
+	virtual void PropagateChanges(const FPropertyChangedEvent& ChangeEvent) override;
+};
+
+
+/**
+ * Proxy structure for translation keys in 3D transform sections.
+ */
+USTRUCT()
+struct FMovieScene3DRotationKeyStruct
+	: public FMovieSceneKeyStruct
+{
+	GENERATED_BODY()
+
+	/** They key's rotation value. */
+	UPROPERTY(EditAnywhere, Category=Key)
+	FRotator Rotation;
+
+	FRichCurveKey* RotationKeys[3];
+
+	virtual void PropagateChanges(const FPropertyChangedEvent& ChangeEvent) override;
+};
+
+
+/**
+ * Proxy structure for translation keys in 3D transform sections.
+ */
+USTRUCT()
+struct FMovieScene3DScaleKeyStruct
+	: public FMovieSceneKeyStruct
+{
+	GENERATED_BODY()
+
+	/** They key's scale value. */
+	UPROPERTY(EditAnywhere, Category=Key)
+	FVector Scale;
+
+	FRichCurveKey* ScaleKeys[3];
+
+	virtual void PropagateChanges(const FPropertyChangedEvent& ChangeEvent) override;
+};
+
+
 /**
  * Proxy structure for 3D transform section key data.
  */
@@ -81,7 +138,7 @@ struct FMovieScene3DTransformKeyStruct
 
 	/** They key's translation value. */
 	UPROPERTY(EditAnywhere, Category=Key)
-	FVector4 Translation;
+	FVector Location;
 
 	/** They key's rotation value. */
 	UPROPERTY(EditAnywhere, Category=Key)
@@ -89,9 +146,9 @@ struct FMovieScene3DTransformKeyStruct
 
 	/** They key's scale value. */
 	UPROPERTY(EditAnywhere, Category=Key)
-	FVector4 Scale;
+	FVector Scale;
 
-	FRichCurveKey* TranslationKeys[3];
+	FRichCurveKey* LocationKeys[3];
 	FRichCurveKey* RotationKeys[3];
 	FRichCurveKey* ScaleKeys[3];
 
diff --git a/Engine/Source/Runtime/MovieSceneTracks/Public/Sections/MovieSceneAudioSection.h b/Engine/Source/Runtime/MovieSceneTracks/Public/Sections/MovieSceneAudioSection.h
--- a/Engine/Source/Runtime/MovieSceneTracks/Public/Sections/MovieSceneAudioSection.h
+++ b/Engine/Source/Runtime/MovieSceneTracks/Public/Sections/MovieSceneAudioSection.h
@@ -56,6 +56,16 @@ public:
 		return Position >= AudioRange.GetLowerBoundValue() && Position <= AudioRange.GetUpperBoundValue();
 	}
 
+#if WITH_EDITORONLY_DATA
+	/**
+	 * @return Whether to show actual intensity on the waveform thumbnail or not
+	 */
+	bool ShouldShowIntensity() const
+	{
+		return bShowIntensity;
+	}
+#endif
+
 public:
 
 	// MovieSceneSection interface
@@ -78,4 +88,10 @@ private:
 	/** The amount which this audio is time dilated by */
 	UPROPERTY(EditAnywhere, Category="Audio")
 	float AudioDilationFactor;
+
+#if WITH_EDITORONLY_DATA
+	/** Whether to show the actual intensity of the wave on the thumbnail, as well as the smoothed RMS */
+	UPROPERTY(EditAnywhere, Category="Audio")
+	bool bShowIntensity;
+#endif
 };
diff --git a/Engine/Source/Runtime/MovieSceneTracks/Public/Sections/MovieSceneLevelVisibilitySection.h b/Engine/Source/Runtime/MovieSceneTracks/Public/Sections/MovieSceneLevelVisibilitySection.h
new file mode 100644
--- /dev/null
+++ b/Engine/Source/Runtime/MovieSceneTracks/Public/Sections/MovieSceneLevelVisibilitySection.h
@@ -0,0 +1,41 @@
+// Copyright 1998-2015 Epic Games, Inc. All Rights Reserved.
+
+#pragma once
+
+#include "MovieSceneLevelVisibilitySection.generated.h"
+
+/** Visibility options for the level visibility section. */
+UENUM()
+enum class ELevelVisibility
+{
+	/** The streamed levels should be visible. */
+	Visible,
+
+	/** The streamed levels should be hidden. */
+	Hidden
+};
+
+/**
+ * A section for use with the movie scene level visibility track, which controls streamed level visibility.
+ */
+UCLASS()
+class MOVIESCENETRACKS_API UMovieSceneLevelVisibilitySection
+	: public UMovieSceneSection
+{
+	GENERATED_UCLASS_BODY()
+public:
+
+	ELevelVisibility GetVisibility() const;
+	void SetVisibility(ELevelVisibility InVisibility);
+
+	TArray<FName>* GetLevelNames();
+
+private:
+	/** Whether or not the levels in this section should be visible or hidden. */
+	UPROPERTY( EditAnywhere, Category = LevelVisibility )
+	TEnumAsByte<ELevelVisibility> Visibility;
+
+	/** The short names of the levels who's visibility is controlled by this section. */
+	UPROPERTY( EditAnywhere, Category = LevelVisibility )
+	TArray<FName> LevelNames;
+};
\ No newline at end of file
diff --git a/Engine/Source/Runtime/MovieSceneTracks/Public/Sections/MovieSceneParameterSection.h b/Engine/Source/Runtime/MovieSceneTracks/Public/Sections/MovieSceneParameterSection.h
--- a/Engine/Source/Runtime/MovieSceneTracks/Public/Sections/MovieSceneParameterSection.h
+++ b/Engine/Source/Runtime/MovieSceneTracks/Public/Sections/MovieSceneParameterSection.h
@@ -239,9 +239,18 @@ public:
 		TArray<FVectorParameterNameAndValue>& OutVectorValues,
 		TArray<FColorParameterNameAndValue>& OutColorValues ) const;
 
+public:
+
+	// UMovieSceneSection interface
+
+	virtual void DilateSection(float DilationFactor, float Origin, TSet<FKeyHandle>& KeyHandles) override;
+	virtual void GetKeyHandles(TSet<FKeyHandle>& KeyHandles) const override;
+	virtual void MoveSection(float DeltaPosition, TSet<FKeyHandle>& KeyHandles) override;
 
 private:
 	void UpdateParameterIndicesFromRemoval(int32 RemovedIndex);
+	void GatherCurves(TArray<const FRichCurve*> &OutCurves) const;
+	void GatherCurves(TArray<FRichCurve*> &OutCurves);
 
 private:
 	/**
diff --git a/Engine/Source/Runtime/MovieSceneTracks/Public/Sections/MovieSceneParticleSection.h b/Engine/Source/Runtime/MovieSceneTracks/Public/Sections/MovieSceneParticleSection.h
--- a/Engine/Source/Runtime/MovieSceneTracks/Public/Sections/MovieSceneParticleSection.h
+++ b/Engine/Source/Runtime/MovieSceneTracks/Public/Sections/MovieSceneParticleSection.h
@@ -30,7 +30,7 @@ class UMovieSceneParticleSection
 {
 	GENERATED_UCLASS_BODY()
 
-	void AddKey(float Time, EParticleKey::Type KeyType);
+	MOVIESCENETRACKS_API void AddKey(float Time, EParticleKey::Type KeyType);
 
 	MOVIESCENETRACKS_API FIntegralCurve& GetParticleCurve();
 
diff --git a/Engine/Source/Runtime/MovieSceneTracks/Public/TrackInstances/MovieSceneLevelVisibilityTrackInstance.h b/Engine/Source/Runtime/MovieSceneTracks/Public/TrackInstances/MovieSceneLevelVisibilityTrackInstance.h
new file mode 100644
--- /dev/null
+++ b/Engine/Source/Runtime/MovieSceneTracks/Public/TrackInstances/MovieSceneLevelVisibilityTrackInstance.h
@@ -0,0 +1,42 @@
+// Copyright 1998-2015 Epic Games, Inc. All Rights Reserved.
+
+#pragma once
+
+#include "IMovieSceneTrackInstance.h"
+#include "MovieSceneLevelVisibilitySection.h"
+#include "ObjectKey.h"
+
+class UMovieSceneLevelVisibilityTrack;
+
+/**
+ * Controls the runtime behavior of the movie scene level visibility track.
+ */
+class FMovieSceneLevelVisibilityTrackInstance
+	: public IMovieSceneTrackInstance
+{
+public:
+
+	FMovieSceneLevelVisibilityTrackInstance( UMovieSceneLevelVisibilityTrack& InLevelVisibilityTrack );
+
+	/** IMovieSceneTrackInstance interface */
+	virtual void SaveState (const TArray<UObject*>& RuntimeObjects, IMovieScenePlayer& Player, FMovieSceneSequenceInstance& SequenceInstance) override;
+	virtual void RestoreState (const TArray<UObject*>& RuntimeObjects, IMovieScenePlayer& Player, FMovieSceneSequenceInstance& SequenceInstance) override;
+	virtual void Update(EMovieSceneUpdateData& UpdateData, const TArray<UObject*>& RuntimeObjects, IMovieScenePlayer& Player, FMovieSceneSequenceInstance& SequenceInstance) override;
+	virtual void RefreshInstance( const TArray<UObject*>& RuntimeObjects, IMovieScenePlayer& Player, FMovieSceneSequenceInstance& SequenceInstance ) override;
+	virtual void ClearInstance( IMovieScenePlayer& Player, FMovieSceneSequenceInstance& SequenceInstance ) override {}
+
+private:
+
+	void SaveLevelVisibilityState( ULevelStreaming* StreamingLevel );
+	void SetEditorLevelVisibilityState( ULevelStreaming* StreamingLevel, bool bShouldBeVisible );
+
+private:
+
+	/** Track that is being instanced */
+	UMovieSceneLevelVisibilityTrack* LevelVisibilityTrack;
+
+	TMap<FObjectKey, bool> SavedLevelVisibility;
+
+	TMap<FName, TWeakObjectPtr<ULevelStreaming>> NameToLevelMap;
+
+};
diff --git a/Engine/Source/Runtime/MovieSceneTracks/Public/TrackInstances/MovieSceneSkeletalAnimationTrackInstance.h b/Engine/Source/Runtime/MovieSceneTracks/Public/TrackInstances/MovieSceneSkeletalAnimationTrackInstance.h
--- a/Engine/Source/Runtime/MovieSceneTracks/Public/TrackInstances/MovieSceneSkeletalAnimationTrackInstance.h
+++ b/Engine/Source/Runtime/MovieSceneTracks/Public/TrackInstances/MovieSceneSkeletalAnimationTrackInstance.h
@@ -26,8 +26,27 @@ public:
 	virtual void RefreshInstance( const TArray<UObject*>& RuntimeObjects, IMovieScenePlayer& Player, FMovieSceneSequenceInstance& SequenceInstance ) override {}
 	virtual void ClearInstance( IMovieScenePlayer& Player, FMovieSceneSequenceInstance& SequenceInstance ) override {}
 
+private:
+
+	void BeginAnimControl(USkeletalMeshComponent* SkeletalMeshComponent);
+
+	bool CanPlayAnimation(USkeletalMeshComponent* SkeletalMeshComponent, class UAnimSequenceBase* AnimAssetBase = nullptr) const;
+
+	void SetAnimPosition(USkeletalMeshComponent* SkeletalMeshComponent, FName SlotName, int32 ChannelIndex, UAnimSequence* InAnimSequence, float InPosition, bool bFireNotifies, bool bLooping);
+
+	void FinishAnimControl(USkeletalMeshComponent* SkeletalMeshComponent);
+
+	void PreviewBeginAnimControl(USkeletalMeshComponent* SkeletalMeshComponent);
+
+	void PreviewFinishAnimControl(USkeletalMeshComponent* SkeletalMeshComponent);
+
+	void PreviewSetAnimPosition(USkeletalMeshComponent* SkeletalMeshComponent, FName SlotName, int32 ChannelIndex, UAnimSequence* InAnimSequence, float InPosition, bool bLooping, bool bFireNotifies, float DeltaTime);
+
 private:
 
 	/** Track that is being instanced */
 	UMovieSceneSkeletalAnimationTrack* AnimationTrack;
+
+	/** Keep track of the currently playhing montage */
+	TWeakObjectPtr<UAnimMontage> CurrentlyPlayingMontage;
 };
diff --git a/Engine/Source/Runtime/MovieSceneTracks/Public/TrackInstances/MovieSceneSubTrackInstance.h b/Engine/Source/Runtime/MovieSceneTracks/Public/TrackInstances/MovieSceneSubTrackInstance.h
--- a/Engine/Source/Runtime/MovieSceneTracks/Public/TrackInstances/MovieSceneSubTrackInstance.h
+++ b/Engine/Source/Runtime/MovieSceneTracks/Public/TrackInstances/MovieSceneSubTrackInstance.h
@@ -30,6 +30,7 @@ public:
 	virtual void RestoreState (const TArray<UObject*>& RuntimeObjects, IMovieScenePlayer& Player, FMovieSceneSequenceInstance& SequenceInstance) override;
 	virtual void SaveState(const TArray<UObject*>& RuntimeObjects, IMovieScenePlayer& Player, FMovieSceneSequenceInstance& SequenceInstance) override;
 	virtual void Update(EMovieSceneUpdateData& UpdateData, const TArray<UObject*>& RuntimeObjects, IMovieScenePlayer& Player, FMovieSceneSequenceInstance& SequenceInstance) override;
+	virtual EMovieSceneUpdatePass HasUpdatePasses() override { return (EMovieSceneUpdatePass)(MSUP_PreUpdate | MSUP_Update); }
 
 protected:
 
diff --git a/Engine/Source/Runtime/MovieSceneTracks/Public/Tracks/MovieScene3DAttachTrack.h b/Engine/Source/Runtime/MovieSceneTracks/Public/Tracks/MovieScene3DAttachTrack.h
--- a/Engine/Source/Runtime/MovieSceneTracks/Public/Tracks/MovieScene3DAttachTrack.h
+++ b/Engine/Source/Runtime/MovieSceneTracks/Public/Tracks/MovieScene3DAttachTrack.h
@@ -20,7 +20,8 @@ public:
 
 	// UMovieScene3DConstraintTrack interface
 
-	virtual void AddConstraint( float Time, float ConstraintEndTime, const FName SocketName, const FGuid& ConstraintId ) override;
+	virtual void AddConstraint( float Time, float ConstraintEndTime, const FName SocketName, const FName ComponentName, const FGuid& ConstraintId ) override;
+	virtual class UMovieSceneSection* CreateNewSection() override;
 
 public:
 
diff --git a/Engine/Source/Runtime/MovieSceneTracks/Public/Tracks/MovieScene3DConstraintTrack.h b/Engine/Source/Runtime/MovieSceneTracks/Public/Tracks/MovieScene3DConstraintTrack.h
--- a/Engine/Source/Runtime/MovieSceneTracks/Public/Tracks/MovieScene3DConstraintTrack.h
+++ b/Engine/Source/Runtime/MovieSceneTracks/Public/Tracks/MovieScene3DConstraintTrack.h
@@ -26,9 +26,10 @@ public:
 	 * @param Time The time relative to the owning movie scene where the section should be.
 	 * @param ConstraintEndTime Set the constraint to end at this time.
 	 * @param SocketName The socket name for the constraint.
+	 * @param ComponentName The name of the component the socket resides in.
 	 * @param ConstraintId The id to the constraint.
 	 */
-	virtual void AddConstraint(float Time, float ConstraintEndTime, const FName SocketName, const FGuid& ConstraintId) { }
+	virtual void AddConstraint(float Time, float ConstraintEndTime, const FName SocketName, const FName ComponentName, const FGuid& ConstraintId) { }
 
 public:
 
diff --git a/Engine/Source/Runtime/MovieSceneTracks/Public/Tracks/MovieScene3DPathTrack.h b/Engine/Source/Runtime/MovieSceneTracks/Public/Tracks/MovieScene3DPathTrack.h
--- a/Engine/Source/Runtime/MovieSceneTracks/Public/Tracks/MovieScene3DPathTrack.h
+++ b/Engine/Source/Runtime/MovieSceneTracks/Public/Tracks/MovieScene3DPathTrack.h
@@ -19,7 +19,7 @@ public:
 
 	// UMovieScene3DConstraintTrack interface
 
-	virtual void AddConstraint(float Time, float ConstraintEndTime,  const FName SocketName, const FGuid& ConstraintId) override;
+	virtual void AddConstraint(float Time, float ConstraintEndTime, const FName SocketName, const FName ComponentName, const FGuid& ConstraintId) override;
 
 public:
 
diff --git a/Engine/Source/Runtime/MovieSceneTracks/Public/Tracks/MovieSceneCameraCutTrack.h b/Engine/Source/Runtime/MovieSceneTracks/Public/Tracks/MovieSceneCameraCutTrack.h
--- a/Engine/Source/Runtime/MovieSceneTracks/Public/Tracks/MovieSceneCameraCutTrack.h
+++ b/Engine/Source/Runtime/MovieSceneTracks/Public/Tracks/MovieSceneCameraCutTrack.h
@@ -18,7 +18,8 @@ class UMovieSceneCameraCutTrack
 	: public UMovieSceneNameableTrack
 {
 	GENERATED_BODY()
-
+	UMovieSceneCameraCutTrack( const FObjectInitializer& ObjectInitializer );
+	
 public:
 
 	/** 
diff --git a/Engine/Source/Runtime/MovieSceneTracks/Public/Tracks/MovieSceneCinematicShotTrack.h b/Engine/Source/Runtime/MovieSceneTracks/Public/Tracks/MovieSceneCinematicShotTrack.h
--- a/Engine/Source/Runtime/MovieSceneTracks/Public/Tracks/MovieSceneCinematicShotTrack.h
+++ b/Engine/Source/Runtime/MovieSceneTracks/Public/Tracks/MovieSceneCinematicShotTrack.h
@@ -21,6 +21,8 @@ class UMovieSceneCinematicShotTrack
 
 public:
 
+	UMovieSceneCinematicShotTrack(const FObjectInitializer& ObjectInitializer);
+	
 	// UMovieSceneSubTrack interface
 
 	MOVIESCENETRACKS_API virtual UMovieSceneSubSection* AddSequence(UMovieSceneSequence* Sequence, float StartTime, float Duration, const bool& bInsertSequence = false);
diff --git a/Engine/Source/Runtime/MovieSceneTracks/Public/Tracks/MovieSceneEventTrack.h b/Engine/Source/Runtime/MovieSceneTracks/Public/Tracks/MovieSceneEventTrack.h
--- a/Engine/Source/Runtime/MovieSceneTracks/Public/Tracks/MovieSceneEventTrack.h
+++ b/Engine/Source/Runtime/MovieSceneTracks/Public/Tracks/MovieSceneEventTrack.h
@@ -21,7 +21,11 @@ public:
 	UMovieSceneEventTrack()
 		: bFireEventsWhenForwards(true)
 		, bFireEventsWhenBackwards(true)
-	{ }
+	{
+#if WITH_EDITORONLY_DATA
+		TrackTint = FColor(0, 255, 128);
+#endif
+	}
 
 public:
 
diff --git a/Engine/Source/Runtime/MovieSceneTracks/Public/Tracks/MovieSceneFadeTrack.h b/Engine/Source/Runtime/MovieSceneTracks/Public/Tracks/MovieSceneFadeTrack.h
--- a/Engine/Source/Runtime/MovieSceneTracks/Public/Tracks/MovieSceneFadeTrack.h
+++ b/Engine/Source/Runtime/MovieSceneTracks/Public/Tracks/MovieSceneFadeTrack.h
@@ -23,6 +23,6 @@ public:
 	virtual UMovieSceneSection* CreateNewSection() override;
 
 #if WITH_EDITORONLY_DATA
-	virtual FText GetDisplayName() const override;
+	virtual FText GetDefaultDisplayName() const override;
 #endif
 };
diff --git a/Engine/Source/Runtime/MovieSceneTracks/Public/Tracks/MovieScenePropertyTrack.h b/Engine/Source/Runtime/MovieSceneTracks/Public/Tracks/MovieScenePropertyTrack.h
--- a/Engine/Source/Runtime/MovieSceneTracks/Public/Tracks/MovieScenePropertyTrack.h
+++ b/Engine/Source/Runtime/MovieSceneTracks/Public/Tracks/MovieScenePropertyTrack.h
@@ -55,9 +55,10 @@ public:
 	 * Finds a section at the current time.
 	 *
 	 * @param Time	The time relative to the owning movie scene where the section should be
+	 * @param bSectionAdded Whether a section was added or not
 	 * @return The found section, or the new section.
 	 */
-	class UMovieSceneSection* FindOrAddSection(float Time);
+	class UMovieSceneSection* FindOrAddSection(float Time, bool& bSectionAdded);
 
 protected:
 
diff --git a/Engine/Source/Runtime/MovieSceneTracks/Public/Tracks/MovieSceneSubTrack.h b/Engine/Source/Runtime/MovieSceneTracks/Public/Tracks/MovieSceneSubTrack.h
--- a/Engine/Source/Runtime/MovieSceneTracks/Public/Tracks/MovieSceneSubTrack.h
+++ b/Engine/Source/Runtime/MovieSceneTracks/Public/Tracks/MovieSceneSubTrack.h
@@ -22,6 +22,8 @@ class UMovieSceneSubTrack
 
 public:
 
+	UMovieSceneSubTrack( const FObjectInitializer& ObjectInitializer );
+
 	/**
 	 * Adds a movie scene section to the end of all the existing sections. If there are no sections, the new movie scene section is added at the requested time.
 	 *
