Copying //UE4/Dev-Sequencer to //UE4/Main

diff --git a/Engine/Source/Runtime/MovieSceneTracks/Private/Sections/MovieScene3DAttachSection.cpp b/Engine/Source/Runtime/MovieSceneTracks/Private/Sections/MovieScene3DAttachSection.cpp
--- a/Engine/Source/Runtime/MovieSceneTracks/Private/Sections/MovieScene3DAttachSection.cpp
+++ b/Engine/Source/Runtime/MovieSceneTracks/Private/Sections/MovieScene3DAttachSection.cpp
@@ -19,14 +19,14 @@ UMovieScene3DAttachSection::UMovieScene3DAttachSection( const FObjectInitializer
 }
 
 
-void UMovieScene3DAttachSection::Eval( USceneComponent* SceneComponent, float Position, AActor* Actor, FVector& OutTranslation, FRotator& OutRotation ) const
+void UMovieScene3DAttachSection::Eval( USceneComponent* SceneComponent, float Position, AActor* ParentActor, FVector& OutTranslation, FRotator& OutRotation ) const
 {
 	bool bFoundAttachment = false;
 	if(AttachSocketName != NAME_None)
 	{
 		if(AttachComponentName != NAME_None )
 		{
-			TInlineComponentArray<USceneComponent*> PotentialAttachComponents(Actor);
+			TInlineComponentArray<USceneComponent*> PotentialAttachComponents(ParentActor);
 			for(USceneComponent* PotentialAttachComponent : PotentialAttachComponents)
 			{
 				if(PotentialAttachComponent->GetFName() == AttachComponentName && PotentialAttachComponent->DoesSocketExist(AttachSocketName))
@@ -39,9 +39,9 @@ void UMovieScene3DAttachSection::Eval( USceneComponent* SceneComponent, float Po
 				}
 			}
 		}
-		else if (Actor->GetRootComponent()->DoesSocketExist(AttachSocketName))
+		else if (ParentActor->GetRootComponent()->DoesSocketExist(AttachSocketName))
 		{
-			FTransform SocketTransform = Actor->GetRootComponent()->GetSocketTransform(AttachSocketName);
+			FTransform SocketTransform = ParentActor->GetRootComponent()->GetSocketTransform(AttachSocketName);
 			OutTranslation = SocketTransform.GetLocation();
 			OutRotation = SocketTransform.GetRotation().Rotator();
 			bFoundAttachment = true;
@@ -49,7 +49,7 @@ void UMovieScene3DAttachSection::Eval( USceneComponent* SceneComponent, float Po
 	}
 	else if(AttachComponentName != NAME_None )
 	{
-		TInlineComponentArray<USceneComponent*> PotentialAttachComponents(Actor);
+		TInlineComponentArray<USceneComponent*> PotentialAttachComponents(ParentActor);
 		for(USceneComponent* PotentialAttachComponent : PotentialAttachComponents)
 		{
 			if(PotentialAttachComponent->GetFName() == AttachComponentName)
@@ -65,10 +65,18 @@ void UMovieScene3DAttachSection::Eval( USceneComponent* SceneComponent, float Po
 
 	if(!bFoundAttachment)
 	{
-		OutTranslation = Actor->GetActorLocation();
-		OutRotation = Actor->GetActorRotation();
+		OutTranslation = ParentActor->GetActorLocation();
+		OutRotation = ParentActor->GetActorRotation();
 	}
 
+	FTransform ParentTransform(OutRotation, OutTranslation);
+	FTransform RelativeTransform = SceneComponent->GetRelativeTransform();
+
+	FTransform NewRelativeTransform = RelativeTransform * ParentTransform;
+
+	OutTranslation = NewRelativeTransform.GetLocation();
+	OutRotation = NewRelativeTransform.GetRotation().Rotator();
+
 	if (!bConstrainTx)
 	{
 		OutTranslation[0] = SceneComponent->GetRelativeTransform().GetLocation()[0];
@@ -78,10 +86,10 @@ void UMovieScene3DAttachSection::Eval( USceneComponent* SceneComponent, float Po
 	{
 		OutTranslation[1] = SceneComponent->GetRelativeTransform().GetLocation()[1];
 	}
-
+	
 	if (!bConstrainTz)
 	{
-		OutTranslation[1] = SceneComponent->GetRelativeTransform().GetLocation()[2];
+		OutTranslation[2] = SceneComponent->GetRelativeTransform().GetLocation()[2];
 	}
 
 	if (!bConstrainRx)
diff --git a/Engine/Source/Runtime/MovieSceneTracks/Private/Sections/MovieScene3DTransformSection.cpp b/Engine/Source/Runtime/MovieSceneTracks/Private/Sections/MovieScene3DTransformSection.cpp
--- a/Engine/Source/Runtime/MovieSceneTracks/Private/Sections/MovieScene3DTransformSection.cpp
+++ b/Engine/Source/Runtime/MovieSceneTracks/Private/Sections/MovieScene3DTransformSection.cpp
@@ -12,7 +12,10 @@ void FMovieScene3DLocationKeyStruct::PropagateChanges(const FPropertyChangedEven
 {
 	for (int32 Index = 0; Index <= 2; ++Index)
 	{
-		LocationKeys[Index]->Value = Location[Index];
+		if(LocationKeys[Index] != nullptr)
+		{
+			LocationKeys[Index]->Value = Location[Index];
+		}
 	}
 }
 
@@ -22,9 +25,18 @@ void FMovieScene3DLocationKeyStruct::PropagateChanges(const FPropertyChangedEven
 
 void FMovieScene3DRotationKeyStruct::PropagateChanges(const FPropertyChangedEvent& ChangeEvent)
 {
-	RotationKeys[0]->Value = Rotation.Pitch;
-	RotationKeys[1]->Value = Rotation.Yaw;
-	RotationKeys[2]->Value = Rotation.Roll;
+	if(RotationKeys[0] != nullptr)
+	{
+		RotationKeys[0]->Value = Rotation.Pitch;
+	}
+	if(RotationKeys[1] != nullptr)
+	{	
+		RotationKeys[1]->Value = Rotation.Yaw;
+	}
+	if(RotationKeys[2] != nullptr)
+	{
+		RotationKeys[2]->Value = Rotation.Roll;
+	}
 }
 
 
@@ -35,7 +47,10 @@ void FMovieScene3DScaleKeyStruct::PropagateChanges(const FPropertyChangedEvent&
 {
 	for (int32 Index = 0; Index <= 2; ++Index)
 	{
-		ScaleKeys[Index]->Value = Scale[Index];
+		if(ScaleKeys[Index] != nullptr)
+		{
+			ScaleKeys[Index]->Value = Scale[Index];
+		}
 	}
 }
 
@@ -47,13 +62,28 @@ void FMovieScene3DTransformKeyStruct::PropagateChanges(const FPropertyChangedEve
 {
 	for (int32 Index = 0; Index <= 2; ++Index)
 	{
-		LocationKeys[Index]->Value = Location[Index];
-		ScaleKeys[Index]->Value = Scale[Index];
+		if(LocationKeys[Index] != nullptr)
+		{
+			LocationKeys[Index]->Value = Location[Index];
+		}
+		if(ScaleKeys[Index] != nullptr)
+		{
+			ScaleKeys[Index]->Value = Scale[Index];
+		}
 	}
 
-	RotationKeys[0]->Value = Rotation.Pitch;
-	RotationKeys[1]->Value = Rotation.Yaw;
-	RotationKeys[2]->Value = Rotation.Roll;
+	if(RotationKeys[0] != nullptr)
+	{
+		RotationKeys[0]->Value = Rotation.Pitch;
+	}
+	if(RotationKeys[1] != nullptr)
+	{	
+		RotationKeys[1]->Value = Rotation.Yaw;
+	}
+	if(RotationKeys[2] != nullptr)
+	{
+		RotationKeys[2]->Value = Rotation.Roll;
+	}
 }
 
 
@@ -202,65 +232,115 @@ TSharedPtr<FStructOnScope> UMovieScene3DTransformSection::GetKeyStruct(const TAr
 	FRichCurveKey* RotationKeys[3];
 	FRichCurveKey* ScaleKeys[3];
 
+	bool bHasTranslationKeys = false;
+	bool bHasRotationKeys = false;
+	bool bHasScaleKeys = false;
+
 	for (int32 Index = 0; Index <= 2; ++Index)
 	{
 		TranslationKeys[Index] = Translation[Index].GetFirstMatchingKey(KeyHandles);
+		if(TranslationKeys[Index] != nullptr)
+		{
+			bHasTranslationKeys = true;
+		}
 		RotationKeys[Index] = Rotation[Index].GetFirstMatchingKey(KeyHandles);
+		if(RotationKeys[Index] != nullptr)
+		{
+			bHasRotationKeys = true;
+		}
 		ScaleKeys[Index] = Scale[Index].GetFirstMatchingKey(KeyHandles);
+		if(ScaleKeys[Index] != nullptr)
+		{
+			bHasScaleKeys = true;
+		}
+	}
+
+	int32 KeyTypeCount = 0;
+	if(bHasTranslationKeys)
+	{
+		KeyTypeCount++;
+	}
+	if(bHasRotationKeys)
+	{
+		KeyTypeCount++;
+	}
+	if(bHasScaleKeys)
+	{
+		KeyTypeCount++;
 	}
 
-	if ((TranslationKeys[0] != nullptr) && (RotationKeys[0] != nullptr) && (ScaleKeys[0] != nullptr))
+	// do we have multiple keys on multiple parts of the transform?
+	if (KeyTypeCount > 1)
 	{
 		TSharedRef<FStructOnScope> KeyStruct = MakeShareable(new FStructOnScope(FMovieScene3DTransformKeyStruct::StaticStruct()));
 		auto Struct = (FMovieScene3DTransformKeyStruct*)KeyStruct->GetStructMemory();
 		{
 			for (int32 Index = 0; Index <= 2; ++Index)
 			{
-				check(TranslationKeys[Index] != nullptr);
-				Struct->LocationKeys[Index] = TranslationKeys[Index];
-				Struct->Location[Index] = TranslationKeys[Index]->Value;
-
-				check(RotationKeys[Index] != nullptr);
-				Struct->RotationKeys[Index] = RotationKeys[Index];
-
-				check(ScaleKeys[Index] != nullptr);
-				Struct->ScaleKeys[Index] = ScaleKeys[Index];
-				Struct->Scale[Index] = ScaleKeys[Index]->Value;
+				if(TranslationKeys[Index] != nullptr)
+				{
+					Struct->LocationKeys[Index] = TranslationKeys[Index];
+					Struct->Location[Index] = TranslationKeys[Index]->Value;
+				}
+
+				if(RotationKeys[Index] != nullptr)
+				{
+					Struct->RotationKeys[Index] = RotationKeys[Index];
+				}
+				
+				if(ScaleKeys[Index] != nullptr)
+				{
+					Struct->ScaleKeys[Index] = ScaleKeys[Index];
+					Struct->Scale[Index] = ScaleKeys[Index]->Value;
+				}
 			}
 
-			Struct->Rotation.Pitch = RotationKeys[0]->Value;
-			Struct->Rotation.Yaw = RotationKeys[1]->Value;
-			Struct->Rotation.Roll = RotationKeys[2]->Value;
+			if(RotationKeys[0] != nullptr)
+			{
+				Struct->Rotation.Pitch = RotationKeys[0]->Value;
+			}
+			if(RotationKeys[1] != nullptr)
+			{
+				Struct->Rotation.Yaw = RotationKeys[1]->Value;
+			}
+			if(RotationKeys[2] != nullptr)
+			{
+				Struct->Rotation.Roll = RotationKeys[2]->Value;
+			}
 		}
 
 		return KeyStruct;
 	}
 	
-	if (TranslationKeys[0] != nullptr)
+	if (bHasTranslationKeys)
 	{
 		TSharedRef<FStructOnScope> KeyStruct = MakeShareable(new FStructOnScope(FMovieScene3DLocationKeyStruct::StaticStruct()));
 		auto Struct = (FMovieScene3DLocationKeyStruct*)KeyStruct->GetStructMemory();
 		{
 			for (int32 Index = 0; Index <= 2; ++Index)
 			{
-				check(TranslationKeys[Index] != nullptr);
-				Struct->LocationKeys[Index] = TranslationKeys[Index];
-				Struct->Location[Index] = TranslationKeys[Index]->Value;
+				if(TranslationKeys[Index] != nullptr)
+				{
+					Struct->LocationKeys[Index] = TranslationKeys[Index];
+					Struct->Location[Index] = TranslationKeys[Index]->Value;
+				}
 			}
 		}
 
 		return KeyStruct;
 	}
 	
-	if (RotationKeys[0] != nullptr)
+	if (bHasRotationKeys)
 	{
 		TSharedRef<FStructOnScope> KeyStruct = MakeShareable(new FStructOnScope(FMovieScene3DRotationKeyStruct::StaticStruct()));
 		auto Struct = (FMovieScene3DRotationKeyStruct*)KeyStruct->GetStructMemory();
 		{
 			for (int32 Index = 0; Index <= 2; ++Index)
 			{
-				check(RotationKeys[Index] != nullptr);
-				Struct->RotationKeys[Index] = RotationKeys[Index];
+				if(RotationKeys[Index] != nullptr)
+				{
+					Struct->RotationKeys[Index] = RotationKeys[Index];
+				}
 			}
 
 			Struct->Rotation.Pitch = Struct->RotationKeys[0]->Value;
@@ -271,16 +351,18 @@ TSharedPtr<FStructOnScope> UMovieScene3DTransformSection::GetKeyStruct(const TAr
 		return KeyStruct;
 	}
 	
-	if (ScaleKeys[0] != nullptr)
+	if (bHasScaleKeys)
 	{
 		TSharedRef<FStructOnScope> KeyStruct = MakeShareable(new FStructOnScope(FMovieScene3DScaleKeyStruct::StaticStruct()));
 		auto Struct = (FMovieScene3DScaleKeyStruct*)KeyStruct->GetStructMemory();
 		{
 			for (int32 Index = 0; Index <= 2; ++Index)
 			{
-				check(ScaleKeys[Index] != nullptr);
-				Struct->ScaleKeys[Index] = ScaleKeys[Index];
-				Struct->Scale[Index] = ScaleKeys[Index]->Value;
+				if(ScaleKeys[Index] != nullptr)
+				{
+					Struct->ScaleKeys[Index] = ScaleKeys[Index];
+					Struct->Scale[Index] = ScaleKeys[Index]->Value;
+				}
 			}
 		}
 
diff --git a/Engine/Source/Runtime/MovieSceneTracks/Private/Sections/MovieSceneActorReferenceSection.cpp b/Engine/Source/Runtime/MovieSceneTracks/Private/Sections/MovieSceneActorReferenceSection.cpp
new file mode 100644
--- /dev/null
+++ b/Engine/Source/Runtime/MovieSceneTracks/Private/Sections/MovieSceneActorReferenceSection.cpp
@@ -0,0 +1,96 @@
+// Copyright 1998-2016 Epic Games, Inc. All Rights Reserved.
+
+#include "MovieSceneTracksPrivatePCH.h"
+#include "MovieSceneActorReferenceSection.h"
+
+
+UMovieSceneActorReferenceSection::UMovieSceneActorReferenceSection( const FObjectInitializer& ObjectInitializer )
+	: Super( ObjectInitializer )
+{ }
+
+
+FGuid UMovieSceneActorReferenceSection::Eval( float Position ) const
+{
+	int32 ActorGuidIndex = ActorGuidIndexCurve.Evaluate( Position );
+	return ActorGuidIndex >= 0 && ActorGuidIndex < ActorGuids.Num() ? ActorGuids[ActorGuidIndex] : FGuid();
+}
+
+
+void UMovieSceneActorReferenceSection::MoveSection( float DeltaPosition, TSet<FKeyHandle>& KeyHandles )
+{
+	Super::MoveSection( DeltaPosition, KeyHandles );
+
+	// Move the curve
+	ActorGuidIndexCurve.ShiftCurve( DeltaPosition, KeyHandles );
+}
+
+
+void UMovieSceneActorReferenceSection::DilateSection( float DilationFactor, float Origin, TSet<FKeyHandle>& KeyHandles )
+{	
+	Super::DilateSection(DilationFactor, Origin, KeyHandles);
+	
+	ActorGuidIndexCurve.ScaleCurve( Origin, DilationFactor, KeyHandles );
+}
+
+
+void UMovieSceneActorReferenceSection::GetKeyHandles(TSet<FKeyHandle>& KeyHandles) const
+{
+	for ( auto It( ActorGuidIndexCurve.GetKeyHandleIterator() ); It; ++It )
+	{
+		float Time = ActorGuidIndexCurve.GetKeyTime( It.Key() );
+		if (IsTimeWithinSection(Time))
+		{
+			KeyHandles.Add(It.Key());
+		}
+	}
+}
+
+
+void UMovieSceneActorReferenceSection::AddKey( float Time, const FGuid& Value, EMovieSceneKeyInterpolation KeyInterpolation )
+{
+	int32 NewActorGuidIndex = ActorGuids.Add(Value);
+	ActorGuidIndexCurve.AddKey(Time, NewActorGuidIndex);
+}
+
+
+bool UMovieSceneActorReferenceSection::NewKeyIsNewData(float Time, const FGuid& Value) const
+{
+	return Eval(Time) != Value;
+}
+
+
+bool UMovieSceneActorReferenceSection::HasKeys( const FGuid& Value ) const
+{
+	return ActorGuidIndexCurve.GetNumKeys() > 0;
+}
+
+
+void UMovieSceneActorReferenceSection::SetDefault( const FGuid& Value )
+{
+	int32 DefaultIndex = ActorGuids.Add( Value );
+	ActorGuidIndexCurve.SetDefaultValue( DefaultIndex );
+}
+
+
+void UMovieSceneActorReferenceSection::PreSave()
+{
+	ActorGuidStrings.Empty();
+	for ( const FGuid& ActorGuid : ActorGuids )
+	{
+		ActorGuidStrings.Add( ActorGuid.ToString() );
+	}
+	Super::PreSave();
+}
+
+
+void UMovieSceneActorReferenceSection::PostLoad()
+{
+	Super::PostLoad();
+	ActorGuids.Empty();
+	for ( const FString& ActorGuidString : ActorGuidStrings )
+	{
+		FGuid ActorGuid;
+		FGuid::Parse( ActorGuidString, ActorGuid );
+		ActorGuids.Add( ActorGuid );
+	}
+}
diff --git a/Engine/Source/Runtime/MovieSceneTracks/Private/Sections/MovieSceneCameraAnimSection.cpp b/Engine/Source/Runtime/MovieSceneTracks/Private/Sections/MovieSceneCameraAnimSection.cpp
new file mode 100644
--- /dev/null
+++ b/Engine/Source/Runtime/MovieSceneTracks/Private/Sections/MovieSceneCameraAnimSection.cpp
@@ -0,0 +1,49 @@
+// Copyright 1998-2016 Epic Games, Inc. All Rights Reserved.
+
+#include "MovieSceneTracksPrivatePCH.h"
+#include "MovieSceneCameraAnimSection.h"
+#include "MovieSceneCameraAnimTrack.h"
+
+
+UMovieSceneCameraAnimSection::UMovieSceneCameraAnimSection(const FObjectInitializer& ObjectInitializer)
+ 	: Super( ObjectInitializer )
+{
+	PlayRate = 1.f;
+	PlayScale = 1.f;
+	BlendInTime = 0.f;
+	BlendOutTime = 0.f;
+	bLooping = false;
+	bRandomStartTime = false;
+//	Duration = false;
+	PlaySpace = ECameraAnimPlaySpace::CameraLocal;
+	UserDefinedPlaySpace = FRotator::ZeroRotator;
+}
+
+void UMovieSceneCameraAnimSection::MoveSection(float DeltaPosition, TSet<FKeyHandle>& KeyHandles)
+{
+	Super::MoveSection( DeltaPosition, KeyHandles );
+
+	// Move the curve
+	//AnimWeightCurve.ShiftCurve(DeltaPosition, KeyHandles);
+}
+
+
+void UMovieSceneCameraAnimSection::DilateSection(float DilationFactor, float Origin, TSet<FKeyHandle>& KeyHandles)
+{	
+	Super::DilateSection(DilationFactor, Origin, KeyHandles);
+	
+	//AnimWeightCurve.ScaleCurve(Origin, DilationFactor, KeyHandles);
+}
+
+
+void UMovieSceneCameraAnimSection::GetKeyHandles(TSet<FKeyHandle>& KeyHandles) const
+{
+// 	for (auto It(AnimWeightCurve.GetKeyHandleIterator()); It; ++It)
+// 	{
+// 		float Time = AnimWeightCurve.GetKeyTime(It.Key());
+// 		if (IsTimeWithinSection(Time))
+// 		{
+// 			KeyHandles.Add(It.Key());
+// 		}
+// 	}
+}
diff --git a/Engine/Source/Runtime/MovieSceneTracks/Private/Sections/MovieSceneCameraShakeSection.cpp b/Engine/Source/Runtime/MovieSceneTracks/Private/Sections/MovieSceneCameraShakeSection.cpp
new file mode 100644
--- /dev/null
+++ b/Engine/Source/Runtime/MovieSceneTracks/Private/Sections/MovieSceneCameraShakeSection.cpp
@@ -0,0 +1,43 @@
+// Copyright 1998-2016 Epic Games, Inc. All Rights Reserved.
+
+#include "MovieSceneTracksPrivatePCH.h"
+#include "MovieSceneCameraShakeSection.h"
+#include "MovieSceneCameraShakeTrack.h"
+
+
+UMovieSceneCameraShakeSection::UMovieSceneCameraShakeSection(const FObjectInitializer& ObjectInitializer)
+ 	: Super( ObjectInitializer )
+{
+	PlayScale = 1.f;
+	PlaySpace = ECameraAnimPlaySpace::CameraLocal;
+	UserDefinedPlaySpace = FRotator::ZeroRotator;
+}
+
+void UMovieSceneCameraShakeSection::MoveSection(float DeltaPosition, TSet<FKeyHandle>& KeyHandles)
+{
+	Super::MoveSection( DeltaPosition, KeyHandles );
+
+	// Move the curve
+//	ShakeWeightCurve.ShiftCurve(DeltaPosition, KeyHandles);
+}
+
+
+void UMovieSceneCameraShakeSection::DilateSection(float DilationFactor, float Origin, TSet<FKeyHandle>& KeyHandles)
+{	
+	Super::DilateSection(DilationFactor, Origin, KeyHandles);
+	
+//	ShakeWeightCurve.ScaleCurve(Origin, DilationFactor, KeyHandles);
+}
+
+
+void UMovieSceneCameraShakeSection::GetKeyHandles(TSet<FKeyHandle>& KeyHandles) const
+{
+// 	for (auto It(ShakeWeightCurve.GetKeyHandleIterator()); It; ++It)
+// 	{
+// 		float Time = ShakeWeightCurve.GetKeyTime(It.Key());
+// 		if (IsTimeWithinSection(Time))
+// 		{
+// 			KeyHandles.Add(It.Key());
+// 		}
+// 	}
+}
diff --git a/Engine/Source/Runtime/MovieSceneTracks/Private/Sections/MovieSceneSubSection.cpp b/Engine/Source/Runtime/MovieSceneTracks/Private/Sections/MovieSceneSubSection.cpp
--- a/Engine/Source/Runtime/MovieSceneTracks/Private/Sections/MovieSceneSubSection.cpp
+++ b/Engine/Source/Runtime/MovieSceneTracks/Private/Sections/MovieSceneSubSection.cpp
@@ -5,6 +5,7 @@
 #include "MovieSceneSequence.h"
 #include "MovieSceneSubSection.h"
 
+TWeakObjectPtr<UMovieSceneSubSection> UMovieSceneSubSection::TheRecordingSection;
 
 /* UMovieSceneSubSection structors
  *****************************************************************************/
@@ -13,4 +14,88 @@ UMovieSceneSubSection::UMovieSceneSubSection()
 	: StartOffset(0.0f)
 	, TimeScale(1.0f)
 	, PrerollTime(0.0f)
-{ }
+{
+}
+
+void UMovieSceneSubSection::SetSequence(UMovieSceneSequence* Sequence)
+{
+	SubSequence = Sequence;
+
+#if WITH_EDITOR
+	OnSequenceChangedDelegate.ExecuteIfBound(SubSequence);
+#endif
+}
+
+UMovieSceneSequence* UMovieSceneSubSection::GetSequence() const
+{
+	// when recording we need to act as if we have no sequence
+	// the sequence is patched at the end of recording
+	if(GetRecordingSection() == this)
+	{
+		return nullptr;
+	}
+	else
+	{
+		return SubSequence;
+	}
+}
+
+UMovieSceneSubSection* UMovieSceneSubSection::GetRecordingSection()
+{
+	// check if the section is still valid and part of a track (i.e. it has not been deleted or GCed)
+	if(TheRecordingSection.IsValid())
+	{
+		UMovieSceneTrack* TrackOuter = Cast<UMovieSceneTrack>(TheRecordingSection->GetOuter());
+		if(TrackOuter)
+		{
+			if(TrackOuter->HasSection(*TheRecordingSection.Get()))
+			{
+				return TheRecordingSection.Get();
+			}
+		}
+	}
+
+	return nullptr;
+}
+
+void UMovieSceneSubSection::SetAsRecording(bool bRecord)
+{
+	if(bRecord)
+	{
+		TheRecordingSection = this;
+	}
+	else
+	{
+		TheRecordingSection = nullptr;
+	}
+}
+
+bool UMovieSceneSubSection::IsSetAsRecording()
+{
+	return GetRecordingSection() != nullptr;
+}
+
+const FString& UMovieSceneSubSection::GetActorToRecord()
+{
+	UMovieSceneSubSection* RecordingSection = GetRecordingSection();
+	if(RecordingSection)
+	{
+		return RecordingSection->NameOfActorToRecord;
+	}
+
+	static FString EmptyString;
+	return EmptyString;
+}
+
+#if WITH_EDITOR
+void UMovieSceneSubSection::PostEditChangeProperty(FPropertyChangedEvent& PropertyChangedEvent)
+{
+	Super::PostEditChangeProperty(PropertyChangedEvent);
+
+	// recreate runtime instance when sequence is changed
+	if (PropertyChangedEvent.Property && PropertyChangedEvent.Property->GetFName() == GET_MEMBER_NAME_CHECKED(UMovieSceneSubSection, SubSequence))
+	{
+		OnSequenceChangedDelegate.ExecuteIfBound(SubSequence);
+	}
+}
+#endif
\ No newline at end of file
diff --git a/Engine/Source/Runtime/MovieSceneTracks/Private/TrackInstances/MovieScene3DAttachTrackInstance.cpp b/Engine/Source/Runtime/MovieSceneTracks/Private/TrackInstances/MovieScene3DAttachTrackInstance.cpp
--- a/Engine/Source/Runtime/MovieSceneTracks/Private/TrackInstances/MovieScene3DAttachTrackInstance.cpp
+++ b/Engine/Source/Runtime/MovieSceneTracks/Private/TrackInstances/MovieScene3DAttachTrackInstance.cpp
@@ -13,7 +13,7 @@ FMovieScene3DAttachTrackInstance::FMovieScene3DAttachTrackInstance( UMovieScene3
 { }
 
 
-void FMovieScene3DAttachTrackInstance::UpdateConstraint( float Position, const TArray<UObject*>& RuntimeObjects, AActor* Actor, UMovieScene3DConstraintSection* ConstraintSection ) 
+void FMovieScene3DAttachTrackInstance::UpdateConstraint( float Position, const TArray<TWeakObjectPtr<UObject>>& RuntimeObjects, AActor* Actor, UMovieScene3DConstraintSection* ConstraintSection ) 
 {
 	FVector Translation;
 	FRotator Rotation;
@@ -22,7 +22,7 @@ void FMovieScene3DAttachTrackInstance::UpdateConstraint( float Position, const T
 
 	for( int32 ObjIndex = 0; ObjIndex < RuntimeObjects.Num(); ++ObjIndex )
 	{
-		USceneComponent* SceneComponent = MovieSceneHelpers::SceneComponentFromRuntimeObject(RuntimeObjects[ObjIndex]);
+		USceneComponent* SceneComponent = MovieSceneHelpers::SceneComponentFromRuntimeObject(RuntimeObjects[ObjIndex].Get());
 
 		if (SceneComponent != nullptr)
 		{
diff --git a/Engine/Source/Runtime/MovieSceneTracks/Private/TrackInstances/MovieScene3DConstraintTrackInstance.cpp b/Engine/Source/Runtime/MovieSceneTracks/Private/TrackInstances/MovieScene3DConstraintTrackInstance.cpp
--- a/Engine/Source/Runtime/MovieSceneTracks/Private/TrackInstances/MovieScene3DConstraintTrackInstance.cpp
+++ b/Engine/Source/Runtime/MovieSceneTracks/Private/TrackInstances/MovieScene3DConstraintTrackInstance.cpp
@@ -15,35 +15,39 @@ FMovieScene3DConstraintTrackInstance::FMovieScene3DConstraintTrackInstance( UMov
 }
 
 
-void FMovieScene3DConstraintTrackInstance::SaveState(const TArray<UObject*>& RuntimeObjects, IMovieScenePlayer& Player, FMovieSceneSequenceInstance& SequenceInstance)
+void FMovieScene3DConstraintTrackInstance::SaveState(const TArray<TWeakObjectPtr<UObject>>& RuntimeObjects, IMovieScenePlayer& Player, FMovieSceneSequenceInstance& SequenceInstance)
 {
 	for (int32 ObjIndex = 0; ObjIndex < RuntimeObjects.Num(); ++ObjIndex)
 	{
-		USceneComponent* SceneComponent = MovieSceneHelpers::SceneComponentFromRuntimeObject(RuntimeObjects[ObjIndex]);
+		UObject* Object = RuntimeObjects[ObjIndex].Get();
+		USceneComponent* SceneComponent = MovieSceneHelpers::SceneComponentFromRuntimeObject(Object);
+
 		if (SceneComponent != nullptr)
 		{
-			if (InitTransformMap.Find(RuntimeObjects[ObjIndex]) == nullptr)
+			if (InitTransformMap.Find(Object) == nullptr)
 			{
-				InitTransformMap.Add(RuntimeObjects[ObjIndex], SceneComponent->GetRelativeTransform());
+				InitTransformMap.Add(Object, SceneComponent->GetRelativeTransform());
 			}
 		}
 	}
 }
 
 
-void FMovieScene3DConstraintTrackInstance::RestoreState(const TArray<UObject*>& RuntimeObjects, IMovieScenePlayer& Player, FMovieSceneSequenceInstance& SequenceInstance)
+void FMovieScene3DConstraintTrackInstance::RestoreState(const TArray<TWeakObjectPtr<UObject>>& RuntimeObjects, IMovieScenePlayer& Player, FMovieSceneSequenceInstance& SequenceInstance)
 {
 	for (int32 ObjIndex = 0; ObjIndex < RuntimeObjects.Num(); ++ObjIndex)
 	{
-		if (!IsValid(RuntimeObjects[ObjIndex]))
+		UObject* Object = RuntimeObjects[ObjIndex].Get();
+
+		if (!IsValid(Object))
 		{
 			continue;
 		}
 
-		USceneComponent* SceneComponent = MovieSceneHelpers::SceneComponentFromRuntimeObject(RuntimeObjects[ObjIndex]);
+		USceneComponent* SceneComponent = MovieSceneHelpers::SceneComponentFromRuntimeObject(Object);
 		if (SceneComponent != nullptr)
 		{
-			FTransform *Transform = InitTransformMap.Find(RuntimeObjects[ObjIndex]);
+			FTransform *Transform = InitTransformMap.Find(Object);
 			if (Transform != nullptr)
 			{
 				SceneComponent->SetRelativeTransform(*Transform);
@@ -53,7 +57,7 @@ void FMovieScene3DConstraintTrackInstance::RestoreState(const TArray<UObject*>&
 }
 
 
-void FMovieScene3DConstraintTrackInstance::Update(EMovieSceneUpdateData& UpdateData, const TArray<UObject*>& RuntimeObjects, class IMovieScenePlayer& Player, FMovieSceneSequenceInstance& SequenceInstance) 
+void FMovieScene3DConstraintTrackInstance::Update(EMovieSceneUpdateData& UpdateData, const TArray<TWeakObjectPtr<UObject>>& RuntimeObjects, class IMovieScenePlayer& Player, FMovieSceneSequenceInstance& SequenceInstance) 
 {
 	UMovieScene3DConstraintSection* FirstConstraintSection = nullptr;
 
@@ -66,7 +70,7 @@ void FMovieScene3DConstraintTrackInstance::Update(EMovieSceneUpdateData& UpdateD
 		if (ConstraintSection->IsTimeWithinSection(UpdateData.Position) &&
 			(FirstConstraintSection == nullptr || FirstConstraintSection->GetRowIndex() > ConstraintSection->GetRowIndex()))
 		{
-			TArray<UObject*> ConstraintObjects;
+			TArray<TWeakObjectPtr<UObject>> ConstraintObjects;
 			FGuid ConstraintId = ConstraintSection->GetConstraintId();
 
 			if (ConstraintId.IsValid())
@@ -82,7 +86,7 @@ void FMovieScene3DConstraintTrackInstance::Update(EMovieSceneUpdateData& UpdateD
 
 				for (int32 ConstraintObjectIndex = 0; ConstraintObjectIndex < ConstraintObjects.Num(); ++ConstraintObjectIndex)
 				{
-					AActor* Actor = Cast<AActor>(ConstraintObjects[ConstraintObjectIndex]);
+					AActor* Actor = Cast<AActor>(ConstraintObjects[ConstraintObjectIndex].Get());
 					if (Actor)
 					{
 						UpdateConstraint(UpdateData.Position, RuntimeObjects, Actor, ConstraintSection);
diff --git a/Engine/Source/Runtime/MovieSceneTracks/Private/TrackInstances/MovieScene3DPathTrackInstance.cpp b/Engine/Source/Runtime/MovieSceneTracks/Private/TrackInstances/MovieScene3DPathTrackInstance.cpp
--- a/Engine/Source/Runtime/MovieSceneTracks/Private/TrackInstances/MovieScene3DPathTrackInstance.cpp
+++ b/Engine/Source/Runtime/MovieSceneTracks/Private/TrackInstances/MovieScene3DPathTrackInstance.cpp
@@ -14,7 +14,7 @@ FMovieScene3DPathTrackInstance::FMovieScene3DPathTrackInstance( UMovieScene3DPat
 { }
 
 
-void FMovieScene3DPathTrackInstance::UpdateConstraint( float Position, const TArray<UObject*>& RuntimeObjects, AActor* Actor, UMovieScene3DConstraintSection* ConstraintSection ) 
+void FMovieScene3DPathTrackInstance::UpdateConstraint( float Position, const TArray<TWeakObjectPtr<UObject>>& RuntimeObjects, AActor* Actor, UMovieScene3DConstraintSection* ConstraintSection ) 
 {
 	FVector Translation;
 	FRotator Rotation;
@@ -28,7 +28,7 @@ void FMovieScene3DPathTrackInstance::UpdateConstraint( float Position, const TAr
 
 		for( int32 ObjIndex = 0; ObjIndex < RuntimeObjects.Num(); ++ObjIndex )
 		{
-			USceneComponent* SceneComponent = MovieSceneHelpers::SceneComponentFromRuntimeObject(RuntimeObjects[ObjIndex]);
+			USceneComponent* SceneComponent = MovieSceneHelpers::SceneComponentFromRuntimeObject(RuntimeObjects[ObjIndex].Get());
 
 			if (SceneComponent != nullptr)
 			{
diff --git a/Engine/Source/Runtime/MovieSceneTracks/Private/TrackInstances/MovieScene3DTransformTrackInstance.cpp b/Engine/Source/Runtime/MovieSceneTracks/Private/TrackInstances/MovieScene3DTransformTrackInstance.cpp
--- a/Engine/Source/Runtime/MovieSceneTracks/Private/TrackInstances/MovieScene3DTransformTrackInstance.cpp
+++ b/Engine/Source/Runtime/MovieSceneTracks/Private/TrackInstances/MovieScene3DTransformTrackInstance.cpp
@@ -11,45 +11,49 @@ FMovieScene3DTransformTrackInstance::FMovieScene3DTransformTrackInstance( UMovie
 }
 
 
-void FMovieScene3DTransformTrackInstance::SaveState(const TArray<UObject*>& RuntimeObjects, IMovieScenePlayer& Player, FMovieSceneSequenceInstance& SequenceInstance)
+void FMovieScene3DTransformTrackInstance::SaveState(const TArray<TWeakObjectPtr<UObject>>& RuntimeObjects, IMovieScenePlayer& Player, FMovieSceneSequenceInstance& SequenceInstance)
 {
 	for (int32 ObjIndex = 0; ObjIndex < RuntimeObjects.Num(); ++ObjIndex)
 	{
-		USceneComponent* SceneComponent = MovieSceneHelpers::SceneComponentFromRuntimeObject(RuntimeObjects[ObjIndex]);
+		UObject* Object = RuntimeObjects[ObjIndex].Get();
+		USceneComponent* SceneComponent = MovieSceneHelpers::SceneComponentFromRuntimeObject(Object);
+
 		if (SceneComponent != nullptr)
 		{
-			if (InitTransformMap.Find(RuntimeObjects[ObjIndex]) == nullptr)
+			if (InitTransformMap.Find(Object) == nullptr)
 			{
-				InitTransformMap.Add(RuntimeObjects[ObjIndex], SceneComponent->GetRelativeTransform());
+				InitTransformMap.Add(Object, SceneComponent->GetRelativeTransform());
 			}
-			if (InitMobilityMap.Find(RuntimeObjects[ObjIndex]) == nullptr)
+			if (InitMobilityMap.Find(Object) == nullptr)
 			{
-				InitMobilityMap.Add(RuntimeObjects[ObjIndex], SceneComponent->Mobility);
+				InitMobilityMap.Add(Object, SceneComponent->Mobility);
 			}
 		}
 	}
 }
 
 
-void FMovieScene3DTransformTrackInstance::RestoreState(const TArray<UObject*>& RuntimeObjects, IMovieScenePlayer& Player, FMovieSceneSequenceInstance& SequenceInstance)
+void FMovieScene3DTransformTrackInstance::RestoreState(const TArray<TWeakObjectPtr<UObject>>& RuntimeObjects, IMovieScenePlayer& Player, FMovieSceneSequenceInstance& SequenceInstance)
 {
 	for (int32 ObjIndex = 0; ObjIndex < RuntimeObjects.Num(); ++ObjIndex)
 	{
-		if (!IsValid(RuntimeObjects[ObjIndex]))
+		UObject* Object = RuntimeObjects[ObjIndex].Get();
+
+		if (!IsValid(Object))
 		{
 			continue;
 		}
 
-		USceneComponent* SceneComponent = MovieSceneHelpers::SceneComponentFromRuntimeObject(RuntimeObjects[ObjIndex]);
+		USceneComponent* SceneComponent = MovieSceneHelpers::SceneComponentFromRuntimeObject(Object);
 		if (SceneComponent != nullptr)
 		{
-			FTransform *Transform = InitTransformMap.Find(RuntimeObjects[ObjIndex]);
+			FTransform *Transform = InitTransformMap.Find(Object);
 			if (Transform != nullptr)
 			{
 				SceneComponent->SetRelativeTransform(*Transform);
 			}
 
-			EComponentMobility::Type* ComponentMobility = InitMobilityMap.Find(RuntimeObjects[ObjIndex]);
+			EComponentMobility::Type* ComponentMobility = InitMobilityMap.Find(Object);
 			if (ComponentMobility != nullptr)
 			{
 				SceneComponent->SetMobility(*ComponentMobility);
@@ -58,19 +62,20 @@ void FMovieScene3DTransformTrackInstance::RestoreState(const TArray<UObject*>& R
 	}
 }
 
-void FMovieScene3DTransformTrackInstance::UpdateRuntimeMobility(const TArray<UObject*>& RuntimeObjects)
+void FMovieScene3DTransformTrackInstance::UpdateRuntimeMobility(const TArray<TWeakObjectPtr<UObject>>& RuntimeObjects)
 {
 	for( int32 ObjIndex = 0; ObjIndex < RuntimeObjects.Num(); ++ObjIndex )
 	{
-		USceneComponent* SceneComponent = MovieSceneHelpers::SceneComponentFromRuntimeObject(RuntimeObjects[ObjIndex]);
+		UObject* Object = RuntimeObjects[ObjIndex].Get();
+		USceneComponent* SceneComponent = MovieSceneHelpers::SceneComponentFromRuntimeObject(Object);
 
 		if (SceneComponent != nullptr)
 		{
 			if (SceneComponent->Mobility != EComponentMobility::Movable)
 			{
-				if (InitMobilityMap.Find(RuntimeObjects[ObjIndex]) == nullptr)
+				if (InitMobilityMap.Find(Object) == nullptr)
 				{
-					InitMobilityMap.Add(RuntimeObjects[ObjIndex], SceneComponent->Mobility);
+					InitMobilityMap.Add(Object, SceneComponent->Mobility);
 				}
 
 				SceneComponent->SetMobility(EComponentMobility::Movable);
@@ -79,7 +84,7 @@ void FMovieScene3DTransformTrackInstance::UpdateRuntimeMobility(const TArray<UOb
 	}
 }
 
-void FMovieScene3DTransformTrackInstance::Update(EMovieSceneUpdateData& UpdateData, const TArray<UObject*>& RuntimeObjects, class IMovieScenePlayer& Player, FMovieSceneSequenceInstance& SequenceInstance) 
+void FMovieScene3DTransformTrackInstance::Update(EMovieSceneUpdateData& UpdateData, const TArray<TWeakObjectPtr<UObject>>& RuntimeObjects, class IMovieScenePlayer& Player, FMovieSceneSequenceInstance& SequenceInstance) 
 {
 	if (UpdateData.UpdatePass == MSUP_PreUpdate)
 	{
@@ -94,7 +99,8 @@ void FMovieScene3DTransformTrackInstance::Update(EMovieSceneUpdateData& UpdateDa
 	{
 		for( int32 ObjIndex = 0; ObjIndex < RuntimeObjects.Num(); ++ObjIndex )
 		{
-			USceneComponent* SceneComponent = MovieSceneHelpers::SceneComponentFromRuntimeObject(RuntimeObjects[ObjIndex]);
+			UObject* Object = RuntimeObjects[ObjIndex].Get();
+			USceneComponent* SceneComponent = MovieSceneHelpers::SceneComponentFromRuntimeObject(Object);
 
 			if (SceneComponent != nullptr)
 			{
@@ -113,7 +119,7 @@ void FMovieScene3DTransformTrackInstance::Update(EMovieSceneUpdateData& UpdateDa
 	}
 }
  
-void FMovieScene3DTransformTrackInstance::RefreshInstance( const TArray<UObject*>& RuntimeObjects, IMovieScenePlayer& Player, FMovieSceneSequenceInstance& SequenceInstance )
+void FMovieScene3DTransformTrackInstance::RefreshInstance( const TArray<TWeakObjectPtr<UObject>>& RuntimeObjects, IMovieScenePlayer& Player, FMovieSceneSequenceInstance& SequenceInstance )
 {
 	UpdateRuntimeMobility(RuntimeObjects);
 }
diff --git a/Engine/Source/Runtime/MovieSceneTracks/Private/TrackInstances/MovieSceneActorReferenceTrackInstance.cpp b/Engine/Source/Runtime/MovieSceneTracks/Private/TrackInstances/MovieSceneActorReferenceTrackInstance.cpp
new file mode 100644
--- /dev/null
+++ b/Engine/Source/Runtime/MovieSceneTracks/Private/TrackInstances/MovieSceneActorReferenceTrackInstance.cpp
@@ -0,0 +1,94 @@
+// Copyright 1998-2016 Epic Games, Inc. All Rights Reserved.
+
+#include "MovieSceneTracksPrivatePCH.h"
+#include "MovieSceneActorReferenceTrack.h"
+#include "MovieSceneActorReferenceTrackInstance.h"
+
+
+FMovieSceneActorReferenceTrackInstance::FMovieSceneActorReferenceTrackInstance( UMovieSceneActorReferenceTrack& InActorReferenceTrack )
+{
+	ActorReferenceTrack = &InActorReferenceTrack;
+	PropertyBindings = MakeShareable( new FTrackInstancePropertyBindings( ActorReferenceTrack->GetPropertyName(), ActorReferenceTrack->GetPropertyPath() ) );
+}
+
+
+void FMovieSceneActorReferenceTrackInstance::SaveState(const TArray<TWeakObjectPtr<UObject>>& RuntimeObjects, IMovieScenePlayer& Player, FMovieSceneSequenceInstance& SequenceInstance)
+{
+	for (auto ObjectPtr : RuntimeObjects)
+	{
+		UObject* Object = ObjectPtr.Get();
+		if (InitActorReferenceMap.Find(Object) == nullptr)
+		{
+			AActor* ActorReferenceValue = PropertyBindings->GetCurrentValue<AActor*>(Object);
+			InitActorReferenceMap.Add(FObjectKey(Object), ActorReferenceValue);
+		}
+	}
+}
+
+
+void FMovieSceneActorReferenceTrackInstance::RestoreState(const TArray<TWeakObjectPtr<UObject>>& RuntimeObjects, IMovieScenePlayer& Player, FMovieSceneSequenceInstance& SequenceInstance)
+{
+	for (auto ObjectPtr : RuntimeObjects)
+	{
+		UObject* Object = ObjectPtr.Get();
+
+		if (!IsValid(Object))
+		{
+			continue;
+		}
+
+		AActor** ActorReferenceValue = InitActorReferenceMap.Find(FObjectKey(Object));
+		if (ActorReferenceValue != nullptr)
+		{
+			PropertyBindings->CallFunction<AActor*>(Object, ActorReferenceValue);
+		}
+	}
+
+	PropertyBindings->UpdateBindings( RuntimeObjects );
+}
+
+
+void FMovieSceneActorReferenceTrackInstance::Update(EMovieSceneUpdateData& UpdateData, const TArray<TWeakObjectPtr<UObject>>& RuntimeObjects, class IMovieScenePlayer& Player, FMovieSceneSequenceInstance& SequenceInstance ) 
+{
+	
+	FGuid ActorReferenceGuid;
+	if( ActorReferenceTrack->Eval( UpdateData.Position, UpdateData.LastPosition, ActorReferenceGuid ) )
+	{
+		AActor* CurrentActor;
+		AActor** CachedActorPtr = GuidToActorCache.Find( ActorReferenceGuid );
+
+		if (CachedActorPtr != nullptr)
+		{
+			CurrentActor = *CachedActorPtr;
+		}
+		else
+		{
+			TArray<TWeakObjectPtr<UObject>> RuntimeObjectsForGuid;
+			Player.GetRuntimeObjects(SequenceInstance.AsShared(), ActorReferenceGuid, RuntimeObjectsForGuid);
+	
+			if (RuntimeObjectsForGuid.Num() == 1)
+			{
+				AActor* ActorForGuid = Cast<AActor>(RuntimeObjectsForGuid[0].Get());
+
+				if (ActorForGuid != nullptr)
+				{
+					CurrentActor = ActorForGuid;
+					GuidToActorCache.Add(ActorReferenceGuid, CurrentActor);
+				}
+			}
+		}
+
+		for (auto ObjectPtr : RuntimeObjects)
+		{
+			UObject* Object = ObjectPtr.Get();
+			PropertyBindings->CallFunction<AActor*>(Object, &CurrentActor);
+		}
+	}
+}
+
+
+void FMovieSceneActorReferenceTrackInstance::RefreshInstance( const TArray<TWeakObjectPtr<UObject>>& RuntimeObjects, class IMovieScenePlayer& Player, FMovieSceneSequenceInstance& SequenceInstance )
+{
+	GuidToActorCache.Empty();
+	PropertyBindings->UpdateBindings(RuntimeObjects);
+}
diff --git a/Engine/Source/Runtime/MovieSceneTracks/Private/TrackInstances/MovieSceneAudioTrackInstance.cpp b/Engine/Source/Runtime/MovieSceneTracks/Private/TrackInstances/MovieSceneAudioTrackInstance.cpp
--- a/Engine/Source/Runtime/MovieSceneTracks/Private/TrackInstances/MovieSceneAudioTrackInstance.cpp
+++ b/Engine/Source/Runtime/MovieSceneTracks/Private/TrackInstances/MovieSceneAudioTrackInstance.cpp
@@ -15,8 +15,16 @@ FMovieSceneAudioTrackInstance::FMovieSceneAudioTrackInstance( UMovieSceneAudioTr
 	AudioTrack = &InAudioTrack;
 }
 
+void FMovieSceneAudioTrackInstance::SaveState(const TArray<TWeakObjectPtr<UObject>>& RuntimeObjects, IMovieScenePlayer& Player, FMovieSceneSequenceInstance& SequenceInstance) 
+{
+}
+
+void FMovieSceneAudioTrackInstance::RestoreState(const TArray<TWeakObjectPtr<UObject>>& RuntimeObjects, IMovieScenePlayer& Player, FMovieSceneSequenceInstance& SequenceInstance)
+{
+	StopAllSounds();
+}
 
-void FMovieSceneAudioTrackInstance::Update(EMovieSceneUpdateData& UpdateData, const TArray<UObject*>& RuntimeObjects, class IMovieScenePlayer& Player, FMovieSceneSequenceInstance& SequenceInstance) 
+void FMovieSceneAudioTrackInstance::Update(EMovieSceneUpdateData& UpdateData, const TArray<TWeakObjectPtr<UObject>>& RuntimeObjects, class IMovieScenePlayer& Player, FMovieSceneSequenceInstance& SequenceInstance) 
 {
 	// If the playback status is jumping, ie. one such occurrence is setting the time for thumbnail generation, disable audio updates
 	if (Player.GetPlaybackStatus() == EMovieScenePlayerStatus::Jumping)
@@ -50,7 +58,7 @@ void FMovieSceneAudioTrackInstance::Update(EMovieSceneUpdateData& UpdateData, co
 
 				for (int32 ActorIndex = 0; ActorIndex < Actors.Num(); ++ActorIndex)
 				{
-					TWeakObjectPtr<UAudioComponent> Component = GetAudioComponent(Actors[ActorIndex], RowIndex);
+					TWeakObjectPtr<UAudioComponent> Component = GetAudioComponent(Player, Actors[ActorIndex], RowIndex);
 
 					bool bComponentIsPlaying = false;
 					if (Component.IsValid())
@@ -95,7 +103,7 @@ void FMovieSceneAudioTrackInstance::Update(EMovieSceneUpdateData& UpdateData, co
 				
 					for (int32 ActorIndex = 0; ActorIndex < Actors.Num(); ++ActorIndex)
 					{
-						TWeakObjectPtr<UAudioComponent> Component = GetAudioComponent(Actors[ActorIndex], RowIndex);
+						TWeakObjectPtr<UAudioComponent> Component = GetAudioComponent(Player, Actors[ActorIndex], RowIndex);
 						if (Component.IsValid())
 						{
 							if (AudioSection->IsTimeWithinAudioRange(UpdateData.Position) && !Component->IsPlaying())
@@ -124,7 +132,7 @@ void FMovieSceneAudioTrackInstance::Update(EMovieSceneUpdateData& UpdateData, co
 		{
 			for (int32 ActorIndex = 0; ActorIndex < Actors.Num(); ++ActorIndex)
 			{
-				TWeakObjectPtr<UAudioComponent> Component = GetAudioComponent(Actors[ActorIndex], RowIndex);
+				TWeakObjectPtr<UAudioComponent> Component = GetAudioComponent(Player, Actors[ActorIndex], RowIndex);
 				if (Component.IsValid())
 				{
 					if (Component->IsPlaying())
@@ -194,7 +202,7 @@ void FMovieSceneAudioTrackInstance::StopAllSounds()
 }
 
 
-TWeakObjectPtr<UAudioComponent> FMovieSceneAudioTrackInstance::GetAudioComponent(AActor* Actor, int32 RowIndex)
+TWeakObjectPtr<UAudioComponent> FMovieSceneAudioTrackInstance::GetAudioComponent(IMovieScenePlayer& Player, AActor* Actor, int32 RowIndex)
 {
 	if (RowIndex + 1 > PlaybackAudioComponents.Num())
 	{
@@ -209,7 +217,7 @@ TWeakObjectPtr<UAudioComponent> FMovieSceneAudioTrackInstance::GetAudioComponent
 	{
 		USoundCue* TempPlaybackAudioCue = NewObject<USoundCue>();
 		
-		UAudioComponent* AudioComponent = FAudioDevice::CreateComponent(TempPlaybackAudioCue, nullptr, Actor, false, false);
+		UAudioComponent* AudioComponent = FAudioDevice::CreateComponent(TempPlaybackAudioCue, Actor ? Actor->GetWorld() : Player.GetPlaybackContext()->GetWorld(), Actor, false, false);
 
 		PlaybackAudioComponents[RowIndex].Add(Actor);
 		PlaybackAudioComponents[RowIndex][Actor] = TWeakObjectPtr<UAudioComponent>(AudioComponent);
@@ -219,15 +227,17 @@ TWeakObjectPtr<UAudioComponent> FMovieSceneAudioTrackInstance::GetAudioComponent
 }
 
 
-TArray<AActor*> FMovieSceneAudioTrackInstance::GetRuntimeActors(const TArray<UObject*>& RuntimeObjects) const
+TArray<AActor*> FMovieSceneAudioTrackInstance::GetRuntimeActors(const TArray<TWeakObjectPtr<UObject>>& RuntimeObjects) const
 {
 	TArray<AActor*> Actors;
 
 	for (int32 ObjectIndex = 0; ObjectIndex < RuntimeObjects.Num(); ++ObjectIndex)
 	{
-		if (RuntimeObjects[ObjectIndex]->IsA<AActor>())
+		UObject* Object = RuntimeObjects[ObjectIndex].Get();
+
+		if (Object->IsA<AActor>())
 		{
-			Actors.Add(Cast<AActor>(RuntimeObjects[ObjectIndex]));
+			Actors.Add(Cast<AActor>(Object));
 		}
 	}
 
diff --git a/Engine/Source/Runtime/MovieSceneTracks/Private/TrackInstances/MovieSceneBoolTrackInstance.cpp b/Engine/Source/Runtime/MovieSceneTracks/Private/TrackInstances/MovieSceneBoolTrackInstance.cpp
--- a/Engine/Source/Runtime/MovieSceneTracks/Private/TrackInstances/MovieSceneBoolTrackInstance.cpp
+++ b/Engine/Source/Runtime/MovieSceneTracks/Private/TrackInstances/MovieSceneBoolTrackInstance.cpp
@@ -13,10 +13,12 @@ FMovieSceneBoolTrackInstance::FMovieSceneBoolTrackInstance( UMovieSceneBoolTrack
 }
 
 
-void FMovieSceneBoolTrackInstance::SaveState(const TArray<UObject*>& RuntimeObjects, IMovieScenePlayer& Player, FMovieSceneSequenceInstance& SequenceInstance)
+void FMovieSceneBoolTrackInstance::SaveState(const TArray<TWeakObjectPtr<UObject>>& RuntimeObjects, IMovieScenePlayer& Player, FMovieSceneSequenceInstance& SequenceInstance)
 {
-	for( UObject* Object : RuntimeObjects )
+	for (auto ObjectPtr : RuntimeObjects)
 	{
+		UObject* Object = ObjectPtr.Get();
+
 		if (InitBoolMap.Find(Object) == nullptr)
 		{
 			bool BoolValue = PropertyBindings->GetCurrentValue<bool>(Object);
@@ -26,10 +28,12 @@ void FMovieSceneBoolTrackInstance::SaveState(const TArray<UObject*>& RuntimeObje
 }
 
 
-void FMovieSceneBoolTrackInstance::RestoreState(const TArray<UObject*>& RuntimeObjects, IMovieScenePlayer& Player, FMovieSceneSequenceInstance& SequenceInstance)
+void FMovieSceneBoolTrackInstance::RestoreState(const TArray<TWeakObjectPtr<UObject>>& RuntimeObjects, IMovieScenePlayer& Player, FMovieSceneSequenceInstance& SequenceInstance)
 {
-	for( UObject* Object : RuntimeObjects )
+	for (auto ObjectPtr : RuntimeObjects)
 	{
+		UObject* Object = ObjectPtr.Get();
+
 		if (!IsValid(Object))
 		{
 			continue;
@@ -47,21 +51,22 @@ void FMovieSceneBoolTrackInstance::RestoreState(const TArray<UObject*>& RuntimeO
 }
 
 
-void FMovieSceneBoolTrackInstance::Update(EMovieSceneUpdateData& UpdateData, const TArray<UObject*>& RuntimeObjects, class IMovieScenePlayer& Player, FMovieSceneSequenceInstance& SequenceInstance) 
+void FMovieSceneBoolTrackInstance::Update(EMovieSceneUpdateData& UpdateData, const TArray<TWeakObjectPtr<UObject>>& RuntimeObjects, class IMovieScenePlayer& Player, FMovieSceneSequenceInstance& SequenceInstance) 
 {
 	bool BoolValue = false;
 
 	if( BoolTrack->Eval( UpdateData.Position, UpdateData.LastPosition, BoolValue ) )
 	{
-		for( UObject* Object : RuntimeObjects )
+		for (auto ObjectPtr : RuntimeObjects)
 		{
+			UObject* Object = ObjectPtr.Get();
 			PropertyBindings->CallFunction<bool>( Object, &BoolValue );
 		}
 	}
 }
 
 
-void FMovieSceneBoolTrackInstance::RefreshInstance( const TArray<UObject*>& RuntimeObjects, class IMovieScenePlayer& Player, FMovieSceneSequenceInstance& SequenceInstance )
+void FMovieSceneBoolTrackInstance::RefreshInstance( const TArray<TWeakObjectPtr<UObject>>& RuntimeObjects, class IMovieScenePlayer& Player, FMovieSceneSequenceInstance& SequenceInstance )
 {
 	PropertyBindings->UpdateBindings( RuntimeObjects );
 }
diff --git a/Engine/Source/Runtime/MovieSceneTracks/Private/TrackInstances/MovieSceneByteTrackInstance.cpp b/Engine/Source/Runtime/MovieSceneTracks/Private/TrackInstances/MovieSceneByteTrackInstance.cpp
--- a/Engine/Source/Runtime/MovieSceneTracks/Private/TrackInstances/MovieSceneByteTrackInstance.cpp
+++ b/Engine/Source/Runtime/MovieSceneTracks/Private/TrackInstances/MovieSceneByteTrackInstance.cpp
@@ -13,10 +13,12 @@ FMovieSceneByteTrackInstance::FMovieSceneByteTrackInstance( UMovieSceneByteTrack
 }
 
 
-void FMovieSceneByteTrackInstance::SaveState(const TArray<UObject*>& RuntimeObjects, IMovieScenePlayer& Player, FMovieSceneSequenceInstance& SequenceInstance)
+void FMovieSceneByteTrackInstance::SaveState(const TArray<TWeakObjectPtr<UObject>>& RuntimeObjects, IMovieScenePlayer& Player, FMovieSceneSequenceInstance& SequenceInstance)
 {
-	for( UObject* Object : RuntimeObjects )
+	for (auto ObjectPtr : RuntimeObjects)
 	{
+		UObject* Object = ObjectPtr.Get();
+
 		if (InitByteMap.Find(Object) == nullptr)
 		{
 			uint8 ByteValue = PropertyBindings->GetCurrentValue<uint8>(Object);
@@ -26,10 +28,12 @@ void FMovieSceneByteTrackInstance::SaveState(const TArray<UObject*>& RuntimeObje
 }
 
 
-void FMovieSceneByteTrackInstance::RestoreState(const TArray<UObject*>& RuntimeObjects, IMovieScenePlayer& Player, FMovieSceneSequenceInstance& SequenceInstance)
+void FMovieSceneByteTrackInstance::RestoreState(const TArray<TWeakObjectPtr<UObject>>& RuntimeObjects, IMovieScenePlayer& Player, FMovieSceneSequenceInstance& SequenceInstance)
 {
-	for( UObject* Object : RuntimeObjects )
+	for (auto ObjectPtr : RuntimeObjects)
 	{
+		UObject* Object = ObjectPtr.Get();
+
 		if (!IsValid(Object))
 		{
 			continue;
@@ -46,20 +50,21 @@ void FMovieSceneByteTrackInstance::RestoreState(const TArray<UObject*>& RuntimeO
 }
 
 
-void FMovieSceneByteTrackInstance::Update(EMovieSceneUpdateData& UpdateData, const TArray<UObject*>& RuntimeObjects, class IMovieScenePlayer& Player, FMovieSceneSequenceInstance& SequenceInstance) 
+void FMovieSceneByteTrackInstance::Update(EMovieSceneUpdateData& UpdateData, const TArray<TWeakObjectPtr<UObject>>& RuntimeObjects, class IMovieScenePlayer& Player, FMovieSceneSequenceInstance& SequenceInstance) 
 {
 	uint8 ByteValue = 0;
 	if( ByteTrack->Eval( UpdateData.Position, UpdateData.LastPosition, ByteValue ) )
 	{
-		for( UObject* Object : RuntimeObjects )
+		for (auto ObjectPtr : RuntimeObjects )
 		{
-			PropertyBindings->CallFunction<uint8>( Object, &ByteValue );
+			UObject* Object = ObjectPtr.Get();
+			PropertyBindings->CallFunction<uint8>(Object, &ByteValue);
 		}
 	}
 }
 
 
-void FMovieSceneByteTrackInstance::RefreshInstance( const TArray<UObject*>& RuntimeObjects, class IMovieScenePlayer& Player, FMovieSceneSequenceInstance& SequenceInstance )
+void FMovieSceneByteTrackInstance::RefreshInstance( const TArray<TWeakObjectPtr<UObject>>& RuntimeObjects, class IMovieScenePlayer& Player, FMovieSceneSequenceInstance& SequenceInstance )
 {
 	PropertyBindings->UpdateBindings( RuntimeObjects );
 }
diff --git a/Engine/Source/Runtime/MovieSceneTracks/Private/TrackInstances/MovieSceneCameraAnimTrackInstance.cpp b/Engine/Source/Runtime/MovieSceneTracks/Private/TrackInstances/MovieSceneCameraAnimTrackInstance.cpp
new file mode 100644
--- /dev/null
+++ b/Engine/Source/Runtime/MovieSceneTracks/Private/TrackInstances/MovieSceneCameraAnimTrackInstance.cpp
@@ -0,0 +1,177 @@
+// Copyright 1998-2016 Epic Games, Inc. All Rights Reserved.
+
+#include "MovieSceneTracksPrivatePCH.h"
+#include "MovieSceneCommonHelpers.h"
+#include "MovieSceneCameraAnimTrackInstance.h"
+#include "MovieSceneCameraAnimTrack.h"
+#include "MovieSceneCameraAnimSection.h"
+#include "IMovieScenePlayer.h"
+
+
+FMovieSceneCameraAnimTrackInstance::FMovieSceneCameraAnimTrackInstance(UMovieSceneCameraAnimTrack& InCameraAnimTrack)
+{
+	CameraAnimTrack = &InCameraAnimTrack;
+}
+
+FMovieSceneCameraAnimTrackInstance::~FMovieSceneCameraAnimTrackInstance()
+{
+	ACameraActor* const CamActor = TempCameraActor.Get();
+	if (CamActor)
+	{
+		CamActor->Destroy();
+	}
+}
+
+void FMovieSceneCameraAnimTrackInstance::Update(EMovieSceneUpdateData& UpdateData, const TArray<TWeakObjectPtr<UObject>>& RuntimeObjects, IMovieScenePlayer& Player, FMovieSceneSequenceInstance& SequenceInstance)
+{
+	if (UpdateData.UpdatePass == MSUP_PreUpdate)
+	{
+		// pre-update pass, clear out the additive offsets (will reapply as necessary in MSUP_Update)
+		for (TWeakObjectPtr<UObject> ObjectWP : RuntimeObjects)
+		{
+			UCameraComponent* const CameraComponent = MovieSceneHelpers::CameraComponentFromRuntimeObject(ObjectWP.Get());
+			if (CameraComponent)
+			{
+				CameraComponent->ClearAdditiveOffset();
+			}
+		}
+	}
+	else
+	{
+		for (TWeakObjectPtr<UObject> ObjectWP : RuntimeObjects)
+		{
+			UCameraComponent* const CameraComponent = MovieSceneHelpers::CameraComponentFromRuntimeObject(ObjectWP.Get());
+			if (CameraComponent)
+			{
+				// clear it for now, will reset it with correct values later
+				TArray<UMovieSceneCameraAnimSection*> ActiveSections;
+				CameraAnimTrack->GetCameraAnimSectionsAtTime(UpdateData.Position, ActiveSections);
+
+				// look for section instances that are no longer playing
+				TArray<UMovieSceneCameraAnimSection*> DeadSections;
+				for (const auto& Iter : SectionInstanceDataMap)
+				{
+					if (ActiveSections.Contains(Iter.Key) == false)
+					{
+						// section isn't active anymore, stop playing the anim
+						FMovieSceneCameraAnimSectionInstanceData const& SectionInstanceData = Iter.Value;
+						SectionInstanceData.CameraAnimInst->Stop(true);
+						SectionInstanceData.CameraAnimInst->RemoveFromRoot();
+
+						DeadSections.Add(Iter.Key);
+					}
+				}
+				for (auto DeadSection : DeadSections)
+				{
+					SectionInstanceDataMap.Remove(DeadSection);
+				}
+
+				// #todo sort by start time to match application order of player camera
+				ActiveSections.Sort(
+					[](const UMovieSceneCameraAnimSection& One, const UMovieSceneCameraAnimSection& Two)
+				{
+					return One.GetStartTime() < Two.GetStartTime();
+				}
+				);
+
+				// look for sections that just started and need instances
+				for (auto ActiveSection : ActiveSections)
+				{
+					FMovieSceneCameraAnimSectionInstanceData& SectionInstanceData = SectionInstanceDataMap.FindOrAdd(ActiveSection);
+
+					if (SectionInstanceData.CameraAnimInst == nullptr)
+					{
+						// Start playing it
+						UCameraAnimInst* const CamAnimInst = NewObject<UCameraAnimInst>(GetTransientPackage());
+						if (CamAnimInst)
+						{
+							// make it root so GC doesn't take it away
+							CamAnimInst->AddToRoot();
+							CamAnimInst->SetStopAutomatically(false);
+							CamAnimInst->Play(ActiveSection->GetCameraAnim(), GetTempCameraActor(), ActiveSection->PlayRate, ActiveSection->PlayScale, ActiveSection->BlendInTime, ActiveSection->BlendOutTime, ActiveSection->bLooping, ActiveSection->bRandomStartTime);
+						}
+
+						SectionInstanceData.CameraAnimInst = CamAnimInst;
+					}
+
+					if (SectionInstanceData.CameraAnimInst.IsValid())
+					{
+						// prepare temp camera actor by resetting it
+						{
+							ACameraActor* const CameraActor = GetTempCameraActor();
+							CameraActor->SetActorLocationAndRotation(FVector::ZeroVector, FRotator::ZeroRotator);
+
+							ACameraActor const* const DefaultCamActor = GetDefault<ACameraActor>();
+							if (DefaultCamActor)
+							{
+								CameraActor->GetCameraComponent()->AspectRatio = DefaultCamActor->GetCameraComponent()->AspectRatio;
+								CameraActor->GetCameraComponent()->PostProcessSettings = SectionInstanceData.CameraAnimInst->CamAnim->BasePostProcessSettings;
+								CameraActor->GetCameraComponent()->PostProcessBlendWeight = SectionInstanceData.CameraAnimInst->CamAnim->BasePostProcessBlendWeight;
+							}
+						}
+
+						// set camera anim to the correct time
+						float const NewCameraAnimTime = UpdateData.Position - ActiveSection->GetStartTime();
+						SectionInstanceData.CameraAnimInst->SetCurrentTime(NewCameraAnimTime);
+
+						// harvest properties from the actor and apply
+						if (SectionInstanceData.CameraAnimInst->CurrentBlendWeight > 0.f)
+						{
+							FMinimalViewInfo POV;
+							POV.Location = CameraComponent->GetComponentLocation();
+							POV.Rotation = CameraComponent->GetComponentRotation();
+							POV.FOV = CameraComponent->FieldOfView;
+
+							SectionInstanceData.CameraAnimInst->ApplyToView(POV);
+
+							FTransform WorldToBaseCamera = CameraComponent->GetComponentToWorld().Inverse();
+							float BaseFOV = CameraComponent->FieldOfView;
+							FTransform NewCameraToWorld(POV.Rotation, POV.Location);
+							float NewFOV = POV.FOV;
+
+							FTransform NewCameraToBaseCamera = NewCameraToWorld * WorldToBaseCamera;
+
+							float NewFOVToBaseFOV = BaseFOV - NewFOV;
+
+							SectionInstanceData.AdditiveFOVOffset = NewFOVToBaseFOV;
+							SectionInstanceData.AdditiveOffset = NewCameraToBaseCamera;
+						}
+					}
+				}
+
+				// now we have all the offsets, collapse them all and add to the camera
+				FTransform TotalTransform = FTransform::Identity;
+				float TotalFOVOffset = 0.f;
+				for (auto ActiveSection : ActiveSections)
+				{
+					FMovieSceneCameraAnimSectionInstanceData& SectionInstanceData = SectionInstanceDataMap.FindChecked(ActiveSection);
+					TotalTransform = TotalTransform * SectionInstanceData.AdditiveOffset;
+					TotalFOVOffset += SectionInstanceData.AdditiveFOVOffset;
+				}
+
+				CameraComponent->AddAdditiveOffset(TotalTransform, TotalFOVOffset);
+			}
+		}
+	}
+}
+
+ACameraActor* FMovieSceneCameraAnimTrackInstance::GetTempCameraActor()
+{
+	if (TempCameraActor.IsValid() == false)
+	{
+		// spawn the temp CameraActor used for updating CameraAnims
+		FActorSpawnParameters SpawnInfo;
+		SpawnInfo.SpawnCollisionHandlingOverride = ESpawnActorCollisionHandlingMethod::AlwaysSpawn;
+		SpawnInfo.ObjectFlags |= RF_Transient;	// We never want to save these temp actors into a map
+		ACameraActor* const Cam = GWorld->SpawnActor<ACameraActor>(SpawnInfo);		// #fixme, GWorld!
+		if (Cam)
+		{
+#if WITH_EDITOR
+			Cam->SetIsTemporarilyHiddenInEditor(true);
+#endif
+			TempCameraActor = Cam;
+		}
+	}
+
+	return TempCameraActor.Get();
+}
\ No newline at end of file
diff --git a/Engine/Source/Runtime/MovieSceneTracks/Private/TrackInstances/MovieSceneCameraCutTrackInstance.cpp b/Engine/Source/Runtime/MovieSceneTracks/Private/TrackInstances/MovieSceneCameraCutTrackInstance.cpp
--- a/Engine/Source/Runtime/MovieSceneTracks/Private/TrackInstances/MovieSceneCameraCutTrackInstance.cpp
+++ b/Engine/Source/Runtime/MovieSceneTracks/Private/TrackInstances/MovieSceneCameraCutTrackInstance.cpp
@@ -26,7 +26,7 @@ void FMovieSceneCameraCutTrackInstance::ClearInstance(IMovieScenePlayer& Player,
 }
 
 
-void FMovieSceneCameraCutTrackInstance::RefreshInstance(const TArray<UObject*>& RuntimeObjects, IMovieScenePlayer& Player, FMovieSceneSequenceInstance& SequenceInstance)
+void FMovieSceneCameraCutTrackInstance::RefreshInstance(const TArray<TWeakObjectPtr<UObject>>& RuntimeObjects, IMovieScenePlayer& Player, FMovieSceneSequenceInstance& SequenceInstance)
 {
 	const TArray<UMovieSceneSection*>& CameraCutSections = CameraCutTrack->GetAllSections();
 
@@ -41,12 +41,12 @@ void FMovieSceneCameraCutTrackInstance::RefreshInstance(const TArray<UObject*>&
 }
 
 	
-void FMovieSceneCameraCutTrackInstance::RestoreState(const TArray<UObject*>& RuntimeObjects, IMovieScenePlayer& Player, FMovieSceneSequenceInstance& SequenceInstance)
+void FMovieSceneCameraCutTrackInstance::RestoreState(const TArray<TWeakObjectPtr<UObject>>& RuntimeObjects, IMovieScenePlayer& Player, FMovieSceneSequenceInstance& SequenceInstance)
 {
 	Player.UpdateCameraCut(nullptr, nullptr);
 }
 	
-void FMovieSceneCameraCutTrackInstance::Update(EMovieSceneUpdateData& UpdateData, const TArray<UObject*>& RuntimeObjects, class IMovieScenePlayer& Player, FMovieSceneSequenceInstance& SequenceInstance) 
+void FMovieSceneCameraCutTrackInstance::Update(EMovieSceneUpdateData& UpdateData, const TArray<TWeakObjectPtr<UObject>>& RuntimeObjects, class IMovieScenePlayer& Player, FMovieSceneSequenceInstance& SequenceInstance) 
 {
 	// Don't switch CameraCuts when in preroll
 	if (UpdateData.bPreroll)
@@ -75,12 +75,12 @@ void FMovieSceneCameraCutTrackInstance::Update(EMovieSceneUpdateData& UpdateData
 
 	if (CameraObject != LastCameraObject)
 	{
-		Player.UpdateCameraCut(CameraObject, LastCameraObject.Get());
+		Player.UpdateCameraCut(CameraObject, LastCameraObject.Get(), UpdateData.bJumpCut);
 		LastCameraObject = CameraObject;
 	}
 	else if (CameraObject != nullptr)
 	{
-		Player.UpdateCameraCut(CameraObject, nullptr);
+		Player.UpdateCameraCut(CameraObject, nullptr, UpdateData.bJumpCut);
 	}
 }
 
diff --git a/Engine/Source/Runtime/MovieSceneTracks/Private/TrackInstances/MovieSceneCameraShakeTrackInstance.cpp b/Engine/Source/Runtime/MovieSceneTracks/Private/TrackInstances/MovieSceneCameraShakeTrackInstance.cpp
new file mode 100644
--- /dev/null
+++ b/Engine/Source/Runtime/MovieSceneTracks/Private/TrackInstances/MovieSceneCameraShakeTrackInstance.cpp
@@ -0,0 +1,186 @@
+// Copyright 1998-2016 Epic Games, Inc. All Rights Reserved.
+
+#include "MovieSceneTracksPrivatePCH.h"
+#include "MovieSceneCommonHelpers.h"
+#include "MovieSceneCameraShakeTrackInstance.h"
+#include "MovieSceneCameraShakeTrack.h"
+#include "MovieSceneCameraShakeSection.h"
+#include "IMovieScenePlayer.h"
+
+
+FMovieSceneCameraShakeTrackInstance::FMovieSceneCameraShakeTrackInstance(UMovieSceneCameraShakeTrack& InCameraShakeTrack)
+{
+	CameraShakeTrack = &InCameraShakeTrack;
+}
+
+FMovieSceneCameraShakeTrackInstance::~FMovieSceneCameraShakeTrackInstance()
+{
+	ACameraActor* const CamActor = TempCameraActor.Get();
+	if (CamActor)
+	{
+		CamActor->Destroy();
+	}
+}
+
+void FMovieSceneCameraShakeTrackInstance::Update(EMovieSceneUpdateData& UpdateData, const TArray<TWeakObjectPtr<UObject>>& RuntimeObjects, IMovieScenePlayer& Player, FMovieSceneSequenceInstance& SequenceInstance)
+{
+	if (UpdateData.UpdatePass == MSUP_PreUpdate)
+	{
+		// pre-update pass, clear out the additive offsets (will reapply as necessary in MSUP_Update)
+		for (TWeakObjectPtr<UObject> ObjectWP : RuntimeObjects)
+		{
+			UCameraComponent* const CameraComponent = MovieSceneHelpers::CameraComponentFromRuntimeObject(ObjectWP.Get());
+			if (CameraComponent)
+			{
+				CameraComponent->ClearAdditiveOffset();
+			}
+		}
+	}
+	else
+	{
+		for (TWeakObjectPtr<UObject> ObjectWP : RuntimeObjects)
+		{
+			UCameraComponent* const CameraComponent = MovieSceneHelpers::CameraComponentFromRuntimeObject(ObjectWP.Get());
+			if (CameraComponent)
+			{
+				// clear it for now, will reset it with correct values later
+				TArray<UMovieSceneCameraShakeSection*> ActiveSections;
+				CameraShakeTrack->GetCameraShakeSectionsAtTime(UpdateData.Position, ActiveSections);
+
+				// look for section instances that are no longer playing
+				TArray<UMovieSceneCameraShakeSection*> DeadSections;
+				for (const auto& Iter : SectionInstanceDataMap)
+				{
+					if (ActiveSections.Contains(Iter.Key) == false)
+					{
+						// section isn't active anymore, stop playing the anim
+						FMovieSceneCameraShakeSectionInstanceData const& SectionInstanceData = Iter.Value;
+
+						if (SectionInstanceData.CameraShakeInst.IsValid())
+						{
+							SectionInstanceData.CameraShakeInst->StopShake(true);
+							SectionInstanceData.CameraShakeInst->RemoveFromRoot();
+						}
+
+						DeadSections.Add(Iter.Key);
+					}
+				}
+				for (auto DeadSection : DeadSections)
+				{
+					SectionInstanceDataMap.Remove(DeadSection);
+				}
+
+				// #todo sort by start time to match application order of player camera
+				ActiveSections.Sort(
+					[](const UMovieSceneCameraShakeSection& One, const UMovieSceneCameraShakeSection& Two)
+				{
+					return One.GetStartTime() < Two.GetStartTime();
+				}
+				);
+
+				// look for sections that just started and need instances
+				for (auto ActiveSection : ActiveSections)
+				{
+					FMovieSceneCameraShakeSectionInstanceData& SectionInstanceData = SectionInstanceDataMap.FindOrAdd(ActiveSection);
+
+					// see if we need to init and start playing the shake
+					if (SectionInstanceData.CameraShakeInst == nullptr)
+					{
+						UCameraShake* const ShakeInst = NewObject<UCameraShake>(GetTransientPackage(), ActiveSection->GetCameraShakeClass());
+						if (ShakeInst)
+						{
+							// make it root so GC doesn't take it away
+							ShakeInst->AddToRoot();
+							ShakeInst->PlayShake(nullptr, ActiveSection->PlayScale, ActiveSection->PlaySpace, ActiveSection->UserDefinedPlaySpace);
+							if (ShakeInst->AnimInst)
+							{
+								// if there was a cameraanim component of the shake, give it our temp camera actor to use
+								ShakeInst->AnimInst->SetCameraActor(GetTempCameraActor());
+								ShakeInst->AnimInst->SetStopAutomatically(false);
+							}
+						}
+
+						SectionInstanceData.CameraShakeInst = ShakeInst;
+					}
+
+					// advance the shake and apply it
+					if (SectionInstanceData.CameraShakeInst.IsValid())
+					{
+						FMinimalViewInfo POV;
+						POV.Location = CameraComponent->GetComponentLocation();
+						POV.Rotation = CameraComponent->GetComponentRotation();
+						POV.FOV = CameraComponent->FieldOfView;
+
+						// prepare temp camera actor by resetting it
+						{
+							ACameraActor* const CameraActor = GetTempCameraActor();
+							ACameraActor const* const DefaultCamActor = GetDefault<ACameraActor>();
+							if (CameraActor && DefaultCamActor)
+							{
+								CameraActor->SetActorLocationAndRotation(FVector::ZeroVector, FRotator::ZeroRotator);
+
+								CameraActor->GetCameraComponent()->AspectRatio = DefaultCamActor->GetCameraComponent()->AspectRatio;
+
+								if (SectionInstanceData.CameraShakeInst->AnimInst && SectionInstanceData.CameraShakeInst->AnimInst->CamAnim)
+								{
+									CameraActor->GetCameraComponent()->PostProcessSettings = SectionInstanceData.CameraShakeInst->AnimInst->CamAnim->BasePostProcessSettings;
+									CameraActor->GetCameraComponent()->PostProcessBlendWeight = SectionInstanceData.CameraShakeInst->AnimInst->CamAnim->BasePostProcessBlendWeight;
+								}
+							}
+						}
+
+						// set camera anim to the correct time
+						float const NewCameraAnimTime = UpdateData.Position - ActiveSection->GetStartTime();
+						SectionInstanceData.CameraShakeInst->SetCurrentTimeAndApplyShake(NewCameraAnimTime, POV);
+
+						// work out the offset transforms
+						FTransform WorldToBaseCamera = CameraComponent->GetComponentToWorld().Inverse();
+						float BaseFOV = CameraComponent->FieldOfView;
+						FTransform NewCameraToWorld(POV.Rotation, POV.Location);
+						float NewFOV = POV.FOV;
+
+						FTransform NewCameraToBaseCamera = NewCameraToWorld * WorldToBaseCamera;
+						float NewFOVToBaseFOV = BaseFOV - NewFOV;
+
+						SectionInstanceData.AdditiveFOVOffset = NewFOVToBaseFOV;
+						SectionInstanceData.AdditiveOffset = NewCameraToBaseCamera;
+					}
+				}
+
+				// now we have all the offsets, collapse them all and add to the camera
+				FTransform TotalTransform = FTransform::Identity;
+				float TotalFOVOffset = 0.f;
+				for (auto ActiveSection : ActiveSections)
+				{
+					FMovieSceneCameraShakeSectionInstanceData& SectionInstanceData = SectionInstanceDataMap.FindChecked(ActiveSection);
+					TotalTransform = TotalTransform * SectionInstanceData.AdditiveOffset;
+					TotalFOVOffset += SectionInstanceData.AdditiveFOVOffset;
+				}
+
+				CameraComponent->AddAdditiveOffset(TotalTransform, TotalFOVOffset);
+			}
+		}
+	}
+}
+
+
+ACameraActor* FMovieSceneCameraShakeTrackInstance::GetTempCameraActor()
+{
+	if (TempCameraActor.IsValid() == false)
+	{
+		// spawn the temp CameraActor used for updating CameraAnims
+		FActorSpawnParameters SpawnInfo;
+		SpawnInfo.SpawnCollisionHandlingOverride = ESpawnActorCollisionHandlingMethod::AlwaysSpawn;
+		SpawnInfo.ObjectFlags |= RF_Transient;	// We never want to save these temp actors into a map
+		ACameraActor* const Cam = GWorld->SpawnActor<ACameraActor>(SpawnInfo);		// #fixme, GWorld!
+		if (Cam)
+		{
+#if WITH_EDITOR
+			Cam->SetIsTemporarilyHiddenInEditor(true);
+#endif
+			TempCameraActor = Cam;
+		}
+	}
+
+	return TempCameraActor.Get();
+}
\ No newline at end of file
diff --git a/Engine/Source/Runtime/MovieSceneTracks/Private/TrackInstances/MovieSceneCinematicShotTrackInstance.cpp b/Engine/Source/Runtime/MovieSceneTracks/Private/TrackInstances/MovieSceneCinematicShotTrackInstance.cpp
--- a/Engine/Source/Runtime/MovieSceneTracks/Private/TrackInstances/MovieSceneCinematicShotTrackInstance.cpp
+++ b/Engine/Source/Runtime/MovieSceneTracks/Private/TrackInstances/MovieSceneCinematicShotTrackInstance.cpp
@@ -25,13 +25,13 @@ void FMovieSceneCinematicShotTrackInstance::ClearInstance(IMovieScenePlayer& Pla
 }
 
 
-void FMovieSceneCinematicShotTrackInstance::RefreshInstance(const TArray<UObject*>& RuntimeObjects, IMovieScenePlayer& Player, FMovieSceneSequenceInstance& SequenceInstance)
+void FMovieSceneCinematicShotTrackInstance::RefreshInstance(const TArray<TWeakObjectPtr<UObject>>& RuntimeObjects, IMovieScenePlayer& Player, FMovieSceneSequenceInstance& SequenceInstance)
 {
 	FMovieSceneSubTrackInstance::RefreshInstance(RuntimeObjects, Player, SequenceInstance);
 }
 
 	
-void FMovieSceneCinematicShotTrackInstance::RestoreState(const TArray<UObject*>& RuntimeObjects, IMovieScenePlayer& Player, FMovieSceneSequenceInstance& SequenceInstance)
+void FMovieSceneCinematicShotTrackInstance::RestoreState(const TArray<TWeakObjectPtr<UObject>>& RuntimeObjects, IMovieScenePlayer& Player, FMovieSceneSequenceInstance& SequenceInstance)
 {
 	FMovieSceneSubTrackInstance::RestoreState(RuntimeObjects, Player, SequenceInstance);
 
@@ -55,7 +55,7 @@ void FMovieSceneCinematicShotTrackInstance::RestoreState(const TArray<UObject*>&
 	*/
 }
 	
-void FMovieSceneCinematicShotTrackInstance::Update(EMovieSceneUpdateData& UpdateData, const TArray<UObject*>& RuntimeObjects, class IMovieScenePlayer& Player, FMovieSceneSequenceInstance& SequenceInstance) 
+void FMovieSceneCinematicShotTrackInstance::Update(EMovieSceneUpdateData& UpdateData, const TArray<TWeakObjectPtr<UObject>>& RuntimeObjects, class IMovieScenePlayer& Player, FMovieSceneSequenceInstance& SequenceInstance) 
 {
 	FMovieSceneSubTrackInstance::Update(UpdateData, RuntimeObjects, Player, SequenceInstance);
 
@@ -116,4 +116,19 @@ void FMovieSceneCinematicShotTrackInstance::Update(EMovieSceneUpdateData& Update
 		}
 		*/
 	}
-}
\ No newline at end of file
+}
+
+bool FMovieSceneCinematicShotTrackInstance::ShouldEvaluateIfOverlapping(const TArray<UMovieSceneSection*>& TraversedSections, UMovieSceneSection* Section) const
+{
+	// Only evaluate the top most row on overlapping cinematic shot sections. Disregard overlap priority.
+	const bool bShouldRemove = TraversedSections.ContainsByPredicate([=](UMovieSceneSection* OtherSection){
+		if (Section->GetRowIndex() > OtherSection->GetRowIndex() &&
+			Section->GetRange().Overlaps(OtherSection->GetRange()))
+		{
+			return true;
+		}
+		return false;
+	});
+
+	return bShouldRemove;
+}
diff --git a/Engine/Source/Runtime/MovieSceneTracks/Private/TrackInstances/MovieSceneColorTrackInstance.cpp b/Engine/Source/Runtime/MovieSceneTracks/Private/TrackInstances/MovieSceneColorTrackInstance.cpp
--- a/Engine/Source/Runtime/MovieSceneTracks/Private/TrackInstances/MovieSceneColorTrackInstance.cpp
+++ b/Engine/Source/Runtime/MovieSceneTracks/Private/TrackInstances/MovieSceneColorTrackInstance.cpp
@@ -15,10 +15,11 @@ FMovieSceneColorTrackInstance::FMovieSceneColorTrackInstance( UMovieSceneColorTr
 }
 
 
-void FMovieSceneColorTrackInstance::SaveState(const TArray<UObject*>& RuntimeObjects, IMovieScenePlayer& Player, FMovieSceneSequenceInstance& SequenceInstance)
+void FMovieSceneColorTrackInstance::SaveState(const TArray<TWeakObjectPtr<UObject>>& RuntimeObjects, IMovieScenePlayer& Player, FMovieSceneSequenceInstance& SequenceInstance)
 {
-	for( UObject* Object : RuntimeObjects )
+	for (auto ObjectPtr : RuntimeObjects)
 	{
+		UObject* Object = ObjectPtr.Get();
 		
 		if( ColorType == EColorType::Slate )
 		{
@@ -55,10 +56,12 @@ void FMovieSceneColorTrackInstance::SaveState(const TArray<UObject*>& RuntimeObj
 }
 
 
-void FMovieSceneColorTrackInstance::RestoreState(const TArray<UObject*>& RuntimeObjects, IMovieScenePlayer& Player, FMovieSceneSequenceInstance& SequenceInstance)
+void FMovieSceneColorTrackInstance::RestoreState(const TArray<TWeakObjectPtr<UObject>>& RuntimeObjects, IMovieScenePlayer& Player, FMovieSceneSequenceInstance& SequenceInstance)
 {
-	for( UObject* Object : RuntimeObjects )
+	for (auto ObjectPtr : RuntimeObjects )
 	{
+		UObject* Object = ObjectPtr.Get();
+
 		if (!IsValid(Object))
 		{
 			continue;
@@ -91,10 +94,12 @@ void FMovieSceneColorTrackInstance::RestoreState(const TArray<UObject*>& Runtime
 }
 
 
-void FMovieSceneColorTrackInstance::Update(EMovieSceneUpdateData& UpdateData, const TArray<UObject*>& RuntimeObjects, class IMovieScenePlayer& Player, FMovieSceneSequenceInstance& SequenceInstance) 
+void FMovieSceneColorTrackInstance::Update(EMovieSceneUpdateData& UpdateData, const TArray<TWeakObjectPtr<UObject>>& RuntimeObjects, class IMovieScenePlayer& Player, FMovieSceneSequenceInstance& SequenceInstance) 
 {
-	for(UObject* Object : RuntimeObjects)
+	for (auto ObjectPtr : RuntimeObjects)
 	{
+		UObject* Object = ObjectPtr.Get();
+
 		if( ColorType == EColorType::Slate )
 		{
 			FSlateColor ColorValue = PropertyBindings->GetCurrentValue<FSlateColor>(Object);
@@ -124,7 +129,7 @@ void FMovieSceneColorTrackInstance::Update(EMovieSceneUpdateData& UpdateData, co
 }
 
 
-void FMovieSceneColorTrackInstance::RefreshInstance( const TArray<UObject*>& RuntimeObjects, class IMovieScenePlayer& Player, FMovieSceneSequenceInstance& SequenceInstance )
+void FMovieSceneColorTrackInstance::RefreshInstance( const TArray<TWeakObjectPtr<UObject>>& RuntimeObjects, class IMovieScenePlayer& Player, FMovieSceneSequenceInstance& SequenceInstance )
 {
 	if( RuntimeObjects.Num() == 0 )
 	{
@@ -134,9 +139,9 @@ void FMovieSceneColorTrackInstance::RefreshInstance( const TArray<UObject*>& Run
 	PropertyBindings->UpdateBindings( RuntimeObjects );
 
 	// Cache off what type of color this is. Just examine the first object since the property should be the same
-	UProperty* ColorProp = PropertyBindings->GetProperty( RuntimeObjects[0] );
-
+	UProperty* ColorProp = PropertyBindings->GetProperty( RuntimeObjects[0].Get() );
 	const UStructProperty* StructProp = Cast<const UStructProperty>(ColorProp);
+
 	if (StructProp && StructProp->Struct)
 	{
 		FName StructName = StructProp->Struct->GetFName();
diff --git a/Engine/Source/Runtime/MovieSceneTracks/Private/TrackInstances/MovieSceneEventTrackInstance.cpp b/Engine/Source/Runtime/MovieSceneTracks/Private/TrackInstances/MovieSceneEventTrackInstance.cpp
--- a/Engine/Source/Runtime/MovieSceneTracks/Private/TrackInstances/MovieSceneEventTrackInstance.cpp
+++ b/Engine/Source/Runtime/MovieSceneTracks/Private/TrackInstances/MovieSceneEventTrackInstance.cpp
@@ -17,7 +17,7 @@ FMovieSceneEventTrackInstance::FMovieSceneEventTrackInstance(UMovieSceneEventTra
 /* IMovieSceneTrackInstance interface
  *****************************************************************************/
 
-void FMovieSceneEventTrackInstance::Update(EMovieSceneUpdateData& UpdateData, const TArray<UObject*>& RuntimeObjects, IMovieScenePlayer& Player, FMovieSceneSequenceInstance& SequenceInstance)
+void FMovieSceneEventTrackInstance::Update(EMovieSceneUpdateData& UpdateData, const TArray<TWeakObjectPtr<UObject>>& RuntimeObjects, IMovieScenePlayer& Player, FMovieSceneSequenceInstance& SequenceInstance)
 {
 	EventTrack->TriggerEvents(UpdateData.Position, UpdateData.LastPosition);
 }
diff --git a/Engine/Source/Runtime/MovieSceneTracks/Private/TrackInstances/MovieSceneFadeTrackInstance.cpp b/Engine/Source/Runtime/MovieSceneTracks/Private/TrackInstances/MovieSceneFadeTrackInstance.cpp
--- a/Engine/Source/Runtime/MovieSceneTracks/Private/TrackInstances/MovieSceneFadeTrackInstance.cpp
+++ b/Engine/Source/Runtime/MovieSceneTracks/Private/TrackInstances/MovieSceneFadeTrackInstance.cpp
@@ -17,7 +17,7 @@ FMovieSceneFadeTrackInstance::FMovieSceneFadeTrackInstance(UMovieSceneFadeTrack&
 /* IMovieSceneTrackInstance interface
  *****************************************************************************/
 
-void FMovieSceneFadeTrackInstance::RestoreState(const TArray<UObject*>& RuntimeObjects, IMovieScenePlayer& Player, FMovieSceneSequenceInstance& SequenceInstance)
+void FMovieSceneFadeTrackInstance::RestoreState(const TArray<TWeakObjectPtr<UObject>>& RuntimeObjects, IMovieScenePlayer& Player, FMovieSceneSequenceInstance& SequenceInstance)
 {
 	// Reset editor preview/fade
 	EMovieSceneViewportParams ViewportParams;
@@ -34,7 +34,7 @@ void FMovieSceneFadeTrackInstance::RestoreState(const TArray<UObject*>& RuntimeO
 	Player.SetViewportSettings(ViewportParamsMap);
 }
 
-void FMovieSceneFadeTrackInstance::Update(EMovieSceneUpdateData& UpdateData, const TArray<UObject*>& RuntimeObjects, IMovieScenePlayer& Player, FMovieSceneSequenceInstance& SequenceInstance)
+void FMovieSceneFadeTrackInstance::Update(EMovieSceneUpdateData& UpdateData, const TArray<TWeakObjectPtr<UObject>>& RuntimeObjects, IMovieScenePlayer& Player, FMovieSceneSequenceInstance& SequenceInstance)
 {
 	FLinearColor FadeColor = FLinearColor::Black;
 	bool bFadeAudio = false;
diff --git a/Engine/Source/Runtime/MovieSceneTracks/Private/TrackInstances/MovieSceneFloatTrackInstance.cpp b/Engine/Source/Runtime/MovieSceneTracks/Private/TrackInstances/MovieSceneFloatTrackInstance.cpp
--- a/Engine/Source/Runtime/MovieSceneTracks/Private/TrackInstances/MovieSceneFloatTrackInstance.cpp
+++ b/Engine/Source/Runtime/MovieSceneTracks/Private/TrackInstances/MovieSceneFloatTrackInstance.cpp
@@ -13,10 +13,12 @@ FMovieSceneFloatTrackInstance::FMovieSceneFloatTrackInstance( UMovieSceneFloatTr
 }
 
 
-void FMovieSceneFloatTrackInstance::SaveState(const TArray<UObject*>& RuntimeObjects, IMovieScenePlayer& Player, FMovieSceneSequenceInstance& SequenceInstance)
+void FMovieSceneFloatTrackInstance::SaveState(const TArray<TWeakObjectPtr<UObject>>& RuntimeObjects, IMovieScenePlayer& Player, FMovieSceneSequenceInstance& SequenceInstance)
 {
-	for( UObject* Object : RuntimeObjects )
+	for (auto ObjectPtr : RuntimeObjects)
 	{
+		UObject* Object = ObjectPtr.Get();
+
 		if (InitFloatMap.Find(Object) == nullptr)
 		{
 			float FloatValue = PropertyBindings->GetCurrentValue<float>(Object);
@@ -26,10 +28,12 @@ void FMovieSceneFloatTrackInstance::SaveState(const TArray<UObject*>& RuntimeObj
 }
 
 
-void FMovieSceneFloatTrackInstance::RestoreState(const TArray<UObject*>& RuntimeObjects, IMovieScenePlayer& Player, FMovieSceneSequenceInstance& SequenceInstance)
+void FMovieSceneFloatTrackInstance::RestoreState(const TArray<TWeakObjectPtr<UObject>>& RuntimeObjects, IMovieScenePlayer& Player, FMovieSceneSequenceInstance& SequenceInstance)
 {
-	for( UObject* Object : RuntimeObjects )
+	for (auto ObjectPtr : RuntimeObjects)
 	{
+		UObject* Object = ObjectPtr.Get();
+
 		if (!IsValid(Object))
 		{
 			continue;
@@ -46,20 +50,21 @@ void FMovieSceneFloatTrackInstance::RestoreState(const TArray<UObject*>& Runtime
 }
 
 
-void FMovieSceneFloatTrackInstance::Update(EMovieSceneUpdateData& UpdateData, const TArray<UObject*>& RuntimeObjects, class IMovieScenePlayer& Player, FMovieSceneSequenceInstance& SequenceInstance ) 
+void FMovieSceneFloatTrackInstance::Update(EMovieSceneUpdateData& UpdateData, const TArray<TWeakObjectPtr<UObject>>& RuntimeObjects, class IMovieScenePlayer& Player, FMovieSceneSequenceInstance& SequenceInstance ) 
 {
 	float FloatValue = 0.0f;
 	if( FloatTrack->Eval( UpdateData.Position, UpdateData.LastPosition, FloatValue ) )
 	{
-		for(UObject* Object : RuntimeObjects)
+		for(auto ObjectPtr : RuntimeObjects)
 		{
+			UObject* Object = ObjectPtr.Get();
 			PropertyBindings->CallFunction<float>( Object, &FloatValue );
 		}
 	}
 }
 
 
-void FMovieSceneFloatTrackInstance::RefreshInstance( const TArray<UObject*>& RuntimeObjects, class IMovieScenePlayer& Player, FMovieSceneSequenceInstance& SequenceInstance )
+void FMovieSceneFloatTrackInstance::RefreshInstance( const TArray<TWeakObjectPtr<UObject>>& RuntimeObjects, class IMovieScenePlayer& Player, FMovieSceneSequenceInstance& SequenceInstance )
 {
 	PropertyBindings->UpdateBindings( RuntimeObjects );
 }
diff --git a/Engine/Source/Runtime/MovieSceneTracks/Private/TrackInstances/MovieSceneLevelVisibilityTrackInstance.cpp b/Engine/Source/Runtime/MovieSceneTracks/Private/TrackInstances/MovieSceneLevelVisibilityTrackInstance.cpp
--- a/Engine/Source/Runtime/MovieSceneTracks/Private/TrackInstances/MovieSceneLevelVisibilityTrackInstance.cpp
+++ b/Engine/Source/Runtime/MovieSceneTracks/Private/TrackInstances/MovieSceneLevelVisibilityTrackInstance.cpp
@@ -70,7 +70,7 @@ void GetAnimatedStreamingLevels( UMovieSceneLevelVisibilityTrack* LevelVisibilit
 }
 
 
-void FMovieSceneLevelVisibilityTrackInstance::SaveState(const TArray<UObject*>& RuntimeObjects, IMovieScenePlayer& Player, FMovieSceneSequenceInstance& SequenceInstance)
+void FMovieSceneLevelVisibilityTrackInstance::SaveState(const TArray<TWeakObjectPtr<UObject>>& RuntimeObjects, IMovieScenePlayer& Player, FMovieSceneSequenceInstance& SequenceInstance)
 {
 	TArray<ULevelStreaming*> AnimatedStreamingLevels;
 	GetAnimatedStreamingLevels(LevelVisibilityTrack, AnimatedStreamingLevels);
@@ -82,7 +82,7 @@ void FMovieSceneLevelVisibilityTrackInstance::SaveState(const TArray<UObject*>&
 }
 
 
-void FMovieSceneLevelVisibilityTrackInstance::RestoreState(const TArray<UObject*>& RuntimeObjects, IMovieScenePlayer& Player, FMovieSceneSequenceInstance& SequenceInstance)
+void FMovieSceneLevelVisibilityTrackInstance::RestoreState(const TArray<TWeakObjectPtr<UObject>>& RuntimeObjects, IMovieScenePlayer& Player, FMovieSceneSequenceInstance& SequenceInstance)
 {
 	TArray<ULevelStreaming*> AnimatedStreamingLevels;
 	GetAnimatedStreamingLevels( LevelVisibilityTrack, AnimatedStreamingLevels );
@@ -100,7 +100,7 @@ void FMovieSceneLevelVisibilityTrackInstance::RestoreState(const TArray<UObject*
 }
 
 
-void FMovieSceneLevelVisibilityTrackInstance::Update(EMovieSceneUpdateData& UpdateData, const TArray<UObject*>& RuntimeObjects, class IMovieScenePlayer& Player, FMovieSceneSequenceInstance& SequenceInstance ) 
+void FMovieSceneLevelVisibilityTrackInstance::Update(EMovieSceneUpdateData& UpdateData, const TArray<TWeakObjectPtr<UObject>>& RuntimeObjects, class IMovieScenePlayer& Player, FMovieSceneSequenceInstance& SequenceInstance ) 
 {
 	for ( UMovieSceneSection* Section : LevelVisibilityTrack->GetAllSections() )
 	{
@@ -158,7 +158,7 @@ void FMovieSceneLevelVisibilityTrackInstance::Update(EMovieSceneUpdateData& Upda
 }
 
 
-void FMovieSceneLevelVisibilityTrackInstance::RefreshInstance( const TArray<UObject*>& RuntimeObjects, class IMovieScenePlayer& Player, FMovieSceneSequenceInstance& SequenceInstance )
+void FMovieSceneLevelVisibilityTrackInstance::RefreshInstance( const TArray<TWeakObjectPtr<UObject>>& RuntimeObjects, class IMovieScenePlayer& Player, FMovieSceneSequenceInstance& SequenceInstance )
 {
 	NameToLevelMap.Empty();
 
diff --git a/Engine/Source/Runtime/MovieSceneTracks/Private/TrackInstances/MovieSceneMaterialTrackInstance.cpp b/Engine/Source/Runtime/MovieSceneTracks/Private/TrackInstances/MovieSceneMaterialTrackInstance.cpp
--- a/Engine/Source/Runtime/MovieSceneTracks/Private/TrackInstances/MovieSceneMaterialTrackInstance.cpp
+++ b/Engine/Source/Runtime/MovieSceneTracks/Private/TrackInstances/MovieSceneMaterialTrackInstance.cpp
@@ -9,10 +9,12 @@ FMovieSceneMaterialTrackInstance::FMovieSceneMaterialTrackInstance( UMovieSceneM
 	MaterialTrack = &InMaterialTrack;
 }
 
-void FMovieSceneMaterialTrackInstance::SaveState(const TArray<UObject*>& RuntimeObjects, IMovieScenePlayer& Player, FMovieSceneSequenceInstance& SequenceInstance)
+void FMovieSceneMaterialTrackInstance::SaveState(const TArray<TWeakObjectPtr<UObject>>& RuntimeObjects, IMovieScenePlayer& Player, FMovieSceneSequenceInstance& SequenceInstance)
 {
-	for( UObject* Object : RuntimeObjects )
+	for (auto ObjectPtr : RuntimeObjects)
 	{
+		UObject* Object = ObjectPtr.Get();
+
 		if (RuntimeObjectToOriginalMaterialMap.Find(FObjectKey(Object)) == nullptr)
 		{
 			UMaterialInterface* CurrentMaterial = GetMaterialForObject(Object);
@@ -36,10 +38,12 @@ void FMovieSceneMaterialTrackInstance::SaveState(const TArray<UObject*>& Runtime
 	}
 }
 
-void FMovieSceneMaterialTrackInstance::RestoreState(const TArray<UObject*>& RuntimeObjects, IMovieScenePlayer& Player, FMovieSceneSequenceInstance& SequenceInstance)
+void FMovieSceneMaterialTrackInstance::RestoreState(const TArray<TWeakObjectPtr<UObject>>& RuntimeObjects, IMovieScenePlayer& Player, FMovieSceneSequenceInstance& SequenceInstance)
 {
-	for( UObject* Object : RuntimeObjects )
+	for (auto ObjectPtr : RuntimeObjects)
 	{
+		UObject* Object = ObjectPtr.Get();
+
 		if (!IsValid(Object))
 		{
 			continue;
@@ -53,7 +57,7 @@ void FMovieSceneMaterialTrackInstance::RestoreState(const TArray<UObject*>& Runt
 	}
 }
 
-void FMovieSceneMaterialTrackInstance::Update(EMovieSceneUpdateData& UpdateData, const TArray<UObject*>& RuntimeObjects, class IMovieScenePlayer& Player, FMovieSceneSequenceInstance& SequenceInstance ) 
+void FMovieSceneMaterialTrackInstance::Update(EMovieSceneUpdateData& UpdateData, const TArray<TWeakObjectPtr<UObject>>& RuntimeObjects, class IMovieScenePlayer& Player, FMovieSceneSequenceInstance& SequenceInstance ) 
 {
 	TArray<FScalarParameterNameAndValue> ScalarValues;
 	TArray<FColorParameterNameAndValue> VectorValues;
@@ -81,19 +85,23 @@ void FMovieSceneMaterialTrackInstance::Update(EMovieSceneUpdateData& UpdateData,
 	}
 }
 
-void FMovieSceneMaterialTrackInstance::RefreshInstance( const TArray<UObject*>& RuntimeObjects, class IMovieScenePlayer& Player, FMovieSceneSequenceInstance& SequenceInstance )
+void FMovieSceneMaterialTrackInstance::RefreshInstance( const TArray<TWeakObjectPtr<UObject>>& RuntimeObjects, class IMovieScenePlayer& Player, FMovieSceneSequenceInstance& SequenceInstance )
 {
 	DynamicMaterialInstances.Empty();
-	for ( UObject* Object : RuntimeObjects )
+
+	for (auto ObjectPtr : RuntimeObjects)
 	{
+		UObject* Object = ObjectPtr.Get();
 		UMaterialInterface* Material = GetMaterialForObject( Object );
 		UMaterialInstanceDynamic* DynamicMaterialInstance = Cast<UMaterialInstanceDynamic>(Material);
-		if ( DynamicMaterialInstance == nullptr )
+
+		if (DynamicMaterialInstance == nullptr )
 		{
 			DynamicMaterialInstance = UMaterialInstanceDynamic::Create(Material, Object);
 			SetMaterialForObject( Object, DynamicMaterialInstance );
 			DynamicMaterialToOriginalMaterialMap.Add( FObjectKey( DynamicMaterialInstance ), Material );
 		}
+
 		DynamicMaterialInstances.Add( DynamicMaterialInstance );
 	}
 }
diff --git a/Engine/Source/Runtime/MovieSceneTracks/Private/TrackInstances/MovieSceneParticleParameterTrackInstance.cpp b/Engine/Source/Runtime/MovieSceneTracks/Private/TrackInstances/MovieSceneParticleParameterTrackInstance.cpp
--- a/Engine/Source/Runtime/MovieSceneTracks/Private/TrackInstances/MovieSceneParticleParameterTrackInstance.cpp
+++ b/Engine/Source/Runtime/MovieSceneTracks/Private/TrackInstances/MovieSceneParticleParameterTrackInstance.cpp
@@ -9,11 +9,13 @@ FMovieSceneParticleParameterTrackInstance::FMovieSceneParticleParameterTrackInst
 	ParticleParameterTrack = &InParticleParameterTrack;
 }
 
-void FMovieSceneParticleParameterTrackInstance::SaveState(const TArray<UObject*>& RuntimeObjects, IMovieScenePlayer& Player, FMovieSceneSequenceInstance& SequenceInstance)
+void FMovieSceneParticleParameterTrackInstance::SaveState(const TArray<TWeakObjectPtr<UObject>>& RuntimeObjects, IMovieScenePlayer& Player, FMovieSceneSequenceInstance& SequenceInstance)
 {
-	for( UObject* Object : RuntimeObjects )
+	for (auto ObjectPtr : RuntimeObjects)
 	{
+		UObject* Object = ObjectPtr.Get();
 		AEmitter* EmitterActor = Cast<AEmitter>(Object);
+
 		if ( EmitterActor != nullptr )
 		{
 			UParticleSystemComponent* ParticleSystemComponent = EmitterActor->GetParticleSystemComponent();
@@ -52,11 +54,13 @@ void FMovieSceneParticleParameterTrackInstance::SaveState(const TArray<UObject*>
 	}
 }
 
-void FMovieSceneParticleParameterTrackInstance::RestoreState(const TArray<UObject*>& RuntimeObjects, IMovieScenePlayer& Player, FMovieSceneSequenceInstance& SequenceInstance)
+void FMovieSceneParticleParameterTrackInstance::RestoreState(const TArray<TWeakObjectPtr<UObject>>& RuntimeObjects, IMovieScenePlayer& Player, FMovieSceneSequenceInstance& SequenceInstance)
 {
-	for ( UObject* Object : RuntimeObjects )
+	for (auto ObjectPtr : RuntimeObjects)
 	{
+		UObject* Object = ObjectPtr.Get();
 		TSharedPtr<FInitialParameterValues>* InitialParameterValuesPtr = ObjectToInitialParameterValuesMap.Find(FObjectKey(Object));
+
 		if ( InitialParameterValuesPtr != nullptr )
 		{
 			TSharedPtr<FInitialParameterValues> InitialParameterValues = *InitialParameterValuesPtr;
@@ -110,7 +114,7 @@ void FMovieSceneParticleParameterTrackInstance::RestoreState(const TArray<UObjec
 	ObjectToInitialParameterValuesMap.Empty();
 }
 
-void FMovieSceneParticleParameterTrackInstance::Update(EMovieSceneUpdateData& UpdateData, const TArray<UObject*>& RuntimeObjects, class IMovieScenePlayer& Player, FMovieSceneSequenceInstance& SequenceInstance)
+void FMovieSceneParticleParameterTrackInstance::Update(EMovieSceneUpdateData& UpdateData, const TArray<TWeakObjectPtr<UObject>>& RuntimeObjects, class IMovieScenePlayer& Player, FMovieSceneSequenceInstance& SequenceInstance)
 {
 	TArray<FScalarParameterNameAndValue> ScalarNamesAndValues;
 	TArray<FVectorParameterNameAndValue> VectorNamesAndValues;
@@ -133,7 +137,7 @@ void FMovieSceneParticleParameterTrackInstance::Update(EMovieSceneUpdateData& Up
 	}
 }
 
-void FMovieSceneParticleParameterTrackInstance::RefreshInstance( const TArray<UObject*>& RuntimeObjects, class IMovieScenePlayer& Player, FMovieSceneSequenceInstance& SequenceInstance )
+void FMovieSceneParticleParameterTrackInstance::RefreshInstance( const TArray<TWeakObjectPtr<UObject>>& RuntimeObjects, class IMovieScenePlayer& Player, FMovieSceneSequenceInstance& SequenceInstance )
 {
 	ParticleSystemComponents.Empty();
 
@@ -146,9 +150,9 @@ void FMovieSceneParticleParameterTrackInstance::RefreshInstance( const TArray<UO
 	}
 
 	// Cache the particle components which have matching parameter names.
-	for ( UObject* Object : RuntimeObjects )
+	for (auto ObjectPtr : RuntimeObjects)
 	{
-		AEmitter* EmitterActor = Cast<AEmitter>( Object );
+		AEmitter* EmitterActor = Cast<AEmitter>(ObjectPtr.Get());
 		if ( EmitterActor != nullptr )
 		{
 			UParticleSystemComponent* ParticleSystemComponent = EmitterActor->GetParticleSystemComponent();
diff --git a/Engine/Source/Runtime/MovieSceneTracks/Private/TrackInstances/MovieSceneParticleTrackInstance.cpp b/Engine/Source/Runtime/MovieSceneTracks/Private/TrackInstances/MovieSceneParticleTrackInstance.cpp
--- a/Engine/Source/Runtime/MovieSceneTracks/Private/TrackInstances/MovieSceneParticleTrackInstance.cpp
+++ b/Engine/Source/Runtime/MovieSceneTracks/Private/TrackInstances/MovieSceneParticleTrackInstance.cpp
@@ -24,7 +24,7 @@ static UParticleSystemComponent* ParticleSystemComponentFromRuntimeObject(UObjec
 	}
 }
 
-void FMovieSceneParticleTrackInstance::Update(EMovieSceneUpdateData& UpdateData, const TArray<UObject*>& RuntimeObjects, class IMovieScenePlayer& Player, FMovieSceneSequenceInstance& SequenceInstance) 
+void FMovieSceneParticleTrackInstance::Update(EMovieSceneUpdateData& UpdateData, const TArray<TWeakObjectPtr<UObject>>& RuntimeObjects, class IMovieScenePlayer& Player, FMovieSceneSequenceInstance& SequenceInstance) 
 {
 	// @todo Sequencer We need something analagous to Matinee 1's particle replay tracks
 	// What we have here is simple toggling/triggering
@@ -45,7 +45,7 @@ void FMovieSceneParticleTrackInstance::Update(EMovieSceneUpdateData& UpdateData,
 				if ( ParticleKeyCurve.IsKeyHandleValid( PreviousHandle ) )
 				{
 					FIntegralKey& PreviousKey = ParticleKeyCurve.GetKey( PreviousHandle );
-					if ( PreviousKey.Time > UpdateData.LastPosition )
+					if ( PreviousKey.Time >= UpdateData.LastPosition )
 					{
 						ParticleKey = (EParticleKey::Type)PreviousKey.Value;
 						bKeyFound = true;
@@ -58,7 +58,7 @@ void FMovieSceneParticleTrackInstance::Update(EMovieSceneUpdateData& UpdateData,
 		{
 			for (int32 i = 0; i < RuntimeObjects.Num(); ++i)
 			{
-				UParticleSystemComponent* ParticleSystemComponent = ParticleSystemComponentFromRuntimeObject(RuntimeObjects[i]);
+				UParticleSystemComponent* ParticleSystemComponent = ParticleSystemComponentFromRuntimeObject(RuntimeObjects[i].Get());
 
 				if ( ParticleSystemComponent != nullptr )
 				{
@@ -86,12 +86,14 @@ void FMovieSceneParticleTrackInstance::Update(EMovieSceneUpdateData& UpdateData,
 	{
 		for (int32 i = 0; i < RuntimeObjects.Num(); ++i)
 		{
-			AEmitter* Emitter = Cast<AEmitter>(RuntimeObjects[i]);
-			if (Emitter)
+			UObject* Object = RuntimeObjects[i].Get();
+			AEmitter* Emitter = Cast<AEmitter>(Object);
+
+			if (Emitter != nullptr)
 			{
 				Emitter->Deactivate();
 			}
-			else if(UParticleSystemComponent* Component =  Cast<UParticleSystemComponent>(RuntimeObjects[i]))
+			else if (UParticleSystemComponent* Component =  Cast<UParticleSystemComponent>(Object))
 			{
 				Component->SetActive(false, true);
 			}
