Copying //UE4/Dev-Sequencer to //UE4/Main

diff --git a/Engine/Source/Runtime/MovieSceneTracks/Private/TrackInstances/MovieSceneSkeletalAnimationTrackInstance.cpp b/Engine/Source/Runtime/MovieSceneTracks/Private/TrackInstances/MovieSceneSkeletalAnimationTrackInstance.cpp
--- a/Engine/Source/Runtime/MovieSceneTracks/Private/TrackInstances/MovieSceneSkeletalAnimationTrackInstance.cpp
+++ b/Engine/Source/Runtime/MovieSceneTracks/Private/TrackInstances/MovieSceneSkeletalAnimationTrackInstance.cpp
@@ -8,40 +8,24 @@
 #include "Matinee/MatineeAnimInterface.h"
 #include "MovieSceneCommonHelpers.h"
 
-struct FAnimEvalTimes
+float MapTimeToAnimation(float ThisPosition, UMovieSceneSkeletalAnimationSection* AnimSection)
 {
-	float Current, Previous;
+	ThisPosition = FMath::Clamp(ThisPosition, AnimSection->GetStartTime(), AnimSection->GetEndTime());
 
-	/** Map the specified times onto the specified anim section */
-	static FAnimEvalTimes MapTimesToAnimation(float ThisPosition, float LastPosition, UMovieSceneSkeletalAnimationSection* AnimSection)
-	{
-		ThisPosition = FMath::Clamp(ThisPosition, AnimSection->GetStartTime(), AnimSection->GetEndTime());
-		LastPosition = FMath::Clamp(LastPosition, AnimSection->GetStartTime(), AnimSection->GetEndTime());
-
-		ThisPosition -= 1 / 1000.0f;
-
-		float AnimPlayRate = FMath::IsNearlyZero(AnimSection->GetPlayRate()) ? 1.0f : AnimSection->GetPlayRate();
-		float SeqLength = AnimSection->GetSequenceLength() - (AnimSection->GetStartOffset() + AnimSection->GetEndOffset());
-
-		ThisPosition = (ThisPosition - AnimSection->GetStartTime()) * AnimPlayRate;
-		ThisPosition = FMath::Fmod(ThisPosition, SeqLength);
-		ThisPosition += AnimSection->GetStartOffset();
-		if (AnimSection->GetReverse())
-		{
-			ThisPosition = (SeqLength - (ThisPosition - AnimSection->GetStartOffset())) + AnimSection->GetStartOffset();
-		}
+	ThisPosition -= 1 / 1000.0f;
 
-		LastPosition = (LastPosition - AnimSection->GetStartTime()) * AnimPlayRate;
-		LastPosition = FMath::Fmod(LastPosition, SeqLength);
-		LastPosition += AnimSection->GetStartOffset();
-		if (AnimSection->GetReverse())
-		{
-			LastPosition = (SeqLength - (LastPosition - AnimSection->GetStartOffset())) + AnimSection->GetStartOffset();
-		}
+	float AnimPlayRate = FMath::IsNearlyZero(AnimSection->GetPlayRate()) ? 1.0f : AnimSection->GetPlayRate();
+	float SeqLength = AnimSection->GetSequenceLength() - (AnimSection->GetStartOffset() + AnimSection->GetEndOffset());
 
-		FAnimEvalTimes EvalTimes = { ThisPosition, LastPosition };
-		return EvalTimes;
+	ThisPosition = (ThisPosition - AnimSection->GetStartTime()) * AnimPlayRate;
+	ThisPosition = FMath::Fmod(ThisPosition, SeqLength);
+	ThisPosition += AnimSection->GetStartOffset();
+	if (AnimSection->GetReverse())
+	{
+		ThisPosition = (SeqLength - (ThisPosition - AnimSection->GetStartOffset())) + AnimSection->GetStartOffset();
 	}
+
+	return ThisPosition;
 };
 
 FMovieSceneSkeletalAnimationTrackInstance::FMovieSceneSkeletalAnimationTrackInstance( UMovieSceneSkeletalAnimationTrack& InAnimationTrack )
@@ -52,10 +36,83 @@ FMovieSceneSkeletalAnimationTrackInstance::FMovieSceneSkeletalAnimationTrackInst
 
 FMovieSceneSkeletalAnimationTrackInstance::~FMovieSceneSkeletalAnimationTrackInstance()
 {
-	// @todo Sequencer Need to find some way to call PreviewFinishAnimControl (needs the runtime objects)
 }
 
-void FMovieSceneSkeletalAnimationTrackInstance::Update(EMovieSceneUpdateData& UpdateData, const TArray<UObject*>& RuntimeObjects, class IMovieScenePlayer& Player, FMovieSceneSequenceInstance& SequenceInstance) 
+void FMovieSceneSkeletalAnimationTrackInstance::SaveState(const TArray<TWeakObjectPtr<UObject>>& RuntimeObjects, IMovieScenePlayer& Player, FMovieSceneSequenceInstance& SequenceInstance)
+{
+	for (int32 i = 0; i < RuntimeObjects.Num(); ++i)
+	{
+		// get the skeletal mesh component we want to control
+		USkeletalMeshComponent* Component = nullptr;
+		UObject* RuntimeObject = RuntimeObjects[i].Get();
+
+		if (RuntimeObject != nullptr)
+		{
+			// first try to control the component directly
+			USkeletalMeshComponent* SkeletalMeshComponent = Cast<USkeletalMeshComponent>(RuntimeObject);
+			if (SkeletalMeshComponent == nullptr)
+			{
+				// then check to see if we are controlling an actor & if so use its first USkeletalMeshComponent 
+				AActor* Actor = Cast<AActor>(RuntimeObject);
+				if(Actor != nullptr)
+				{
+					SkeletalMeshComponent = Actor->FindComponentByClass<USkeletalMeshComponent>();
+				}
+			}
+
+			if (SkeletalMeshComponent)
+			{
+				if (GIsEditor && !RuntimeObject->GetWorld()->HasBegunPlay())
+				{
+					PreviewBeginAnimControl(SkeletalMeshComponent);
+				}
+				else
+				{
+					BeginAnimControl(SkeletalMeshComponent);
+				}
+			}
+		}
+	}
+}
+
+void FMovieSceneSkeletalAnimationTrackInstance::RestoreState(const TArray<TWeakObjectPtr<UObject>>& RuntimeObjects, IMovieScenePlayer& Player, FMovieSceneSequenceInstance& SequenceInstance)
+{
+	for (int32 i = 0; i < RuntimeObjects.Num(); ++i)
+	{
+		// get the skeletal mesh component we want to control
+		USkeletalMeshComponent* Component = nullptr;
+		UObject* RuntimeObject = RuntimeObjects[i].Get();
+
+		if (RuntimeObject != nullptr)
+		{
+			// first try to control the component directly
+			USkeletalMeshComponent* SkeletalMeshComponent = Cast<USkeletalMeshComponent>(RuntimeObject);
+			if (SkeletalMeshComponent == nullptr)
+			{
+				// then check to see if we are controlling an actor & if so use its first USkeletalMeshComponent 
+				AActor* Actor = Cast<AActor>(RuntimeObject);
+				if(Actor != nullptr)
+				{
+					SkeletalMeshComponent = Actor->FindComponentByClass<USkeletalMeshComponent>();
+				}
+			}
+
+			if (SkeletalMeshComponent)
+			{
+				if (GIsEditor && !RuntimeObject->GetWorld()->HasBegunPlay())
+				{
+					PreviewFinishAnimControl(SkeletalMeshComponent);
+				}
+				else
+				{
+					FinishAnimControl(SkeletalMeshComponent);
+				}
+			}
+		}
+	}
+}
+
+void FMovieSceneSkeletalAnimationTrackInstance::Update(EMovieSceneUpdateData& UpdateData, const TArray<TWeakObjectPtr<UObject>>& RuntimeObjects, class IMovieScenePlayer& Player, FMovieSceneSequenceInstance& SequenceInstance) 
 {
 	// @todo Sequencer gameplay update has a different code path than editor update for animation
 
@@ -63,13 +120,13 @@ void FMovieSceneSkeletalAnimationTrackInstance::Update(EMovieSceneUpdateData& Up
 	{
 		// get the skeletal mesh component we want to control
 		USkeletalMeshComponent* Component = nullptr;
+		UObject* RuntimeObject = RuntimeObjects[i].Get();
 
-		UObject* RuntimeObject = RuntimeObjects[i];
-		if(RuntimeObject != nullptr)
+		if (RuntimeObject != nullptr)
 		{
 			// first try to control the component directly
 			USkeletalMeshComponent* SkeletalMeshComponent = Cast<USkeletalMeshComponent>(RuntimeObject);
-			if(SkeletalMeshComponent == nullptr)
+			if (SkeletalMeshComponent == nullptr)
 			{
 				// then check to see if we are controlling an actor & if so use its first USkeletalMeshComponent 
 				AActor* Actor = Cast<AActor>(RuntimeObject);
@@ -79,7 +136,7 @@ void FMovieSceneSkeletalAnimationTrackInstance::Update(EMovieSceneUpdateData& Up
 				}
 			}
 
-			if(SkeletalMeshComponent)
+			if (SkeletalMeshComponent)
 			{
 				UMovieSceneSkeletalAnimationSection* AnimSection = Cast<UMovieSceneSkeletalAnimationSection>(AnimationTrack->GetAnimSectionAtTime(UpdateData.Position));
 			
@@ -99,35 +156,34 @@ void FMovieSceneSkeletalAnimationTrackInstance::Update(EMovieSceneUpdateData& Up
 
 					if (AnimSequence)
 					{
-						const FAnimEvalTimes EvalTimes = FAnimEvalTimes::MapTimesToAnimation(UpdateData.Position, UpdateData.LastPosition, AnimSection);
+						float EvalTime = MapTimeToAnimation(UpdateData.Position, AnimSection);
 
 						int32 ChannelIndex = 0;
 
-						// NOTE: The order of the bLooping and bFireNotifiers parameters are swapped in the preview and non-preview SetAnimPosition functions.
 						const bool bLooping = false;
-						if (GIsEditor && !GWorld->HasBegunPlay())
+
+						// we also use PreviewSetAnimPosition in PIE when not playing, as we can preview in PIE
+						const bool bIsNotInPIEOrNotPlaying = !RuntimeObject->GetWorld()->HasBegunPlay() || Player.GetPlaybackStatus() != EMovieScenePlayerStatus::Playing;
+						if (GIsEditor && bIsNotInPIEOrNotPlaying)
 						{
-							// Pass an absolute delta time to the preview anim evaluation
-							float DeltaTime = FMath::Abs(EvalTimes.Current - EvalTimes.Previous);
-						
 							// If the playback status is jumping, ie. one such occurrence is setting the time for thumbnail generation, disable anim notifies updates because it could fire audio
 							const bool bFireNotifies = Player.GetPlaybackStatus() == EMovieScenePlayerStatus::Jumping || Player.GetPlaybackStatus() == EMovieScenePlayerStatus::Stopped ? false : true;
 
+							float DeltaTime = UpdateData.Position > UpdateData.LastPosition ? UpdateData.Position - UpdateData.LastPosition : 0.f;
+
 							// When jumping from one cut to another cut, the delta time should be 0 so that anim notifies before the current position are not evaluated. Note, anim notifies at the current time should still be evaluated.
 							if (UpdateData.bJumpCut)
 							{
 								DeltaTime = 0.f;
 							}
 						
-							PreviewBeginAnimControl(SkeletalMeshComponent);
-							PreviewSetAnimPosition(SkeletalMeshComponent, AnimSection->GetSlotName(), ChannelIndex, AnimSequence, EvalTimes.Current, bLooping, bFireNotifies, DeltaTime);
+							PreviewSetAnimPosition(SkeletalMeshComponent, AnimSection->GetSlotName(), ChannelIndex, AnimSequence, EvalTime, bLooping, bFireNotifies, DeltaTime, Player.GetPlaybackStatus() == EMovieScenePlayerStatus::Playing);
 						}
 						else
 						{
 							// Don't fire notifies at runtime since they will be fired through the standard animation tick path.
 							const bool bFireNotifies = false;
-							BeginAnimControl(SkeletalMeshComponent);
-							SetAnimPosition(SkeletalMeshComponent, AnimSection->GetSlotName(), ChannelIndex, AnimSequence, EvalTimes.Current, bFireNotifies, bLooping);
+							SetAnimPosition(SkeletalMeshComponent, AnimSection->GetSlotName(), ChannelIndex, AnimSequence, EvalTime, bLooping, bFireNotifies);
 						}
 					}
 				}
@@ -154,7 +210,7 @@ bool FMovieSceneSkeletalAnimationTrackInstance::CanPlayAnimation(USkeletalMeshCo
 		(!AnimAssetBase || SkeletalMeshComponent->SkeletalMesh->Skeleton->IsCompatible(AnimAssetBase->GetSkeleton())));
 }
 
-void FMovieSceneSkeletalAnimationTrackInstance::SetAnimPosition(USkeletalMeshComponent* SkeletalMeshComponent, FName SlotName, int32 ChannelIndex, UAnimSequence* InAnimSequence, float InPosition, bool bFireNotifies, bool bLooping)
+void FMovieSceneSkeletalAnimationTrackInstance::SetAnimPosition(USkeletalMeshComponent* SkeletalMeshComponent, FName SlotName, int32 ChannelIndex, UAnimSequence* InAnimSequence, float InPosition, bool bLooping, bool bFireNotifies)
 {
 	if (CanPlayAnimation(SkeletalMeshComponent, InAnimSequence))
 	{
@@ -215,10 +271,32 @@ void FMovieSceneSkeletalAnimationTrackInstance::PreviewFinishAnimControl(USkelet
 	}
 }
 
-void FMovieSceneSkeletalAnimationTrackInstance::PreviewSetAnimPosition(USkeletalMeshComponent* SkeletalMeshComponent, FName SlotName, int32 ChannelIndex, UAnimSequence* InAnimSequence, float InPosition, bool bLooping, bool bFireNotifies, float DeltaTime)
+void FMovieSceneSkeletalAnimationTrackInstance::PreviewSetAnimPosition(USkeletalMeshComponent* SkeletalMeshComponent, FName SlotName, int32 ChannelIndex, UAnimSequence* InAnimSequence, float InPosition, bool bLooping, bool bFireNotifies, float DeltaTime, bool bPlaying)
 {
 	if(CanPlayAnimation(SkeletalMeshComponent, InAnimSequence))
 	{
 		FAnimMontageInstance::PreviewMatineeSetAnimPositionInner(SlotName, SkeletalMeshComponent, InAnimSequence, CurrentlyPlayingMontage, InPosition, bLooping, bFireNotifies, DeltaTime);
+
+		// if we are not playing, make sure we dont continue (as skeletal meshes can still tick us onwards)
+		UAnimInstance* AnimInst = SkeletalMeshComponent->GetAnimInstance();
+		UAnimSingleNodeInstance * SingleNodeInst = SkeletalMeshComponent->GetSingleNodeInstance();
+		if(SingleNodeInst)
+		{
+			SingleNodeInst->SetPlaying(bPlaying);
+		}
+		else if (AnimInst)
+		{
+			if(CurrentlyPlayingMontage.IsValid())
+			{
+				if(bPlaying)
+				{
+					AnimInst->Montage_Resume(CurrentlyPlayingMontage.Get());
+				}
+				else
+				{
+					AnimInst->Montage_Pause(CurrentlyPlayingMontage.Get());
+				}
+			}
+		}
 	}
 }
\ No newline at end of file
diff --git a/Engine/Source/Runtime/MovieSceneTracks/Private/TrackInstances/MovieSceneSlomoTrackInstance.cpp b/Engine/Source/Runtime/MovieSceneTracks/Private/TrackInstances/MovieSceneSlomoTrackInstance.cpp
--- a/Engine/Source/Runtime/MovieSceneTracks/Private/TrackInstances/MovieSceneSlomoTrackInstance.cpp
+++ b/Engine/Source/Runtime/MovieSceneTracks/Private/TrackInstances/MovieSceneSlomoTrackInstance.cpp
@@ -14,7 +14,7 @@ FMovieSceneSlomoTrackInstance::FMovieSceneSlomoTrackInstance(UMovieSceneSlomoTra
 { }
 
 	
-void FMovieSceneSlomoTrackInstance::RestoreState(const TArray<UObject*>& RuntimeObjects, IMovieScenePlayer& Player, FMovieSceneSequenceInstance& SequenceInstance)
+void FMovieSceneSlomoTrackInstance::RestoreState(const TArray<TWeakObjectPtr<UObject>>& RuntimeObjects, IMovieScenePlayer& Player, FMovieSceneSequenceInstance& SequenceInstance)
 {
 	AWorldSettings* WorldSettings = GWorld->GetWorldSettings();
 
@@ -26,7 +26,7 @@ void FMovieSceneSlomoTrackInstance::RestoreState(const TArray<UObject*>& Runtime
 	WorldSettings->MatineeTimeDilation = InitMatineeTimeDilation;
 }
 
-void FMovieSceneSlomoTrackInstance::SaveState(const TArray<UObject*>& RuntimeObjects, IMovieScenePlayer& Player, FMovieSceneSequenceInstance& SequenceInstance)
+void FMovieSceneSlomoTrackInstance::SaveState(const TArray<TWeakObjectPtr<UObject>>& RuntimeObjects, IMovieScenePlayer& Player, FMovieSceneSequenceInstance& SequenceInstance)
 {
 	AWorldSettings* WorldSettings = GWorld->GetWorldSettings();
 
@@ -41,7 +41,7 @@ void FMovieSceneSlomoTrackInstance::SaveState(const TArray<UObject*>& RuntimeObj
 /* IMovieSceneTrackInstance interface
  *****************************************************************************/
 
-void FMovieSceneSlomoTrackInstance::Update(EMovieSceneUpdateData& UpdateData, const TArray<UObject*>& RuntimeObjects, IMovieScenePlayer& Player, FMovieSceneSequenceInstance& SequenceInstance)
+void FMovieSceneSlomoTrackInstance::Update(EMovieSceneUpdateData& UpdateData, const TArray<TWeakObjectPtr<UObject>>& RuntimeObjects, IMovieScenePlayer& Player, FMovieSceneSequenceInstance& SequenceInstance)
 {
 	if (!ShouldBeApplied())
 	{
diff --git a/Engine/Source/Runtime/MovieSceneTracks/Private/TrackInstances/MovieSceneSpawnTrackInstance.cpp b/Engine/Source/Runtime/MovieSceneTracks/Private/TrackInstances/MovieSceneSpawnTrackInstance.cpp
--- a/Engine/Source/Runtime/MovieSceneTracks/Private/TrackInstances/MovieSceneSpawnTrackInstance.cpp
+++ b/Engine/Source/Runtime/MovieSceneTracks/Private/TrackInstances/MovieSceneSpawnTrackInstance.cpp
@@ -13,18 +13,32 @@ FMovieSceneSpawnTrackInstance::FMovieSceneSpawnTrackInstance(UMovieSceneSpawnTra
 	Track = &InTrack;
 }
 
-void FMovieSceneSpawnTrackInstance::Update(EMovieSceneUpdateData& UpdateData, const TArray<UObject*>& RuntimeObjects, IMovieScenePlayer& Player, FMovieSceneSequenceInstance& SequenceInstance)
+void FMovieSceneSpawnTrackInstance::Update(EMovieSceneUpdateData& UpdateData, const TArray<TWeakObjectPtr<UObject>>& RuntimeObjects, IMovieScenePlayer& Player, FMovieSceneSequenceInstance& SequenceInstance)
 {
 	IMovieSceneSpawnRegister& SpawnRegister = Player.GetSpawnRegister();
 	FMovieSceneSpawnable* Spawnable = SequenceInstance.GetSequence()->GetMovieScene()->FindSpawnable(Track->GetObjectId());
 
 	TRange<float> Range = SequenceInstance.GetTimeRange();
 
+	const bool bIsPreview = Player.IsPreview();
+
 	// If we're evaluating outside of the instance's time range, and the sequence owns the spawnable, there's no reason to evaluate - it should already be destroyed
-	if (Spawnable && Spawnable->GetSpawnOwnership() == ESpawnOwnership::InnerSequence && !Range.Contains(UpdateData.Position) && !Range.Contains(UpdateData.LastPosition))
+	if (Spawnable && Spawnable->GetSpawnOwnership() == ESpawnOwnership::InnerSequence && !Range.Contains(UpdateData.Position))
 	{
-		SpawnRegister.DestroySpawnedObject(Track->GetObjectId(), SequenceInstance, Player);
-		return;
+		bool bDestroy = true;
+#if WITH_EDITORONLY_DATA
+		bDestroy = !Spawnable->ShouldIgnoreOwnershipInEditor();
+		// Don't destroy cameras while previewing
+		if (bIsPreview && MovieSceneHelpers::CameraComponentFromActor(GetDefault<AActor>(Spawnable->GetClass())))
+		{
+			bDestroy = false;
+		}
+#endif
+		if (bDestroy)
+		{
+			SpawnRegister.DestroySpawnedObject(Track->GetObjectId(), SequenceInstance, Player);
+			return;
+		}
 	}
 
 	bool bIsSpawned = false;
diff --git a/Engine/Source/Runtime/MovieSceneTracks/Private/TrackInstances/MovieSceneSubTrackInstance.cpp b/Engine/Source/Runtime/MovieSceneTracks/Private/TrackInstances/MovieSceneSubTrackInstance.cpp
--- a/Engine/Source/Runtime/MovieSceneTracks/Private/TrackInstances/MovieSceneSubTrackInstance.cpp
+++ b/Engine/Source/Runtime/MovieSceneTracks/Private/TrackInstances/MovieSceneSubTrackInstance.cpp
@@ -28,7 +28,7 @@ void FMovieSceneSubTrackInstance::ClearInstance(IMovieScenePlayer& Player, FMovi
 }
 
 
-void FMovieSceneSubTrackInstance::RefreshInstance(const TArray<UObject*>& RuntimeObjects, class IMovieScenePlayer& Player, FMovieSceneSequenceInstance& SequenceInstance)
+void FMovieSceneSubTrackInstance::RefreshInstance(const TArray<TWeakObjectPtr<UObject>>& RuntimeObjects, class IMovieScenePlayer& Player, FMovieSceneSequenceInstance& SequenceInstance)
 {
 	TSet<UMovieSceneSection*> FoundSections;
 	const TArray<UMovieSceneSection*>& AllSections = SubTrack->GetAllSections();
@@ -37,14 +37,7 @@ void FMovieSceneSubTrackInstance::RefreshInstance(const TArray<UObject*>& Runtim
 	for (const auto Section : AllSections)
 	{
 		UMovieSceneSubSection* SubSection = CastChecked<UMovieSceneSubSection>(Section);
-
-		// If the section doesn't have a valid movie scene or no longer has one 
-		// (e.g user deleted it) then skip adding an instance for it
-		if (SubSection->GetSequence() == nullptr)
-		{
-			continue;
-		}
-
+		
 		FoundSections.Add(Section);
 
 		// create an instance for the section
@@ -52,12 +45,24 @@ void FMovieSceneSubTrackInstance::RefreshInstance(const TArray<UObject*>& Runtim
 
 		if (!Instance.IsValid())
 		{
-			Instance = MakeShareable(new FMovieSceneSequenceInstance(*SubSection->GetSequence()));
+			UMovieSceneSequence* Sequence = SubSection->GetSequence();
+			if(Sequence != nullptr)
+			{
+				Instance = MakeShareable(new FMovieSceneSequenceInstance(*Sequence));
+			}
+			else
+			{
+				Instance = MakeShareable(new FMovieSceneSequenceInstance(*SubTrack));
+			}
 			SequenceInstancesBySection.Add(SubSection, Instance.ToSharedRef());
 		}
 
 		Player.AddOrUpdateMovieSceneInstance(*SubSection, Instance.ToSharedRef());
 		Instance->RefreshInstance(Player);
+
+#if WITH_EDITOR
+		SubSection->OnSequenceChanged() = FOnSequenceChanged::CreateSP(Instance.Get(), &FMovieSceneSequenceInstance::HandleSequenceSectionChanged);
+#endif
 	}
 
 	// remove sections that no longer exist
@@ -74,7 +79,7 @@ void FMovieSceneSubTrackInstance::RefreshInstance(const TArray<UObject*>& Runtim
 }
 
 
-void FMovieSceneSubTrackInstance::RestoreState(const TArray<UObject*>& RuntimeObjects, IMovieScenePlayer& Player, FMovieSceneSequenceInstance& SequenceInstance)
+void FMovieSceneSubTrackInstance::RestoreState(const TArray<TWeakObjectPtr<UObject>>& RuntimeObjects, IMovieScenePlayer& Player, FMovieSceneSequenceInstance& SequenceInstance)
 {
 	for (const auto Section : SubTrack->GetAllSections())
 	{
@@ -89,7 +94,7 @@ void FMovieSceneSubTrackInstance::RestoreState(const TArray<UObject*>& RuntimeOb
 }
 
 
-void FMovieSceneSubTrackInstance::SaveState(const TArray<UObject*>& RuntimeObjects, IMovieScenePlayer& Player, FMovieSceneSequenceInstance& SequenceInstance)
+void FMovieSceneSubTrackInstance::SaveState(const TArray<TWeakObjectPtr<UObject>>& RuntimeObjects, IMovieScenePlayer& Player, FMovieSceneSequenceInstance& SequenceInstance)
 {
 	for (const auto Section : SubTrack->GetAllSections())
 	{
@@ -134,6 +139,21 @@ TArray<UMovieSceneSection*> FMovieSceneSubTrackInstance::GetAllTraversedSections
 	return TraversedSections;
 }
 
+bool FMovieSceneSubTrackInstance::ShouldEvaluateIfOverlapping(const TArray<UMovieSceneSection*>& TraversedSections, UMovieSceneSection* Section) const
+{
+	const bool bShouldRemove = TraversedSections.ContainsByPredicate([=](UMovieSceneSection* OtherSection){
+		if (Section->GetRowIndex() == OtherSection->GetRowIndex() &&
+			Section->GetRange().Overlaps(OtherSection->GetRange()) &&
+			Section->GetOverlapPriority() < OtherSection->GetOverlapPriority())
+		{
+			return true;
+		}
+		return false;
+	});
+
+	return bShouldRemove;
+}
+
 TArray<UMovieSceneSection*> FMovieSceneSubTrackInstance::GetTraversedSectionsWithPreroll( const TArray<UMovieSceneSection*>& Sections, float CurrentTime, float PreviousTime )
 {
 	TArray<UMovieSceneSection*> TraversedSections = GetAllTraversedSectionsWithPreroll(Sections, CurrentTime, PreviousTime);
@@ -143,15 +163,7 @@ TArray<UMovieSceneSection*> FMovieSceneSubTrackInstance::GetTraversedSectionsWit
 	{
 		UMovieSceneSection* Section = TraversedSections[RemoveAt];
 		
-		const bool bShouldRemove = TraversedSections.ContainsByPredicate([=](UMovieSceneSection* OtherSection){
-			if (Section->GetRowIndex() == OtherSection->GetRowIndex() &&
-				Section->GetRange().Overlaps(OtherSection->GetRange()) &&
-				Section->GetOverlapPriority() < OtherSection->GetOverlapPriority())
-			{
-				return true;
-			}
-			return false;
-		});
+		const bool bShouldRemove = ShouldEvaluateIfOverlapping(TraversedSections, Section);
 		
 		if (bShouldRemove)
 		{
@@ -167,7 +179,7 @@ TArray<UMovieSceneSection*> FMovieSceneSubTrackInstance::GetTraversedSectionsWit
 }
 
 
-void FMovieSceneSubTrackInstance::Update(EMovieSceneUpdateData& UpdateData, const TArray<UObject*>& RuntimeObjects, class IMovieScenePlayer& Player, FMovieSceneSequenceInstance& SequenceInstance) 
+void FMovieSceneSubTrackInstance::Update(EMovieSceneUpdateData& UpdateData, const TArray<TWeakObjectPtr<UObject>>& RuntimeObjects, class IMovieScenePlayer& Player, FMovieSceneSequenceInstance& SequenceInstance) 
 {
 	const TArray<UMovieSceneSection*>& AllSections = SubTrack->GetAllSections();
 
@@ -203,11 +215,22 @@ void FMovieSceneSubTrackInstance::Update(EMovieSceneUpdateData& UpdateData, cons
 		const float InstancePosition = InstanceOffset + (UpdateData.Position - (SubSection->GetStartTime()- SubSection->PrerollTime)) / SubSection->TimeScale;
 
 		EMovieSceneUpdateData SubUpdateData(InstancePosition, InstanceLastPosition);
-		SubUpdateData.bJumpCut = UpdateData.LastPosition < SubSection->GetStartTime();
+		SubUpdateData.bJumpCut = UpdateData.LastPosition < SubSection->GetStartTime() || UpdateData.LastPosition > SubSection->GetEndTime();
 		SubUpdateData.UpdatePass = UpdateData.UpdatePass;
 		SubUpdateData.bPreroll = InitialUpdatePosition < SubSection->GetStartTime();
 
-		// update section
-		Instance->Update(SubUpdateData, Player);
+		// update sub sections
+
+		if (SubUpdateData.UpdatePass == MSUP_PreUpdate)
+		{
+			Instance->PreUpdate(Player);
+		}
+		
+		Instance->UpdatePassSingle(SubUpdateData, Player);
+
+		if (SubUpdateData.UpdatePass == MSUP_PostUpdate)
+		{
+			Instance->PostUpdate(Player);
+		}
 	}
 }
diff --git a/Engine/Source/Runtime/MovieSceneTracks/Private/TrackInstances/MovieSceneVectorTrackInstance.cpp b/Engine/Source/Runtime/MovieSceneTracks/Private/TrackInstances/MovieSceneVectorTrackInstance.cpp
--- a/Engine/Source/Runtime/MovieSceneTracks/Private/TrackInstances/MovieSceneVectorTrackInstance.cpp
+++ b/Engine/Source/Runtime/MovieSceneTracks/Private/TrackInstances/MovieSceneVectorTrackInstance.cpp
@@ -12,11 +12,14 @@ FMovieSceneVectorTrackInstance::FMovieSceneVectorTrackInstance( UMovieSceneVecto
 }
 
 
-void FMovieSceneVectorTrackInstance::SaveState(const TArray<UObject*>& RuntimeObjects, IMovieScenePlayer& Player, FMovieSceneSequenceInstance& SequenceInstance)
+void FMovieSceneVectorTrackInstance::SaveState(const TArray<TWeakObjectPtr<UObject>>& RuntimeObjects, IMovieScenePlayer& Player, FMovieSceneSequenceInstance& SequenceInstance)
 {
 	int32 NumChannelsUsed = VectorTrack->GetNumChannelsUsed();
-	for( UObject* Object : RuntimeObjects )
+	
+	for (auto ObjectPtr : RuntimeObjects)
 	{
+		UObject* Object = ObjectPtr.Get();
+
 		switch( NumChannelsUsed )
 		{
 			case 2:
@@ -51,11 +54,14 @@ void FMovieSceneVectorTrackInstance::SaveState(const TArray<UObject*>& RuntimeOb
 }
 
 
-void FMovieSceneVectorTrackInstance::RestoreState(const TArray<UObject*>& RuntimeObjects, IMovieScenePlayer& Player, FMovieSceneSequenceInstance& SequenceInstance)
+void FMovieSceneVectorTrackInstance::RestoreState(const TArray<TWeakObjectPtr<UObject>>& RuntimeObjects, IMovieScenePlayer& Player, FMovieSceneSequenceInstance& SequenceInstance)
 {
 	int32 NumChannelsUsed = VectorTrack->GetNumChannelsUsed();
-	for( UObject* Object : RuntimeObjects )
+
+	for (auto ObjectPtr : RuntimeObjects)
 	{
+		UObject* Object = ObjectPtr.Get();
+
 		if (!IsValid(Object))
 		{
 			continue;
@@ -97,7 +103,7 @@ void FMovieSceneVectorTrackInstance::RestoreState(const TArray<UObject*>& Runtim
 }
 
 
-void FMovieSceneVectorTrackInstance::Update( EMovieSceneUpdateData& UpdateData, const TArray<UObject*>& RuntimeObjects, class IMovieScenePlayer& Player, FMovieSceneSequenceInstance& SequenceInstance ) 
+void FMovieSceneVectorTrackInstance::Update( EMovieSceneUpdateData& UpdateData, const TArray<TWeakObjectPtr<UObject>>& RuntimeObjects, class IMovieScenePlayer& Player, FMovieSceneSequenceInstance& SequenceInstance ) 
 {
 	FVector4 Vector;
 	if( VectorTrack->Eval( UpdateData.Position, UpdateData.LastPosition, Vector ) )
@@ -108,26 +114,26 @@ void FMovieSceneVectorTrackInstance::Update( EMovieSceneUpdateData& UpdateData,
 			case 2:
 			{
 				FVector2D Value(Vector.X, Vector.Y);
-				for(UObject* Object : RuntimeObjects)
+				for(auto Object : RuntimeObjects)
 				{
-					PropertyBindings->CallFunction<FVector2D>(Object, &Value);
+					PropertyBindings->CallFunction<FVector2D>(Object.Get(), &Value);
 				}
 				break;
 			}
 			case 3:
 			{
 				FVector Value(Vector.X, Vector.Y, Vector.Z);
-				for(UObject* Object : RuntimeObjects)
+				for(auto Object : RuntimeObjects)
 				{
-					PropertyBindings->CallFunction<FVector>(Object, &Value);
+					PropertyBindings->CallFunction<FVector>(Object.Get(), &Value);
 				}
 				break;
 			}
 			case 4:
 			{
-				for(UObject* Object : RuntimeObjects)
+				for(auto Object : RuntimeObjects)
 				{
-					PropertyBindings->CallFunction<FVector4>(Object, &Vector);
+					PropertyBindings->CallFunction<FVector4>(Object.Get(), &Vector);
 				}
 				break;
 			}
@@ -140,7 +146,7 @@ void FMovieSceneVectorTrackInstance::Update( EMovieSceneUpdateData& UpdateData,
 }
 
 
-void FMovieSceneVectorTrackInstance::RefreshInstance(const TArray<UObject*>& RuntimeObjects, class IMovieScenePlayer& Player, FMovieSceneSequenceInstance& SequenceInstance)
+void FMovieSceneVectorTrackInstance::RefreshInstance(const TArray<TWeakObjectPtr<UObject>>& RuntimeObjects, class IMovieScenePlayer& Player, FMovieSceneSequenceInstance& SequenceInstance)
 {
 	PropertyBindings->UpdateBindings(RuntimeObjects);
 }
diff --git a/Engine/Source/Runtime/MovieSceneTracks/Private/TrackInstances/MovieSceneVisibilityTrackInstance.cpp b/Engine/Source/Runtime/MovieSceneTracks/Private/TrackInstances/MovieSceneVisibilityTrackInstance.cpp
--- a/Engine/Source/Runtime/MovieSceneTracks/Private/TrackInstances/MovieSceneVisibilityTrackInstance.cpp
+++ b/Engine/Source/Runtime/MovieSceneTracks/Private/TrackInstances/MovieSceneVisibilityTrackInstance.cpp
@@ -13,11 +13,11 @@ FMovieSceneVisibilityTrackInstance::FMovieSceneVisibilityTrackInstance(UMovieSce
 }
 
 
-void FMovieSceneVisibilityTrackInstance::SaveState(const TArray<UObject*>& RuntimeObjects, IMovieScenePlayer& Player, FMovieSceneSequenceInstance& SequenceInstance)
+void FMovieSceneVisibilityTrackInstance::SaveState(const TArray<TWeakObjectPtr<UObject>>& RuntimeObjects, IMovieScenePlayer& Player, FMovieSceneSequenceInstance& SequenceInstance)
 {
- 	for( UObject* Object : RuntimeObjects )
+ 	for (auto ObjectPtr : RuntimeObjects)
  	{
-		AActor* Actor = Cast<AActor>(Object);
+		AActor* Actor = Cast<AActor>(ObjectPtr.Get());
 
 		if (Actor != nullptr)
 		{
@@ -48,10 +48,12 @@ void FMovieSceneVisibilityTrackInstance::SaveState(const TArray<UObject*>& Runti
 }
 
 
-void FMovieSceneVisibilityTrackInstance::RestoreState(const TArray<UObject*>& RuntimeObjects, IMovieScenePlayer& Player, FMovieSceneSequenceInstance& SequenceInstance)
+void FMovieSceneVisibilityTrackInstance::RestoreState(const TArray<TWeakObjectPtr<UObject>>& RuntimeObjects, IMovieScenePlayer& Player, FMovieSceneSequenceInstance& SequenceInstance)
 {
- 	for( UObject* Object : RuntimeObjects )
+ 	for (auto ObjectPtr : RuntimeObjects)
  	{
+		UObject* Object = ObjectPtr.Get();
+
 		if (!IsValid(Object))
 		{
 			continue;
@@ -98,7 +100,7 @@ void FMovieSceneVisibilityTrackInstance::RestoreState(const TArray<UObject*>& Ru
 }
 
 
-void FMovieSceneVisibilityTrackInstance::Update( EMovieSceneUpdateData& UpdateData, const TArray<UObject*>& RuntimeObjects, class IMovieScenePlayer& Player, FMovieSceneSequenceInstance& SequenceInstance ) 
+void FMovieSceneVisibilityTrackInstance::Update( EMovieSceneUpdateData& UpdateData, const TArray<TWeakObjectPtr<UObject>>& RuntimeObjects, class IMovieScenePlayer& Player, FMovieSceneSequenceInstance& SequenceInstance ) 
 {
  	bool Visible = false;
  	if( VisibilityTrack->Eval( UpdateData.Position, UpdateData.LastPosition, Visible ) )
@@ -106,8 +108,9 @@ void FMovieSceneVisibilityTrackInstance::Update( EMovieSceneUpdateData& UpdateDa
 		// Invert this evaluation since the property is "bHiddenInGame" and we want the visualization to be the inverse of that. Green means visible.
 		Visible = !Visible;
 
- 		for( UObject* Object : RuntimeObjects )
+ 		for (auto ObjectPtr : RuntimeObjects)
  		{
+			UObject* Object = ObjectPtr.Get();
 			AActor* Actor = Cast<AActor>(Object);
 			USceneComponent* SceneComponent = Cast<USceneComponent>(Object);
 
@@ -148,7 +151,7 @@ void FMovieSceneVisibilityTrackInstance::Update( EMovieSceneUpdateData& UpdateDa
 }
 
 
-void FMovieSceneVisibilityTrackInstance::RefreshInstance( const TArray<UObject*>& RuntimeObjects, class IMovieScenePlayer& Player, FMovieSceneSequenceInstance& SequenceInstance )
+void FMovieSceneVisibilityTrackInstance::RefreshInstance( const TArray<TWeakObjectPtr<UObject>>& RuntimeObjects, class IMovieScenePlayer& Player, FMovieSceneSequenceInstance& SequenceInstance )
 {
 	PropertyBindings->UpdateBindings( RuntimeObjects );
 }
diff --git a/Engine/Source/Runtime/MovieSceneTracks/Private/Tracks/MovieScene3DConstraintTrack.cpp b/Engine/Source/Runtime/MovieSceneTracks/Private/Tracks/MovieScene3DConstraintTrack.cpp
--- a/Engine/Source/Runtime/MovieSceneTracks/Private/Tracks/MovieScene3DConstraintTrack.cpp
+++ b/Engine/Source/Runtime/MovieSceneTracks/Private/Tracks/MovieScene3DConstraintTrack.cpp
@@ -14,7 +14,7 @@ UMovieScene3DConstraintTrack::UMovieScene3DConstraintTrack( const FObjectInitial
 	: Super( ObjectInitializer )
 {
 #if WITH_EDITORONLY_DATA
-	TrackTint = FColor(128, 90, 0);
+	TrackTint = FColor(128, 90, 0, 65);
 #endif
 }
 
@@ -67,13 +67,4 @@ TRange<float> UMovieScene3DConstraintTrack::GetSectionBoundaries() const
 	return TRange<float>::Hull(Bounds);
 }
 
-
-#if WITH_EDITORONLY_DATA
-FText UMovieScene3DConstraintTrack::GetDefaultDisplayName() const
-{
-	return LOCTEXT("TrackName", "Constraint");
-}
-#endif
-
-
 #undef LOCTEXT_NAMESPACE
diff --git a/Engine/Source/Runtime/MovieSceneTracks/Private/Tracks/MovieScene3DTransformTrack.cpp b/Engine/Source/Runtime/MovieSceneTracks/Private/Tracks/MovieScene3DTransformTrack.cpp
--- a/Engine/Source/Runtime/MovieSceneTracks/Private/Tracks/MovieScene3DTransformTrack.cpp
+++ b/Engine/Source/Runtime/MovieSceneTracks/Private/Tracks/MovieScene3DTransformTrack.cpp
@@ -14,7 +14,7 @@ UMovieScene3DTransformTrack::UMovieScene3DTransformTrack( const FObjectInitializ
 	SetPropertyNameAndPath(Transform, Transform.ToString());
 
 #if WITH_EDITORONLY_DATA
-	TrackTint = FColor(48, 227, 255);
+	TrackTint = FColor(48, 227, 255, 65);
 #endif
 }
 
diff --git a/Engine/Source/Runtime/MovieSceneTracks/Private/Tracks/MovieSceneActorReferenceTrack.cpp b/Engine/Source/Runtime/MovieSceneTracks/Private/Tracks/MovieSceneActorReferenceTrack.cpp
new file mode 100644
--- /dev/null
+++ b/Engine/Source/Runtime/MovieSceneTracks/Private/Tracks/MovieSceneActorReferenceTrack.cpp
@@ -0,0 +1,43 @@
+// Copyright 1998-2016 Epic Games, Inc. All Rights Reserved.
+
+#include "MovieSceneTracksPrivatePCH.h"
+#include "MovieSceneActorReferenceSection.h"
+#include "MovieSceneActorReferenceTrack.h"
+#include "IMovieScenePlayer.h"
+#include "MovieSceneActorReferenceTrackInstance.h"
+
+
+UMovieSceneActorReferenceTrack::UMovieSceneActorReferenceTrack( const FObjectInitializer& ObjectInitializer )
+	: Super( ObjectInitializer )
+{ }
+
+
+UMovieSceneSection* UMovieSceneActorReferenceTrack::CreateNewSection()
+{
+	return NewObject<UMovieSceneSection>(this, UMovieSceneActorReferenceSection::StaticClass(), NAME_None, RF_Transactional);
+}
+
+
+TSharedPtr<IMovieSceneTrackInstance> UMovieSceneActorReferenceTrack::CreateInstance()
+{
+	return MakeShareable( new FMovieSceneActorReferenceTrackInstance( *this ) ); 
+}
+
+
+bool UMovieSceneActorReferenceTrack::Eval( float Position, float LastPosition, FGuid& OutActorReferenceGuid ) const
+{
+	const UMovieSceneSection* Section = MovieSceneHelpers::FindNearestSectionAtTime( Sections, Position );
+
+	if ( Section )
+	{
+		if ( !Section->IsInfinite() )
+		{
+			Position = FMath::Clamp( Position, Section->GetStartTime(), Section->GetEndTime() );
+		}
+
+		OutActorReferenceGuid = CastChecked<UMovieSceneActorReferenceSection>( Section )->Eval( Position );
+	}
+
+	return Section != nullptr;
+}
+
diff --git a/Engine/Source/Runtime/MovieSceneTracks/Private/Tracks/MovieSceneAudioTrack.cpp b/Engine/Source/Runtime/MovieSceneTracks/Private/Tracks/MovieSceneAudioTrack.cpp
--- a/Engine/Source/Runtime/MovieSceneTracks/Private/Tracks/MovieSceneAudioTrack.cpp
+++ b/Engine/Source/Runtime/MovieSceneTracks/Private/Tracks/MovieSceneAudioTrack.cpp
@@ -19,7 +19,7 @@ UMovieSceneAudioTrack::UMovieSceneAudioTrack( const FObjectInitializer& ObjectIn
 	: Super( ObjectInitializer )
 {
 #if WITH_EDITORONLY_DATA
-	TrackTint = FColor(0, 0, 255);
+	TrackTint = FColor(0, 9, 186);
 #endif
 }
 
diff --git a/Engine/Source/Runtime/MovieSceneTracks/Private/Tracks/MovieSceneCameraAnimTrack.cpp b/Engine/Source/Runtime/MovieSceneTracks/Private/Tracks/MovieSceneCameraAnimTrack.cpp
new file mode 100644
--- /dev/null
+++ b/Engine/Source/Runtime/MovieSceneTracks/Private/Tracks/MovieSceneCameraAnimTrack.cpp
@@ -0,0 +1,113 @@
+// Copyright 1998-2016 Epic Games, Inc. All Rights Reserved.
+
+#include "MovieSceneTracksPrivatePCH.h"
+#include "MovieSceneCameraAnimSection.h"
+#include "MovieSceneCameraAnimTrack.h"
+#include "IMovieScenePlayer.h"
+#include "MovieSceneCameraAnimTrackInstance.h"
+#include "Camera/CameraAnim.h"
+
+#define LOCTEXT_NAMESPACE "MovieSceneCameraAnimTrack"
+
+
+void UMovieSceneCameraAnimTrack::AddNewCameraAnim(float KeyTime, UCameraAnim* CameraAnim)
+{
+	UMovieSceneCameraAnimSection* const NewSection = Cast<UMovieSceneCameraAnimSection>(CreateNewSection());
+	if (NewSection)
+	{
+		NewSection->InitialPlacement(CameraAnimSections, KeyTime, KeyTime + CameraAnim->AnimLength, SupportsMultipleRows());
+		NewSection->SetCameraAnim(CameraAnim);
+
+		AddSection(*NewSection);
+	}
+}
+
+/* UMovieSceneTrack interface
+*****************************************************************************/
+
+TSharedPtr<IMovieSceneTrackInstance> UMovieSceneCameraAnimTrack::CreateInstance()
+{
+	return MakeShareable(new FMovieSceneCameraAnimTrackInstance(*this));
+}
+
+
+const TArray<UMovieSceneSection*>& UMovieSceneCameraAnimTrack::GetAllSections() const
+{
+	return CameraAnimSections;
+}
+
+
+UMovieSceneSection* UMovieSceneCameraAnimTrack::CreateNewSection()
+{
+	return NewObject<UMovieSceneCameraAnimSection>(this);
+}
+
+
+void UMovieSceneCameraAnimTrack::RemoveAllAnimationData()
+{
+	CameraAnimSections.Empty();
+}
+
+
+bool UMovieSceneCameraAnimTrack::HasSection(const UMovieSceneSection& Section) const
+{
+	return CameraAnimSections.Contains(&Section);
+}
+
+
+void UMovieSceneCameraAnimTrack::AddSection(UMovieSceneSection& Section)
+{
+	CameraAnimSections.Add(&Section);
+}
+
+
+void UMovieSceneCameraAnimTrack::RemoveSection(UMovieSceneSection& Section)
+{
+	CameraAnimSections.Remove(&Section);
+}
+
+
+bool UMovieSceneCameraAnimTrack::IsEmpty() const
+{
+	return CameraAnimSections.Num() == 0;
+}
+
+
+TRange<float> UMovieSceneCameraAnimTrack::GetSectionBoundaries() const
+{
+	TArray<TRange<float>> Bounds;
+
+	for (auto Section : CameraAnimSections)
+	{
+		Bounds.Add(Section->GetRange());
+	}
+
+	return TRange<float>::Hull(Bounds);
+}
+
+
+
+
+#if WITH_EDITORONLY_DATA
+FText UMovieSceneCameraAnimTrack::GetDisplayName() const
+{
+	return LOCTEXT("TrackName", "Camera Anim");
+}
+#endif
+
+void UMovieSceneCameraAnimTrack::GetCameraAnimSectionsAtTime(float Time, TArray<UMovieSceneCameraAnimSection*>& OutSections)
+{
+	OutSections.Empty();
+
+	for (auto Section : CameraAnimSections)
+	{
+		UMovieSceneCameraAnimSection* const CASection = dynamic_cast<UMovieSceneCameraAnimSection*>(Section);
+		if (CASection && CASection->IsTimeWithinSection(Time))
+		{
+			OutSections.Add(CASection);
+		}
+	}
+}
+
+
+#undef LOCTEXT_NAMESPACE
diff --git a/Engine/Source/Runtime/MovieSceneTracks/Private/Tracks/MovieSceneCameraCutTrack.cpp b/Engine/Source/Runtime/MovieSceneTracks/Private/Tracks/MovieSceneCameraCutTrack.cpp
--- a/Engine/Source/Runtime/MovieSceneTracks/Private/Tracks/MovieSceneCameraCutTrack.cpp
+++ b/Engine/Source/Runtime/MovieSceneTracks/Private/Tracks/MovieSceneCameraCutTrack.cpp
@@ -17,7 +17,7 @@ UMovieSceneCameraCutTrack::UMovieSceneCameraCutTrack( const FObjectInitializer&
 	: Super( ObjectInitializer )
 {
 #if WITH_EDITORONLY_DATA
-	TrackTint = FColor(0, 0, 0);
+	TrackTint = FColor(0, 0, 0, 65);
 #endif
 }
 
diff --git a/Engine/Source/Runtime/MovieSceneTracks/Private/Tracks/MovieSceneCameraShakeTrack.cpp b/Engine/Source/Runtime/MovieSceneTracks/Private/Tracks/MovieSceneCameraShakeTrack.cpp
new file mode 100644
--- /dev/null
+++ b/Engine/Source/Runtime/MovieSceneTracks/Private/Tracks/MovieSceneCameraShakeTrack.cpp
@@ -0,0 +1,113 @@
+// Copyright 1998-2016 Epic Games, Inc. All Rights Reserved.
+
+#include "MovieSceneTracksPrivatePCH.h"
+#include "MovieSceneCameraShakeSection.h"
+#include "MovieSceneCameraShakeTrack.h"
+#include "IMovieScenePlayer.h"
+#include "MovieSceneCameraShakeTrackInstance.h"
+
+#define LOCTEXT_NAMESPACE "MovieSceneCameraShakeTrack"
+
+void UMovieSceneCameraShakeTrack::AddNewCameraShake(float KeyTime, TSubclassOf<UCameraShake> ShakeClass)
+{
+	UMovieSceneCameraShakeSection* const NewSection = Cast<UMovieSceneCameraShakeSection>(CreateNewSection());
+	if (NewSection)
+	{
+		// #fixme get length
+		NewSection->InitialPlacement(CameraShakeSections, KeyTime, KeyTime + 5.f /*AnimSequence->SequenceLength*/, SupportsMultipleRows());
+		NewSection->SetCameraShakeClass(ShakeClass);
+		
+		AddSection(*NewSection);
+	}
+}
+
+#if WITH_EDITORONLY_DATA
+FText UMovieSceneCameraShakeTrack::GetDisplayName() const
+{
+	return LOCTEXT("TrackName", "Camera Shake");
+}
+#endif
+
+
+
+
+/* UMovieSceneTrack interface
+*****************************************************************************/
+
+TSharedPtr<IMovieSceneTrackInstance> UMovieSceneCameraShakeTrack::CreateInstance()
+{
+	return MakeShareable(new FMovieSceneCameraShakeTrackInstance(*this));
+}
+
+
+const TArray<UMovieSceneSection*>& UMovieSceneCameraShakeTrack::GetAllSections() const
+{
+	return CameraShakeSections;
+}
+
+
+UMovieSceneSection* UMovieSceneCameraShakeTrack::CreateNewSection()
+{
+	return NewObject<UMovieSceneCameraShakeSection>(this);
+}
+
+
+void UMovieSceneCameraShakeTrack::RemoveAllAnimationData()
+{
+	CameraShakeSections.Empty();
+}
+
+
+bool UMovieSceneCameraShakeTrack::HasSection(const UMovieSceneSection& Section) const
+{
+	return CameraShakeSections.Contains(&Section);
+}
+
+
+void UMovieSceneCameraShakeTrack::AddSection(UMovieSceneSection& Section)
+{
+	CameraShakeSections.Add(&Section);
+}
+
+
+void UMovieSceneCameraShakeTrack::RemoveSection(UMovieSceneSection& Section)
+{
+	CameraShakeSections.Remove(&Section);
+}
+
+
+bool UMovieSceneCameraShakeTrack::IsEmpty() const
+{
+	return CameraShakeSections.Num() == 0;
+}
+
+
+TRange<float> UMovieSceneCameraShakeTrack::GetSectionBoundaries() const
+{
+	TArray<TRange<float>> Bounds;
+
+	for (auto Section : CameraShakeSections)
+	{
+		Bounds.Add(Section->GetRange());
+	}
+
+	return TRange<float>::Hull(Bounds);
+}
+
+
+void UMovieSceneCameraShakeTrack::GetCameraShakeSectionsAtTime(float Time, TArray<UMovieSceneCameraShakeSection*>& OutSections)
+{
+	OutSections.Empty();
+
+	for (auto Section : CameraShakeSections)
+	{
+		UMovieSceneCameraShakeSection* const CSSection = dynamic_cast<UMovieSceneCameraShakeSection*>(Section);
+		if (CSSection && CSSection->IsTimeWithinSection(Time))
+		{
+			OutSections.Add(CSSection);
+		}
+	}
+}
+
+
+#undef LOCTEXT_NAMESPACE
diff --git a/Engine/Source/Runtime/MovieSceneTracks/Private/Tracks/MovieSceneCinematicShotTrack.cpp b/Engine/Source/Runtime/MovieSceneTracks/Private/Tracks/MovieSceneCinematicShotTrack.cpp
--- a/Engine/Source/Runtime/MovieSceneTracks/Private/Tracks/MovieSceneCinematicShotTrack.cpp
+++ b/Engine/Source/Runtime/MovieSceneTracks/Private/Tracks/MovieSceneCinematicShotTrack.cpp
@@ -17,7 +17,7 @@ UMovieSceneCinematicShotTrack::UMovieSceneCinematicShotTrack(const FObjectInitia
 	: Super(ObjectInitializer)
 {
 #if WITH_EDITORONLY_DATA
-	TrackTint = FColor(0,0,0);
+	TrackTint = FColor(0, 0, 0, 127);
 #endif
 }
 
@@ -40,7 +40,7 @@ UMovieSceneSubSection* UMovieSceneCinematicShotTrack::AddSequence(UMovieSceneSeq
 	MovieSceneHelpers::SortConsecutiveSections(Sections);
 
 	// Once sequences are sorted fixup the surrounding sequences to fix any gaps
-	MovieSceneHelpers::FixupConsecutiveSections(Sections, *NewSection, false);
+	//MovieSceneHelpers::FixupConsecutiveSections(Sections, *NewSection, false);
 
 	return NewSection;
 }
@@ -71,7 +71,7 @@ UMovieSceneSection* UMovieSceneCinematicShotTrack::CreateNewSection()
 void UMovieSceneCinematicShotTrack::RemoveSection(UMovieSceneSection& Section)
 {
 	Sections.Remove(&Section);
-	MovieSceneHelpers::FixupConsecutiveSections(Sections, Section, true);
+	//MovieSceneHelpers::FixupConsecutiveSections(Sections, Section, true);
 	MovieSceneHelpers::SortConsecutiveSections(Sections);
 
 	// @todo Sequencer: The movie scene owned by the section is now abandoned.  Should we offer to delete it?  
@@ -79,7 +79,7 @@ void UMovieSceneCinematicShotTrack::RemoveSection(UMovieSceneSection& Section)
 
 bool UMovieSceneCinematicShotTrack::SupportsMultipleRows() const
 {
-	return false;
+	return true;
 }
 
 bool UMovieSceneCinematicShotTrack::AddsSectionBoundsToPlayRange() const
@@ -90,7 +90,7 @@ bool UMovieSceneCinematicShotTrack::AddsSectionBoundsToPlayRange() const
 #if WITH_EDITOR
 void UMovieSceneCinematicShotTrack::OnSectionMoved(UMovieSceneSection& Section)
 {
-	MovieSceneHelpers::FixupConsecutiveSections(Sections, Section, false);
+	//MovieSceneHelpers::FixupConsecutiveSections(Sections, Section, false);
 }
 #endif
 
diff --git a/Engine/Source/Runtime/MovieSceneTracks/Private/Tracks/MovieSceneMaterialTrack.cpp b/Engine/Source/Runtime/MovieSceneTracks/Private/Tracks/MovieSceneMaterialTrack.cpp
--- a/Engine/Source/Runtime/MovieSceneTracks/Private/Tracks/MovieSceneMaterialTrack.cpp
+++ b/Engine/Source/Runtime/MovieSceneTracks/Private/Tracks/MovieSceneMaterialTrack.cpp
@@ -11,7 +11,7 @@ UMovieSceneMaterialTrack::UMovieSceneMaterialTrack(const FObjectInitializer& Obj
 	: Super(ObjectInitializer)
 {
 #if WITH_EDITORONLY_DATA
-	TrackTint = FColor(64,192,64);
+	TrackTint = FColor(64,192,64,65);
 #endif
 }
 
diff --git a/Engine/Source/Runtime/MovieSceneTracks/Private/Tracks/MovieSceneParticleParameterTrack.cpp b/Engine/Source/Runtime/MovieSceneTracks/Private/Tracks/MovieSceneParticleParameterTrack.cpp
--- a/Engine/Source/Runtime/MovieSceneTracks/Private/Tracks/MovieSceneParticleParameterTrack.cpp
+++ b/Engine/Source/Runtime/MovieSceneTracks/Private/Tracks/MovieSceneParticleParameterTrack.cpp
@@ -12,7 +12,7 @@ UMovieSceneParticleParameterTrack::UMovieSceneParticleParameterTrack( const FObj
 	: Super( ObjectInitializer )
 {
 #if WITH_EDITORONLY_DATA
-	TrackTint = FColor(0, 170, 255);
+	TrackTint = FColor(0, 170, 255, 65);
 #endif
 }
 
diff --git a/Engine/Source/Runtime/MovieSceneTracks/Private/Tracks/MovieSceneParticleTrack.cpp b/Engine/Source/Runtime/MovieSceneTracks/Private/Tracks/MovieSceneParticleTrack.cpp
--- a/Engine/Source/Runtime/MovieSceneTracks/Private/Tracks/MovieSceneParticleTrack.cpp
+++ b/Engine/Source/Runtime/MovieSceneTracks/Private/Tracks/MovieSceneParticleTrack.cpp
@@ -14,7 +14,7 @@ UMovieSceneParticleTrack::UMovieSceneParticleTrack( const FObjectInitializer& Ob
 	: Super( ObjectInitializer )
 {
 #if WITH_EDITORONLY_DATA
-	TrackTint = FColor(255,255,255);
+	TrackTint = FColor(255,255,255,160);
 #endif
 }
 
diff --git a/Engine/Source/Runtime/MovieSceneTracks/Private/Tracks/MovieSceneSkeletalAnimationTrack.cpp b/Engine/Source/Runtime/MovieSceneTracks/Private/Tracks/MovieSceneSkeletalAnimationTrack.cpp
--- a/Engine/Source/Runtime/MovieSceneTracks/Private/Tracks/MovieSceneSkeletalAnimationTrack.cpp
+++ b/Engine/Source/Runtime/MovieSceneTracks/Private/Tracks/MovieSceneSkeletalAnimationTrack.cpp
@@ -17,7 +17,7 @@ UMovieSceneSkeletalAnimationTrack::UMovieSceneSkeletalAnimationTrack(const FObje
 	: Super(ObjectInitializer)
 {
 #if WITH_EDITORONLY_DATA
-	TrackTint = FColor(194, 29, 194);
+	TrackTint = FColor(194, 29, 194, 65);
 #endif
 }
 
diff --git a/Engine/Source/Runtime/MovieSceneTracks/Private/Tracks/MovieSceneSpawnTrack.cpp b/Engine/Source/Runtime/MovieSceneTracks/Private/Tracks/MovieSceneSpawnTrack.cpp
--- a/Engine/Source/Runtime/MovieSceneTracks/Private/Tracks/MovieSceneSpawnTrack.cpp
+++ b/Engine/Source/Runtime/MovieSceneTracks/Private/Tracks/MovieSceneSpawnTrack.cpp
@@ -16,7 +16,7 @@ UMovieSceneSpawnTrack::UMovieSceneSpawnTrack(const FObjectInitializer& Obj)
 	: Super(Obj)
 {
 #if WITH_EDITORONLY_DATA
-	TrackTint = FColor(255, 0, 234);
+	TrackTint = FColor(255, 0, 93, 65);
 #endif
 }
 
@@ -97,7 +97,7 @@ const TArray<UMovieSceneSection*>& UMovieSceneSpawnTrack::GetAllSections() const
 
 #if WITH_EDITORONLY_DATA
 
-FText UMovieSceneSpawnTrack::GetDefaultDisplayName() const
+FText UMovieSceneSpawnTrack::GetDisplayName() const
 {
 	return LOCTEXT("TrackName", "Spawned");
 }
diff --git a/Engine/Source/Runtime/MovieSceneTracks/Private/Tracks/MovieSceneSubTrack.cpp b/Engine/Source/Runtime/MovieSceneTracks/Private/Tracks/MovieSceneSubTrack.cpp
--- a/Engine/Source/Runtime/MovieSceneTracks/Private/Tracks/MovieSceneSubTrack.cpp
+++ b/Engine/Source/Runtime/MovieSceneTracks/Private/Tracks/MovieSceneSubTrack.cpp
@@ -17,7 +17,7 @@ UMovieSceneSubTrack::UMovieSceneSubTrack( const FObjectInitializer& ObjectInitia
 	: Super(ObjectInitializer)
 {
 #if WITH_EDITORONLY_DATA
-	TrackTint = FColor(255, 0, 0);
+	TrackTint = FColor(255, 0, 0, 65);
 #endif
 }
 
@@ -99,6 +99,29 @@ UMovieSceneSubSection* UMovieSceneSubTrack::AddSequence(UMovieSceneSequence* Seq
 	return NewSection;
 }
 
+UMovieSceneSubSection* UMovieSceneSubTrack::AddSequenceToRecord()
+{
+	UMovieScene* MovieScene = CastChecked<UMovieScene>(GetOuter());
+	TRange<float> PlaybackRange = MovieScene->GetPlaybackRange();
+
+	int32 MaxRowIndex = 0;
+	for (auto Section : Sections)
+	{
+		MaxRowIndex = FMath::Max(Section->GetRowIndex(), MaxRowIndex);
+	}
+
+	UMovieSceneSubSection* NewSection = CastChecked<UMovieSceneSubSection>(CreateNewSection());
+	{
+		NewSection->SetRowIndex(MaxRowIndex + 1);
+		NewSection->SetAsRecording(true);
+		NewSection->SetStartTime(PlaybackRange.GetLowerBoundValue());
+		NewSection->SetEndTime(PlaybackRange.GetUpperBoundValue());
+	}
+
+	Sections.Add(NewSection);
+
+	return NewSection;
+}
 
 bool UMovieSceneSubTrack::ContainsSequence(const UMovieSceneSequence& Sequence, bool Recursively) const
 {
diff --git a/Engine/Source/Runtime/MovieSceneTracks/Public/Sections/MovieSceneActorReferenceSection.h b/Engine/Source/Runtime/MovieSceneTracks/Public/Sections/MovieSceneActorReferenceSection.h
new file mode 100644
--- /dev/null
+++ b/Engine/Source/Runtime/MovieSceneTracks/Public/Sections/MovieSceneActorReferenceSection.h
@@ -0,0 +1,57 @@
+// Copyright 1998-2016 Epic Games, Inc. All Rights Reserved.
+
+#pragma once
+
+#include "MovieSceneSection.h"
+#include "IKeyframeSection.h"
+#include "MovieSceneActorReferenceSection.generated.h"
+
+/**
+ * A single actor reference point section
+ */
+UCLASS( MinimalAPI )
+class UMovieSceneActorReferenceSection
+	: public UMovieSceneSection
+	, public IKeyframeSection<FGuid>
+{
+	GENERATED_UCLASS_BODY()
+public:
+	/**
+	 * Updates this section
+	 *
+	 * @param Position	The position in time within the movie scene
+	 */
+	FGuid Eval( float Position ) const;
+
+	// IKeyframeSection interface.
+
+	void AddKey( float Time, const FGuid& Value, EMovieSceneKeyInterpolation KeyInterpolation );
+	bool NewKeyIsNewData(float Time, const FGuid& Value) const;
+	bool HasKeys( const FGuid& Value ) const;
+	void SetDefault( const FGuid& Value );
+
+	/**
+	 * UMovieSceneSection interface 
+	 */
+	virtual void MoveSection(float DeltaPosition, TSet<FKeyHandle>& KeyHandles) override;
+	virtual void DilateSection(float DilationFactor, float Origin, TSet<FKeyHandle>& KeyHandles) override;
+	virtual void GetKeyHandles(TSet<FKeyHandle>& KeyHandles) const override;
+
+	/**
+	 * @return The integral curve on this section
+	 */
+	FIntegralCurve& GetActorReferenceCurve() { return ActorGuidIndexCurve; }
+
+	virtual void PreSave() override;
+	virtual void PostLoad() override;
+
+private:
+	/** Curve data */
+	UPROPERTY()
+	FIntegralCurve ActorGuidIndexCurve;
+
+	TArray<FGuid> ActorGuids;
+
+	UPROPERTY()
+	TArray<FString> ActorGuidStrings;
+};
\ No newline at end of file
diff --git a/Engine/Source/Runtime/MovieSceneTracks/Public/Sections/MovieSceneBoolSection.h b/Engine/Source/Runtime/MovieSceneTracks/Public/Sections/MovieSceneBoolSection.h
--- a/Engine/Source/Runtime/MovieSceneTracks/Public/Sections/MovieSceneBoolSection.h
+++ b/Engine/Source/Runtime/MovieSceneTracks/Public/Sections/MovieSceneBoolSection.h
@@ -2,6 +2,7 @@
 
 #pragma once
 
+#include "IKeyframeSection.h"
 #include "MovieSceneBoolSection.generated.h"
 
 /**
diff --git a/Engine/Source/Runtime/MovieSceneTracks/Public/Sections/MovieSceneByteSection.h b/Engine/Source/Runtime/MovieSceneTracks/Public/Sections/MovieSceneByteSection.h
--- a/Engine/Source/Runtime/MovieSceneTracks/Public/Sections/MovieSceneByteSection.h
+++ b/Engine/Source/Runtime/MovieSceneTracks/Public/Sections/MovieSceneByteSection.h
@@ -2,6 +2,7 @@
 
 #pragma once
 
+#include "IKeyframeSection.h"
 #include "MovieSceneByteSection.generated.h"
 
 /**
diff --git a/Engine/Source/Runtime/MovieSceneTracks/Public/Sections/MovieSceneCameraAnimSection.h b/Engine/Source/Runtime/MovieSceneTracks/Public/Sections/MovieSceneCameraAnimSection.h
new file mode 100644
--- /dev/null
+++ b/Engine/Source/Runtime/MovieSceneTracks/Public/Sections/MovieSceneCameraAnimSection.h
@@ -0,0 +1,70 @@
+// Copyright 1998-2016 Epic Games, Inc. All Rights Reserved.
+
+#pragma once
+
+#include "MovieSceneSection.h"
+#include "MovieSceneCameraAnimSection.generated.h"
+
+/**
+ *
+ */
+UCLASS(MinimalAPI)
+class UMovieSceneCameraAnimSection : public UMovieSceneSection
+{
+	GENERATED_BODY()
+
+public:
+	// ctor
+	UMovieSceneCameraAnimSection(const FObjectInitializer& ObjectInitializer);
+
+	/** Sets the CameraAnim for this section */
+	void SetCameraAnim(UCameraAnim* InCameraAnim) { CameraAnim = InCameraAnim; }
+	
+	/** Gets the CameraAnim for this section */
+	class UCameraAnim* GetCameraAnim() const { return CameraAnim; }
+	
+	/** MovieSceneSection interface */
+	virtual void MoveSection( float DeltaPosition, TSet<FKeyHandle>& KeyHandles ) override;
+	virtual void DilateSection( float DilationFactor, float Origin, TSet<FKeyHandle>& KeyHandles ) override;
+	virtual void GetKeyHandles(TSet<FKeyHandle>& KeyHandles) const override;
+
+	/** #todo */
+//	MOVIESCENETRACKS_API FRichCurve& GetAnimWeightCurve() { return AnimWeightCurve; }
+
+private:
+// 	UPROPERTY(EditAnywhere, Category="Camera Anim")
+// 	FRichCurve AnimWeightCurve;
+
+	UPROPERTY(EditAnywhere, Category = "Camera Anim")
+	UCameraAnim* CameraAnim;
+
+public:
+	/** How fast to play back the animation. */
+	UPROPERTY(EditAnywhere, Category = "Camera Anim")
+	float PlayRate;
+	
+	/** Scalar to control intensity of the animation. */
+	UPROPERTY(EditAnywhere, Category = "Camera Anim")
+	float PlayScale;
+	
+	UPROPERTY(EditAnywhere, Category = "Camera Anim")
+	float BlendInTime;
+	
+	UPROPERTY(EditAnywhere, Category = "Camera Anim")
+	float BlendOutTime;
+	
+	UPROPERTY(EditAnywhere, Category = "Camera Anim")
+	bool bLooping;
+	
+// 	UPROPERTY(EditAnywhere, Category = "Camera Anim")
+ 	bool bRandomStartTime;
+	
+// 	UPROPERTY(EditAnywhere, Category = "Camera Anim")
+// 	float Duration;
+	
+	UPROPERTY(EditAnywhere, Category = "Camera Anim")
+	TEnumAsByte<ECameraAnimPlaySpace::Type> PlaySpace;
+	
+	UPROPERTY(EditAnywhere, Category = "Camera Anim")
+	FRotator UserDefinedPlaySpace;
+};
\ No newline at end of file
diff --git a/Engine/Source/Runtime/MovieSceneTracks/Public/Sections/MovieSceneCameraCutSection.h b/Engine/Source/Runtime/MovieSceneTracks/Public/Sections/MovieSceneCameraCutSection.h
--- a/Engine/Source/Runtime/MovieSceneTracks/Public/Sections/MovieSceneCameraCutSection.h
+++ b/Engine/Source/Runtime/MovieSceneTracks/Public/Sections/MovieSceneCameraCutSection.h
@@ -34,4 +34,26 @@ private:
 	/** The camera possessable or spawnable that this movie CameraCut uses */
 	UPROPERTY()
 	FGuid CameraGuid;
+
+#if WITH_EDITORONLY_DATA
+public:
+	/** @return The thumbnail reference frame offset from the start of this section */
+	float GetThumbnailReferenceOffset() const
+	{
+		return ThumbnailReferenceOffset;
+	}
+
+	/** Set the thumbnail reference offset */
+	void SetThumbnailReferenceOffset(float InNewOffset)
+	{
+		Modify();
+		ThumbnailReferenceOffset = InNewOffset;
+	}
+
+private:
+
+	/** The reference frame offset for single thumbnail rendering */
+	UPROPERTY()
+	float ThumbnailReferenceOffset;
+#endif
 };
diff --git a/Engine/Source/Runtime/MovieSceneTracks/Public/Sections/MovieSceneCameraShakeSection.h b/Engine/Source/Runtime/MovieSceneTracks/Public/Sections/MovieSceneCameraShakeSection.h
new file mode 100644
--- /dev/null
+++ b/Engine/Source/Runtime/MovieSceneTracks/Public/Sections/MovieSceneCameraShakeSection.h
@@ -0,0 +1,59 @@
+// Copyright 1998-2016 Epic Games, Inc. All Rights Reserved.
+
+#pragma once
+
+#include "MovieSceneSection.h"
+#include "MovieSceneCameraShakeSection.generated.h"
+
+class UCameraShake;
+
+/**
+ *
+ */
+UCLASS(MinimalAPI)
+class UMovieSceneCameraShakeSection : public UMovieSceneSection
+{
+	GENERATED_BODY()
+
+public:
+
+	// ctor
+	UMovieSceneCameraShakeSection(const FObjectInitializer& ObjectInitializer);
+
+	/** Sets the CameraShake for this section */
+	void SetCameraShakeClass(TSubclassOf<UCameraShake> InShakeClass) { ShakeClass = InShakeClass; }
+
+	/** Gets the CameraAnim for this section */
+	UClass* GetCameraShakeClass() const { return ShakeClass; }
+
+	/** MovieSceneSection interface */
+	virtual void MoveSection( float DeltaPosition, TSet<FKeyHandle>& KeyHandles ) override;
+	virtual void DilateSection( float DilationFactor, float Origin, TSet<FKeyHandle>& KeyHandles ) override;
+	virtual void GetKeyHandles(TSet<FKeyHandle>& KeyHandles) const override;
+
+	/** 
+	 * @return Returns the weighting curve
+	 */
+//	MOVIESCENETRACKS_API FRichCurve& GetShakeWeightCurve() { return ShakeWeightCurve; }
+	
+private:
+	
+	/** */
+// 	UPROPERTY(EditAnywhere, Category = "Camera Shake")
+// 	FRichCurve ShakeWeightCurve;
+
+public:
+	/** Class of the camera shake to play */
+	UPROPERTY(EditAnywhere, Category = "Camera Shake")
+	TSubclassOf<UCameraShake> ShakeClass;
+	
+	/** Scalar that affects shake intensity */
+	UPROPERTY(EditAnywhere, Category = "Camera Shake")
+	float PlayScale;
+	
+	UPROPERTY(EditAnywhere, Category = "Camera Shake")
+	TEnumAsByte<ECameraAnimPlaySpace::Type> PlaySpace;
+
+	UPROPERTY(EditAnywhere, Category = "Camera Shake")
+	FRotator UserDefinedPlaySpace;
+};
\ No newline at end of file
diff --git a/Engine/Source/Runtime/MovieSceneTracks/Public/Sections/MovieSceneCinematicShotSection.h b/Engine/Source/Runtime/MovieSceneTracks/Public/Sections/MovieSceneCinematicShotSection.h
--- a/Engine/Source/Runtime/MovieSceneTracks/Public/Sections/MovieSceneCinematicShotSection.h
+++ b/Engine/Source/Runtime/MovieSceneTracks/Public/Sections/MovieSceneCinematicShotSection.h
@@ -42,4 +42,26 @@ private:
 	/** The Shot's display name */
 	UPROPERTY()
 	FText DisplayName;
+
+#if WITH_EDITORONLY_DATA
+public:
+	/** @return The shot thumbnail reference frame offset from the start of this section */
+	float GetThumbnailReferenceOffset() const
+	{
+		return ThumbnailReferenceOffset;
+	}
+
+	/** Set the thumbnail reference offset */
+	void SetThumbnailReferenceOffset(float InNewOffset)
+	{
+		Modify();
+		ThumbnailReferenceOffset = InNewOffset;
+	}
+
+private:
+
+	/** The shot's reference frame offset for single thumbnail rendering */
+	UPROPERTY()
+	float ThumbnailReferenceOffset;
+#endif
 };
diff --git a/Engine/Source/Runtime/MovieSceneTracks/Public/Sections/MovieSceneSubSection.h b/Engine/Source/Runtime/MovieSceneTracks/Public/Sections/MovieSceneSubSection.h
--- a/Engine/Source/Runtime/MovieSceneTracks/Public/Sections/MovieSceneSubSection.h
+++ b/Engine/Source/Runtime/MovieSceneTracks/Public/Sections/MovieSceneSubSection.h
@@ -7,6 +7,7 @@
 
 class UMovieSceneSequence;
 
+DECLARE_DELEGATE_OneParam(FOnSequenceChanged, UMovieSceneSequence* /*Sequence*/);
 
 /**
  * Implements a section in sub-sequence tracks.
@@ -28,10 +29,7 @@ public:
 	 * @return The sequence.
 	 * @see SetSequence
 	 */
-	UMovieSceneSequence* GetSequence() const
-	{
-		return SubSequence;
-	}
+	UMovieSceneSequence* GetSequence() const;
 
 	/**
 	 * Sets the sequence played by this section.
@@ -39,9 +37,49 @@ public:
 	 * @param Sequence The sequence to play.
 	 * @see GetSequence
 	 */
-	void SetSequence(UMovieSceneSequence* Sequence)
+	void SetSequence(UMovieSceneSequence* Sequence);
+
+	/** Prime this section as the one and only recording section */
+	void SetAsRecording(bool bRecord);
+
+	/** Get the section we are recording to */
+	static UMovieSceneSubSection* GetRecordingSection();
+
+	/** Get the actor we are targeting for recording */
+	static const FString& GetActorToRecord();
+
+	/** Check if we are primed for recording */
+	static bool IsSetAsRecording();
+
+#if WITH_EDITOR
+	virtual void PostEditChangeProperty(FPropertyChangedEvent& PropertyChangedEvent) override;
+
+	/** Delegate to fire when our sequence is changed in the property editor */
+	FOnSequenceChanged& OnSequenceChanged() { return OnSequenceChangedDelegate; }
+#endif
+
+	/** Get the name of the sequence we are going to try to record to */
+	const FString& GetTargetSequenceName() const
+	{
+		return TargetSequenceName;
+	}
+
+	/** Set the name of the sequence we are going to try to record to */
+	void SetTargetSequenceName(const FString& Name)
 	{
-		SubSequence = Sequence;
+		TargetSequenceName = Name;
+	}
+
+	/** Get the path of the sequence we are going to try to record to */
+	const FString& GetTargetPathToRecordTo() const
+	{
+		return TargetPathToRecordTo.Path;
+	}
+
+	/** Set the path of the sequence we are going to try to record to */
+	void SetTargetPathToRecordTo(const FString& Path)
+	{
+		TargetPathToRecordTo.Path = Path;
 	}
 
 public:
@@ -65,6 +103,26 @@ protected:
 	 *
 	 * @todo Sequencer: Should this be lazy loaded?
 	 */
-	UPROPERTY()
+	UPROPERTY(EditAnywhere, Category="Sequence")
 	UMovieSceneSequence* SubSequence;
+
+	/** Target actor name to record */
+	UPROPERTY(EditAnywhere, Category="Sequence Recording")
+	FString NameOfActorToRecord;
+
+	/** Target name of sequence to try to record to (will record automatically to another if this already exists) */
+	UPROPERTY(EditAnywhere, Category="Sequence Recording")
+	FString TargetSequenceName;
+
+	/** Target path of sequence to record to */
+	UPROPERTY(EditAnywhere, Category="Sequence Recording", meta=(ContentDir))
+	FDirectoryPath TargetPathToRecordTo;
+
+	/** Keep track of our constructed recordings */
+	static TWeakObjectPtr<UMovieSceneSubSection> TheRecordingSection;
+
+#if WITH_EDITOR
+	/** Delegate to fire when our sequence is changed in the property editor */
+	FOnSequenceChanged OnSequenceChangedDelegate;
+#endif
 };
diff --git a/Engine/Source/Runtime/MovieSceneTracks/Public/TrackInstances/MovieScene3DAttachTrackInstance.h b/Engine/Source/Runtime/MovieSceneTracks/Public/TrackInstances/MovieScene3DAttachTrackInstance.h
--- a/Engine/Source/Runtime/MovieSceneTracks/Public/TrackInstances/MovieScene3DAttachTrackInstance.h
+++ b/Engine/Source/Runtime/MovieSceneTracks/Public/TrackInstances/MovieScene3DAttachTrackInstance.h
@@ -17,5 +17,5 @@ public:
 	FMovieScene3DAttachTrackInstance( UMovieScene3DAttachTrack& InAttachTrack );
 
 	/** FMovieScene3DConstraintTrackInstance interface */
-	virtual void UpdateConstraint(float Position, const TArray<UObject*>& RuntimeObjects, AActor* Actor, UMovieScene3DConstraintSection* ConstraintSection) override;
+	virtual void UpdateConstraint(float Position, const TArray<TWeakObjectPtr<UObject>>& RuntimeObjects, AActor* Actor, UMovieScene3DConstraintSection* ConstraintSection) override;
 };
diff --git a/Engine/Source/Runtime/MovieSceneTracks/Public/TrackInstances/MovieScene3DConstraintTrackInstance.h b/Engine/Source/Runtime/MovieSceneTracks/Public/TrackInstances/MovieScene3DConstraintTrackInstance.h
--- a/Engine/Source/Runtime/MovieSceneTracks/Public/TrackInstances/MovieScene3DConstraintTrackInstance.h
+++ b/Engine/Source/Runtime/MovieSceneTracks/Public/TrackInstances/MovieScene3DConstraintTrackInstance.h
@@ -17,14 +17,16 @@ public:
 	FMovieScene3DConstraintTrackInstance( UMovieScene3DConstraintTrack& InConstraintTrack );
 
 	/** IMovieSceneTrackInstance interface */
-	virtual void SaveState (const TArray<UObject*>& RuntimeObjects, IMovieScenePlayer& Player, FMovieSceneSequenceInstance& SequenceInstance) override;
-	virtual void RestoreState (const TArray<UObject*>& RuntimeObjects, IMovieScenePlayer& Player, FMovieSceneSequenceInstance& SequenceInstance) override;
-	virtual void Update(EMovieSceneUpdateData& UpdateData, const TArray<UObject*>& RuntimeObjects, IMovieScenePlayer& Player, FMovieSceneSequenceInstance& SequenceInstance) override;
-	virtual void RefreshInstance( const TArray<UObject*>& RuntimeObjects, IMovieScenePlayer& Player, FMovieSceneSequenceInstance& SequenceInstance ) override {}
+	virtual void SaveState (const TArray<TWeakObjectPtr<UObject>>& RuntimeObjects, IMovieScenePlayer& Player, FMovieSceneSequenceInstance& SequenceInstance) override;
+	virtual void RestoreState (const TArray<TWeakObjectPtr<UObject>>& RuntimeObjects, IMovieScenePlayer& Player, FMovieSceneSequenceInstance& SequenceInstance) override;
+	virtual void Update(EMovieSceneUpdateData& UpdateData, const TArray<TWeakObjectPtr<UObject>>& RuntimeObjects, IMovieScenePlayer& Player, FMovieSceneSequenceInstance& SequenceInstance) override;
+	virtual void RefreshInstance( const TArray<TWeakObjectPtr<UObject>>& RuntimeObjects, IMovieScenePlayer& Player, FMovieSceneSequenceInstance& SequenceInstance ) override {}
 	virtual void ClearInstance( IMovieScenePlayer& Player, FMovieSceneSequenceInstance& SequenceInstance ) override {}
-	virtual float EvalOrder() { return -1.f; }
 
-	virtual void UpdateConstraint(float Position, const TArray<UObject*>& RuntimeObjects, AActor* Actor, UMovieScene3DConstraintSection* ConstraintSection) = 0;
+	// Evaluate after TransformTrackInstance so that parents can be evaluated first
+	virtual EMovieSceneUpdatePass HasUpdatePasses() override { return MSUP_PostUpdate; }
+
+	virtual void UpdateConstraint(float Position, const TArray<TWeakObjectPtr<UObject>>& RuntimeObjects, AActor* Actor, UMovieScene3DConstraintSection* ConstraintSection) = 0;
 
 protected:
 	/** Track that is being instanced */
diff --git a/Engine/Source/Runtime/MovieSceneTracks/Public/TrackInstances/MovieScene3DPathTrackInstance.h b/Engine/Source/Runtime/MovieSceneTracks/Public/TrackInstances/MovieScene3DPathTrackInstance.h
--- a/Engine/Source/Runtime/MovieSceneTracks/Public/TrackInstances/MovieScene3DPathTrackInstance.h
+++ b/Engine/Source/Runtime/MovieSceneTracks/Public/TrackInstances/MovieScene3DPathTrackInstance.h
@@ -17,5 +17,5 @@ public:
 	FMovieScene3DPathTrackInstance( UMovieScene3DPathTrack& InPathTrack );
 
 	/** FMovieScene3DConstraintTrackInstance interface */
-	virtual void UpdateConstraint(float Position, const TArray<UObject*>& RuntimeObjects, AActor* Actor, UMovieScene3DConstraintSection* ConstraintSection) override;
+	virtual void UpdateConstraint(float Position, const TArray<TWeakObjectPtr<UObject>>& RuntimeObjects, AActor* Actor, UMovieScene3DConstraintSection* ConstraintSection) override;
 };
diff --git a/Engine/Source/Runtime/MovieSceneTracks/Public/TrackInstances/MovieScene3DTransformTrackInstance.h b/Engine/Source/Runtime/MovieSceneTracks/Public/TrackInstances/MovieScene3DTransformTrackInstance.h
--- a/Engine/Source/Runtime/MovieSceneTracks/Public/TrackInstances/MovieScene3DTransformTrackInstance.h
+++ b/Engine/Source/Runtime/MovieSceneTracks/Public/TrackInstances/MovieScene3DTransformTrackInstance.h
@@ -19,16 +19,16 @@ public:
 	FMovieScene3DTransformTrackInstance( UMovieScene3DTransformTrack& InTransformTrack );
 
 	/** IMovieSceneTrackInstance interface */
-	virtual void SaveState (const TArray<UObject*>& RuntimeObjects, IMovieScenePlayer& Player, FMovieSceneSequenceInstance& SequenceInstance) override;
-	virtual void RestoreState (const TArray<UObject*>& RuntimeObjects, IMovieScenePlayer& Player, FMovieSceneSequenceInstance& SequenceInstance) override;
-	virtual void Update(EMovieSceneUpdateData& UpdateData, const TArray<UObject*>& RuntimeObjects, IMovieScenePlayer& Player, FMovieSceneSequenceInstance& SequenceInstance) override;
+	virtual void SaveState (const TArray<TWeakObjectPtr<UObject>>& RuntimeObjects, IMovieScenePlayer& Player, FMovieSceneSequenceInstance& SequenceInstance) override;
+	virtual void RestoreState (const TArray<TWeakObjectPtr<UObject>>& RuntimeObjects, IMovieScenePlayer& Player, FMovieSceneSequenceInstance& SequenceInstance) override;
+	virtual void Update(EMovieSceneUpdateData& UpdateData, const TArray<TWeakObjectPtr<UObject>>& RuntimeObjects, IMovieScenePlayer& Player, FMovieSceneSequenceInstance& SequenceInstance) override;
 	virtual EMovieSceneUpdatePass HasUpdatePasses() override { return (EMovieSceneUpdatePass)(MSUP_PreUpdate | MSUP_Update); }
-	virtual void RefreshInstance( const TArray<UObject*>& RuntimeObjects, IMovieScenePlayer& Player, FMovieSceneSequenceInstance& SequenceInstance ) override;
+	virtual void RefreshInstance( const TArray<TWeakObjectPtr<UObject>>& RuntimeObjects, IMovieScenePlayer& Player, FMovieSceneSequenceInstance& SequenceInstance ) override;
 	virtual void ClearInstance( IMovieScenePlayer& Player, FMovieSceneSequenceInstance& SequenceInstance ) override {}
 
 private:
 
-	void UpdateRuntimeMobility(const TArray<UObject*>& RuntimeObjects);
+	void UpdateRuntimeMobility(const TArray<TWeakObjectPtr<UObject>>& RuntimeObjects);
 
 	/** Track that is being instanced */
 	UMovieScene3DTransformTrack* TransformTrack;
diff --git a/Engine/Source/Runtime/MovieSceneTracks/Public/TrackInstances/MovieSceneActorReferenceTrackInstance.h b/Engine/Source/Runtime/MovieSceneTracks/Public/TrackInstances/MovieSceneActorReferenceTrackInstance.h
new file mode 100644
--- /dev/null
+++ b/Engine/Source/Runtime/MovieSceneTracks/Public/TrackInstances/MovieSceneActorReferenceTrackInstance.h
@@ -0,0 +1,42 @@
+// Copyright 1998-2016 Epic Games, Inc. All Rights Reserved.
+
+#pragma once
+
+#include "ObjectKey.h"
+#include "IMovieSceneTrackInstance.h"
+
+
+class FTrackInstancePropertyBindings;
+class UMovieSceneActorReferenceTrack;
+
+
+/**
+ * Runtime instance of a UMovieSceneActorReferenceTrack
+ */
+class FMovieSceneActorReferenceTrackInstance
+	: public IMovieSceneTrackInstance
+{
+public:
+
+	FMovieSceneActorReferenceTrackInstance( UMovieSceneActorReferenceTrack& InActorReferenceTrack );
+
+	/** IMovieSceneTrackInstance interface */
+	virtual void SaveState (const TArray<TWeakObjectPtr<UObject>>& RuntimeObjects, IMovieScenePlayer& Player, FMovieSceneSequenceInstance& SequenceInstance) override;
+	virtual void RestoreState (const TArray<TWeakObjectPtr<UObject>>& RuntimeObjects, IMovieScenePlayer& Player, FMovieSceneSequenceInstance& SequenceInstance) override;
+	virtual void Update(EMovieSceneUpdateData& UpdateData, const TArray<TWeakObjectPtr<UObject>>& RuntimeObjects, IMovieScenePlayer& Player, FMovieSceneSequenceInstance& SequenceInstance) override;
+	virtual void RefreshInstance( const TArray<TWeakObjectPtr<UObject>>& RuntimeObjects, IMovieScenePlayer& Player, FMovieSceneSequenceInstance& SequenceInstance ) override;
+	virtual void ClearInstance( IMovieScenePlayer& Player, FMovieSceneSequenceInstance& SequenceInstance ) override {}
+
+private:
+
+	/** Track that is being instanced */
+	UMovieSceneActorReferenceTrack* ActorReferenceTrack;
+
+	/** Runtime property bindings */
+	TSharedPtr<FTrackInstancePropertyBindings> PropertyBindings;
+
+	/** Map from object to initial state */
+	TMap<FObjectKey, AActor*> InitActorReferenceMap;
+
+	TMap<FGuid, AActor*> GuidToActorCache;
+};
diff --git a/Engine/Source/Runtime/MovieSceneTracks/Public/TrackInstances/MovieSceneAudioTrackInstance.h b/Engine/Source/Runtime/MovieSceneTracks/Public/TrackInstances/MovieSceneAudioTrackInstance.h
--- a/Engine/Source/Runtime/MovieSceneTracks/Public/TrackInstances/MovieSceneAudioTrackInstance.h
+++ b/Engine/Source/Runtime/MovieSceneTracks/Public/TrackInstances/MovieSceneAudioTrackInstance.h
@@ -20,11 +20,11 @@ public:
 	FMovieSceneAudioTrackInstance( UMovieSceneAudioTrack& InAudioTrack );
 
 	/** IMovieSceneTrackInstance interface */
-	virtual void SaveState(const TArray<UObject*>& RuntimeObjects, IMovieScenePlayer& Player, FMovieSceneSequenceInstance& SequenceInstance) override {}
-	virtual void RestoreState(const TArray<UObject*>& RuntimeObjects, IMovieScenePlayer& Player, FMovieSceneSequenceInstance& SequenceInstance) override {}
-	virtual void Update(EMovieSceneUpdateData& UpdateData, const TArray<UObject*>& RuntimeObjects, IMovieScenePlayer& Player, FMovieSceneSequenceInstance& SequenceInstance) override;
-	virtual void RefreshInstance( const TArray<UObject*>& RuntimeObjects, IMovieScenePlayer& Player, FMovieSceneSequenceInstance& SequenceInstance ) override {}
-	virtual void ClearInstance( IMovieScenePlayer& Player, FMovieSceneSequenceInstance& SequenceInstance ) override {}
+	virtual void SaveState(const TArray<TWeakObjectPtr<UObject>>& RuntimeObjects, IMovieScenePlayer& Player, FMovieSceneSequenceInstance& SequenceInstance) override;
+	virtual void RestoreState(const TArray<TWeakObjectPtr<UObject>>& RuntimeObjects, IMovieScenePlayer& Player, FMovieSceneSequenceInstance& SequenceInstance) override;
+	virtual void Update(EMovieSceneUpdateData& UpdateData, const TArray<TWeakObjectPtr<UObject>>& RuntimeObjects, IMovieScenePlayer& Player, FMovieSceneSequenceInstance& SequenceInstance) override;
+	virtual void RefreshInstance(const TArray<TWeakObjectPtr<UObject>>& RuntimeObjects, IMovieScenePlayer& Player, FMovieSceneSequenceInstance& SequenceInstance) override {}
+	virtual void ClearInstance(IMovieScenePlayer& Player, FMovieSceneSequenceInstance& SequenceInstance) override {}
 
 private:
 
@@ -38,10 +38,10 @@ private:
 	void StopAllSounds();
 
 	/** Gets the audio component component for the actor and row index, creating it if necessary */
-	TWeakObjectPtr<UAudioComponent> GetAudioComponent(AActor* Actor, int32 RowIndex);
+	TWeakObjectPtr<UAudioComponent> GetAudioComponent(IMovieScenePlayer& Player, AActor* Actor, int32 RowIndex);
 
 	/** Utility function for getting actors from objects array */
-	TArray<AActor*> GetRuntimeActors(const TArray<UObject*>& RuntimeObjects) const;
+	TArray<AActor*> GetRuntimeActors(const TArray<TWeakObjectPtr<UObject>>& RuntimeObjects) const;
 
 private:
 
diff --git a/Engine/Source/Runtime/MovieSceneTracks/Public/TrackInstances/MovieSceneBoolTrackInstance.h b/Engine/Source/Runtime/MovieSceneTracks/Public/TrackInstances/MovieSceneBoolTrackInstance.h
--- a/Engine/Source/Runtime/MovieSceneTracks/Public/TrackInstances/MovieSceneBoolTrackInstance.h
+++ b/Engine/Source/Runtime/MovieSceneTracks/Public/TrackInstances/MovieSceneBoolTrackInstance.h
@@ -20,10 +20,10 @@ public:
 	FMovieSceneBoolTrackInstance( UMovieSceneBoolTrack& InBoolTrack );
 
 	/** IMovieSceneTrackInstance interface */
-	virtual void SaveState (const TArray<UObject*>& RuntimeObjects, IMovieScenePlayer& Player, FMovieSceneSequenceInstance& SequenceInstance) override;
-	virtual void RestoreState (const TArray<UObject*>& RuntimeObjects, IMovieScenePlayer& Player, FMovieSceneSequenceInstance& SequenceInstance) override;
-	virtual void Update(EMovieSceneUpdateData& UpdateData, const TArray<UObject*>& RuntimeObjects, IMovieScenePlayer& Player, FMovieSceneSequenceInstance& SequenceInstance) override;
-	virtual void RefreshInstance( const TArray<UObject*>& RuntimeObjects, IMovieScenePlayer& Player, FMovieSceneSequenceInstance& SequenceInstance ) override;
+	virtual void SaveState (const TArray<TWeakObjectPtr<UObject>>& RuntimeObjects, IMovieScenePlayer& Player, FMovieSceneSequenceInstance& SequenceInstance) override;
+	virtual void RestoreState (const TArray<TWeakObjectPtr<UObject>>& RuntimeObjects, IMovieScenePlayer& Player, FMovieSceneSequenceInstance& SequenceInstance) override;
+	virtual void Update(EMovieSceneUpdateData& UpdateData, const TArray<TWeakObjectPtr<UObject>>& RuntimeObjects, IMovieScenePlayer& Player, FMovieSceneSequenceInstance& SequenceInstance) override;
+	virtual void RefreshInstance( const TArray<TWeakObjectPtr<UObject>>& RuntimeObjects, IMovieScenePlayer& Player, FMovieSceneSequenceInstance& SequenceInstance ) override;
 	virtual void ClearInstance( IMovieScenePlayer& Player, FMovieSceneSequenceInstance& SequenceInstance ) override {}
 
 private:
diff --git a/Engine/Source/Runtime/MovieSceneTracks/Public/TrackInstances/MovieSceneByteTrackInstance.h b/Engine/Source/Runtime/MovieSceneTracks/Public/TrackInstances/MovieSceneByteTrackInstance.h
--- a/Engine/Source/Runtime/MovieSceneTracks/Public/TrackInstances/MovieSceneByteTrackInstance.h
+++ b/Engine/Source/Runtime/MovieSceneTracks/Public/TrackInstances/MovieSceneByteTrackInstance.h
@@ -20,10 +20,10 @@ public:
 	FMovieSceneByteTrackInstance( UMovieSceneByteTrack& InByteTrack );
 
 	/** IMovieSceneTrackInstance interface */
-	virtual void SaveState (const TArray<UObject*>& RuntimeObjects, IMovieScenePlayer& Player, FMovieSceneSequenceInstance& SequenceInstance) override;
-	virtual void RestoreState (const TArray<UObject*>& RuntimeObjects, IMovieScenePlayer& Player, FMovieSceneSequenceInstance& SequenceInstance) override;
-	virtual void Update(EMovieSceneUpdateData& UpdateData, const TArray<UObject*>& RuntimeObjects, IMovieScenePlayer& Player, FMovieSceneSequenceInstance& SequenceInstance) override;
-	virtual void RefreshInstance( const TArray<UObject*>& RuntimeObjects, IMovieScenePlayer& Player, FMovieSceneSequenceInstance& SequenceInstance ) override;
+	virtual void SaveState (const TArray<TWeakObjectPtr<UObject>>& RuntimeObjects, IMovieScenePlayer& Player, FMovieSceneSequenceInstance& SequenceInstance) override;
+	virtual void RestoreState (const TArray<TWeakObjectPtr<UObject>>& RuntimeObjects, IMovieScenePlayer& Player, FMovieSceneSequenceInstance& SequenceInstance) override;
+	virtual void Update(EMovieSceneUpdateData& UpdateData, const TArray<TWeakObjectPtr<UObject>>& RuntimeObjects, IMovieScenePlayer& Player, FMovieSceneSequenceInstance& SequenceInstance) override;
+	virtual void RefreshInstance( const TArray<TWeakObjectPtr<UObject>>& RuntimeObjects, IMovieScenePlayer& Player, FMovieSceneSequenceInstance& SequenceInstance ) override;
 	virtual void ClearInstance( IMovieScenePlayer& Player, FMovieSceneSequenceInstance& SequenceInstance ) override {}
 
 private:
diff --git a/Engine/Source/Runtime/MovieSceneTracks/Public/TrackInstances/MovieSceneCameraAnimTrackInstance.h b/Engine/Source/Runtime/MovieSceneTracks/Public/TrackInstances/MovieSceneCameraAnimTrackInstance.h
new file mode 100644
--- /dev/null
+++ b/Engine/Source/Runtime/MovieSceneTracks/Public/TrackInstances/MovieSceneCameraAnimTrackInstance.h
@@ -0,0 +1,49 @@
+// Copyright 1998-2016 Epic Games, Inc. All Rights Reserved.
+
+#pragma once
+
+#include "IMovieSceneTrackInstance.h"
+
+class UMovieSceneCameraAnimTrack;
+class UMovieSceneCameraAnimSection;
+
+struct FMovieSceneCameraAnimSectionInstanceData
+{
+	TWeakObjectPtr<UCameraAnimInst> CameraAnimInst;
+	FTransform AdditiveOffset;		// AdditiveCamToBaseCam
+	float AdditiveFOVOffset;
+};
+
+/**
+ * Instance of a UMovieSceneCameraAnimTrack
+ */
+class FMovieSceneCameraAnimTrackInstance
+	: public IMovieSceneTrackInstance
+{
+public:
+
+	FMovieSceneCameraAnimTrackInstance(UMovieSceneCameraAnimTrack& InAnimationTrack);
+	virtual ~FMovieSceneCameraAnimTrackInstance();
+
+	/** IMovieSceneTrackInstance interface */
+	virtual void SaveState(const TArray<TWeakObjectPtr<UObject>>& RuntimeObjects, IMovieScenePlayer& Player, FMovieSceneSequenceInstance& SequenceInstance) override {}
+	virtual void RestoreState(const TArray<TWeakObjectPtr<UObject>>& RuntimeObjects, IMovieScenePlayer& Player, FMovieSceneSequenceInstance& SequenceInstance) override {}
+	virtual void Update(EMovieSceneUpdateData& UpdateData, const TArray<TWeakObjectPtr<UObject>>& RuntimeObjects, IMovieScenePlayer& Player, FMovieSceneSequenceInstance& SequenceInstance) override;
+	virtual void RefreshInstance(const TArray<TWeakObjectPtr<UObject>>& RuntimeObjects, IMovieScenePlayer& Player, FMovieSceneSequenceInstance& SequenceInstance) override {}
+	virtual void ClearInstance(IMovieScenePlayer& Player, FMovieSceneSequenceInstance& SequenceInstance) override {}
+
+	virtual EMovieSceneUpdatePass HasUpdatePasses() override { return (EMovieSceneUpdatePass)(MSUP_PreUpdate | MSUP_Update); }
+
+private:
+	/** Track that is being instanced */
+	UMovieSceneCameraAnimTrack* CameraAnimTrack;
+
+	/** Keep track of the currently playing montage */
+	TWeakObjectPtr<UCameraAnimInst> CurrentlyPlayingCameraAnimInst;
+
+	TWeakObjectPtr<ACameraActor> TempCameraActor;
+
+	ACameraActor* GetTempCameraActor();
+
+	TMap<UMovieSceneCameraAnimSection*, FMovieSceneCameraAnimSectionInstanceData> SectionInstanceDataMap;
+};
diff --git a/Engine/Source/Runtime/MovieSceneTracks/Public/TrackInstances/MovieSceneCameraCutTrackInstance.h b/Engine/Source/Runtime/MovieSceneTracks/Public/TrackInstances/MovieSceneCameraCutTrackInstance.h
--- a/Engine/Source/Runtime/MovieSceneTracks/Public/TrackInstances/MovieSceneCameraCutTrackInstance.h
+++ b/Engine/Source/Runtime/MovieSceneTracks/Public/TrackInstances/MovieSceneCameraCutTrackInstance.h
@@ -22,10 +22,10 @@ public:
 	// IMovieSceneTrackInstance interface
 
 	virtual void ClearInstance(IMovieScenePlayer& Player, FMovieSceneSequenceInstance& SequenceInstance) override;
-	virtual void RefreshInstance(const TArray<UObject*>& RuntimeObjects, IMovieScenePlayer& Player, FMovieSceneSequenceInstance& SequenceInstance) override;
-	virtual void RestoreState(const TArray<UObject*>& RuntimeObjects, IMovieScenePlayer& Player, FMovieSceneSequenceInstance& SequenceInstance) override;
-	virtual void SaveState(const TArray<UObject*>& RuntimeObjects, IMovieScenePlayer& Player, FMovieSceneSequenceInstance& SequenceInstance) override { }
-	virtual void Update(EMovieSceneUpdateData& UpdateData, const TArray<UObject*>& RuntimeObjects, IMovieScenePlayer& Player, FMovieSceneSequenceInstance& SequenceInstance) override;
+	virtual void RefreshInstance(const TArray<TWeakObjectPtr<UObject>>& RuntimeObjects, IMovieScenePlayer& Player, FMovieSceneSequenceInstance& SequenceInstance) override;
+	virtual void RestoreState(const TArray<TWeakObjectPtr<UObject>>& RuntimeObjects, IMovieScenePlayer& Player, FMovieSceneSequenceInstance& SequenceInstance) override;
+	virtual void SaveState(const TArray<TWeakObjectPtr<UObject>>& RuntimeObjects, IMovieScenePlayer& Player, FMovieSceneSequenceInstance& SequenceInstance) override { }
+	virtual void Update(EMovieSceneUpdateData& UpdateData, const TArray<TWeakObjectPtr<UObject>>& RuntimeObjects, IMovieScenePlayer& Player, FMovieSceneSequenceInstance& SequenceInstance) override;
 
 protected:
 
diff --git a/Engine/Source/Runtime/MovieSceneTracks/Public/TrackInstances/MovieSceneCameraShakeTrackInstance.h b/Engine/Source/Runtime/MovieSceneTracks/Public/TrackInstances/MovieSceneCameraShakeTrackInstance.h
new file mode 100644
--- /dev/null
+++ b/Engine/Source/Runtime/MovieSceneTracks/Public/TrackInstances/MovieSceneCameraShakeTrackInstance.h
@@ -0,0 +1,50 @@
+// Copyright 1998-2016 Epic Games, Inc. All Rights Reserved.
+
+#pragma once
+
+#include "IMovieSceneTrackInstance.h"
+
+class UMovieSceneCameraShakeTrack;
+class UMovieSceneCameraShakeSection;
+
+
+struct FMovieSceneCameraShakeSectionInstanceData
+{
+	TWeakObjectPtr<UCameraShake> CameraShakeInst;
+	FTransform AdditiveOffset;		// AdditiveCamToBaseCam
+	float AdditiveFOVOffset;
+};
+
+/**
+* Instance of a UMovieSceneCameraShakeTrack
+*/
+class FMovieSceneCameraShakeTrackInstance
+	: public IMovieSceneTrackInstance
+{
+public:
+
+	FMovieSceneCameraShakeTrackInstance(UMovieSceneCameraShakeTrack& InShakeTrack);
+	virtual ~FMovieSceneCameraShakeTrackInstance();
+
+	/** IMovieSceneTrackInstance interface */
+	virtual void SaveState(const TArray<TWeakObjectPtr<UObject>>& RuntimeObjects, IMovieScenePlayer& Player, FMovieSceneSequenceInstance& SequenceInstance) override {}
+	virtual void RestoreState(const TArray<TWeakObjectPtr<UObject>>& RuntimeObjects, IMovieScenePlayer& Player, FMovieSceneSequenceInstance& SequenceInstance) override {}
+	virtual void Update(EMovieSceneUpdateData& UpdateData, const TArray<TWeakObjectPtr<UObject>>& RuntimeObjects, IMovieScenePlayer& Player, FMovieSceneSequenceInstance& SequenceInstance) override;
+	virtual void RefreshInstance(const TArray<TWeakObjectPtr<UObject>>& RuntimeObjects, IMovieScenePlayer& Player, FMovieSceneSequenceInstance& SequenceInstance) override {}
+	virtual void ClearInstance(IMovieScenePlayer& Player, FMovieSceneSequenceInstance& SequenceInstance) override {}
+
+	virtual EMovieSceneUpdatePass HasUpdatePasses() override { return (EMovieSceneUpdatePass)(MSUP_PreUpdate | MSUP_Update); }
+
+private:
+	/** Track that is being instanced */
+	UMovieSceneCameraShakeTrack* CameraShakeTrack;
+
+	/** Keep track of the currently playing montage */
+	TWeakObjectPtr<UCameraShake> CurrentlyPlayingCameraShakeInst;
+
+	TWeakObjectPtr<ACameraActor> TempCameraActor;
+
+	TMap<UMovieSceneCameraShakeSection*, FMovieSceneCameraShakeSectionInstanceData> SectionInstanceDataMap;
+
+	ACameraActor* GetTempCameraActor();
+};
diff --git a/Engine/Source/Runtime/MovieSceneTracks/Public/TrackInstances/MovieSceneCinematicShotTrackInstance.h b/Engine/Source/Runtime/MovieSceneTracks/Public/TrackInstances/MovieSceneCinematicShotTrackInstance.h
--- a/Engine/Source/Runtime/MovieSceneTracks/Public/TrackInstances/MovieSceneCinematicShotTrackInstance.h
+++ b/Engine/Source/Runtime/MovieSceneTracks/Public/TrackInstances/MovieSceneCinematicShotTrackInstance.h
@@ -24,8 +24,13 @@ public:
 	// IMovieSceneTrackInstance interface
 
 	virtual void ClearInstance(IMovieScenePlayer& Player, FMovieSceneSequenceInstance& SequenceInstance) override;
-	virtual void RefreshInstance(const TArray<UObject*>& RuntimeObjects, IMovieScenePlayer& Player, FMovieSceneSequenceInstance& SequenceInstance) override;
-	virtual void RestoreState (const TArray<UObject*>& RuntimeObjects, IMovieScenePlayer& Player, FMovieSceneSequenceInstance& SequenceInstance) override;
-	virtual void SaveState(const TArray<UObject*>& RuntimeObjects, IMovieScenePlayer& Player, FMovieSceneSequenceInstance& SequenceInstance) override { }
-	virtual void Update(EMovieSceneUpdateData& UpdateData, const TArray<UObject*>& RuntimeObjects, IMovieScenePlayer& Player, FMovieSceneSequenceInstance& SequenceInstance) override;
+	virtual void RefreshInstance(const TArray<TWeakObjectPtr<UObject>>& RuntimeObjects, IMovieScenePlayer& Player, FMovieSceneSequenceInstance& SequenceInstance) override;
+	virtual void RestoreState (const TArray<TWeakObjectPtr<UObject>>& RuntimeObjects, IMovieScenePlayer& Player, FMovieSceneSequenceInstance& SequenceInstance) override;
+	virtual void SaveState(const TArray<TWeakObjectPtr<UObject>>& RuntimeObjects, IMovieScenePlayer& Player, FMovieSceneSequenceInstance& SequenceInstance) override { }
+	virtual void Update(EMovieSceneUpdateData& UpdateData, const TArray<TWeakObjectPtr<UObject>>& RuntimeObjects, IMovieScenePlayer& Player, FMovieSceneSequenceInstance& SequenceInstance) override;
+
+protected:
+
+	// FMovieSceneSubTrackinstance
+	virtual bool ShouldEvaluateIfOverlapping(const TArray<UMovieSceneSection*>& Sections, UMovieSceneSection* Section) const override;
 };
diff --git a/Engine/Source/Runtime/MovieSceneTracks/Public/TrackInstances/MovieSceneEventTrackInstance.h b/Engine/Source/Runtime/MovieSceneTracks/Public/TrackInstances/MovieSceneEventTrackInstance.h
--- a/Engine/Source/Runtime/MovieSceneTracks/Public/TrackInstances/MovieSceneEventTrackInstance.h
+++ b/Engine/Source/Runtime/MovieSceneTracks/Public/TrackInstances/MovieSceneEventTrackInstance.h
@@ -31,10 +31,10 @@ public:
 	// IMovieSceneTrackInstance interface
 
 	virtual void ClearInstance(IMovieScenePlayer& Player, FMovieSceneSequenceInstance& SequenceInstance) override { }
-	virtual void RefreshInstance(const TArray<UObject*>& RuntimeObjects, IMovieScenePlayer& Player, FMovieSceneSequenceInstance& SequenceInstance) override { }
-	virtual void RestoreState(const TArray<UObject*>& RuntimeObjects, IMovieScenePlayer& Player, FMovieSceneSequenceInstance& SequenceInstance) override { }
-	virtual void SaveState(const TArray<UObject*>& RuntimeObjects, IMovieScenePlayer& Player, FMovieSceneSequenceInstance& SequenceInstance) override { }
-	virtual void Update(EMovieSceneUpdateData& UpdateData, const TArray<UObject*>& RuntimeObjects, IMovieScenePlayer& Player, FMovieSceneSequenceInstance& SequenceInstance) override;
+	virtual void RefreshInstance(const TArray<TWeakObjectPtr<UObject>>& RuntimeObjects, IMovieScenePlayer& Player, FMovieSceneSequenceInstance& SequenceInstance) override { }
+	virtual void RestoreState(const TArray<TWeakObjectPtr<UObject>>& RuntimeObjects, IMovieScenePlayer& Player, FMovieSceneSequenceInstance& SequenceInstance) override { }
+	virtual void SaveState(const TArray<TWeakObjectPtr<UObject>>& RuntimeObjects, IMovieScenePlayer& Player, FMovieSceneSequenceInstance& SequenceInstance) override { }
+	virtual void Update(EMovieSceneUpdateData& UpdateData, const TArray<TWeakObjectPtr<UObject>>& RuntimeObjects, IMovieScenePlayer& Player, FMovieSceneSequenceInstance& SequenceInstance) override;
 
 private:
 
diff --git a/Engine/Source/Runtime/MovieSceneTracks/Public/TrackInstances/MovieSceneFadeTrackInstance.h b/Engine/Source/Runtime/MovieSceneTracks/Public/TrackInstances/MovieSceneFadeTrackInstance.h
--- a/Engine/Source/Runtime/MovieSceneTracks/Public/TrackInstances/MovieSceneFadeTrackInstance.h
+++ b/Engine/Source/Runtime/MovieSceneTracks/Public/TrackInstances/MovieSceneFadeTrackInstance.h
@@ -31,10 +31,10 @@ public:
 	// IMovieSceneTrackInstance interface
 
 	virtual void ClearInstance(IMovieScenePlayer& Player, FMovieSceneSequenceInstance& SequenceInstance) override { }
-	virtual void RefreshInstance(const TArray<UObject*>& RuntimeObjects, IMovieScenePlayer& Player, FMovieSceneSequenceInstance& SequenceInstance) override { }
-	virtual void RestoreState(const TArray<UObject*>& RuntimeObjects, IMovieScenePlayer& Player, FMovieSceneSequenceInstance& SequenceInstance) override;
-	virtual void SaveState(const TArray<UObject*>& RuntimeObjects, IMovieScenePlayer& Player, FMovieSceneSequenceInstance& SequenceInstance) override { }
-	virtual void Update(EMovieSceneUpdateData& UpdateData, const TArray<UObject*>& RuntimeObjects, IMovieScenePlayer& Player, FMovieSceneSequenceInstance& SequenceInstance) override;
+	virtual void RefreshInstance(const TArray<TWeakObjectPtr<UObject>>& RuntimeObjects, IMovieScenePlayer& Player, FMovieSceneSequenceInstance& SequenceInstance) override { }
+	virtual void RestoreState(const TArray<TWeakObjectPtr<UObject>>& RuntimeObjects, IMovieScenePlayer& Player, FMovieSceneSequenceInstance& SequenceInstance) override;
+	virtual void SaveState(const TArray<TWeakObjectPtr<UObject>>& RuntimeObjects, IMovieScenePlayer& Player, FMovieSceneSequenceInstance& SequenceInstance) override { }
+	virtual void Update(EMovieSceneUpdateData& UpdateData, const TArray<TWeakObjectPtr<UObject>>& RuntimeObjects, IMovieScenePlayer& Player, FMovieSceneSequenceInstance& SequenceInstance) override;
 
 private:
 
diff --git a/Engine/Source/Runtime/MovieSceneTracks/Public/TrackInstances/MovieSceneFloatTrackInstance.h b/Engine/Source/Runtime/MovieSceneTracks/Public/TrackInstances/MovieSceneFloatTrackInstance.h
--- a/Engine/Source/Runtime/MovieSceneTracks/Public/TrackInstances/MovieSceneFloatTrackInstance.h
+++ b/Engine/Source/Runtime/MovieSceneTracks/Public/TrackInstances/MovieSceneFloatTrackInstance.h
@@ -20,11 +20,11 @@ public:
 	FMovieSceneFloatTrackInstance( UMovieSceneFloatTrack& InFloatTrack );
 
 	/** IMovieSceneTrackInstance interface */
-	virtual void SaveState (const TArray<UObject*>& RuntimeObjects, IMovieScenePlayer& Player, FMovieSceneSequenceInstance& SequenceInstance) override;
-	virtual void RestoreState (const TArray<UObject*>& RuntimeObjects, IMovieScenePlayer& Player, FMovieSceneSequenceInstance& SequenceInstance) override;
-	virtual void Update(EMovieSceneUpdateData& UpdateData, const TArray<UObject*>& RuntimeObjects, IMovieScenePlayer& Player, FMovieSceneSequenceInstance& SequenceInstance) override;
-	virtual void RefreshInstance( const TArray<UObject*>& RuntimeObjects, IMovieScenePlayer& Player, FMovieSceneSequenceInstance& SequenceInstance ) override;
-	virtual void ClearInstance( IMovieScenePlayer& Player, FMovieSceneSequenceInstance& SequenceInstance ) override {}
+	virtual void SaveState(const TArray<TWeakObjectPtr<UObject>>& RuntimeObjects, IMovieScenePlayer& Player, FMovieSceneSequenceInstance& SequenceInstance) override;
+	virtual void RestoreState(const TArray<TWeakObjectPtr<UObject>>& RuntimeObjects, IMovieScenePlayer& Player, FMovieSceneSequenceInstance& SequenceInstance) override;
+	virtual void Update(EMovieSceneUpdateData& UpdateData, const TArray<TWeakObjectPtr<UObject>>& RuntimeObjects, IMovieScenePlayer& Player, FMovieSceneSequenceInstance& SequenceInstance) override;
+	virtual void RefreshInstance(const TArray<TWeakObjectPtr<UObject>>& RuntimeObjects, IMovieScenePlayer& Player, FMovieSceneSequenceInstance& SequenceInstance ) override;
+	virtual void ClearInstance(IMovieScenePlayer& Player, FMovieSceneSequenceInstance& SequenceInstance ) override {}
 
 private:
 
diff --git a/Engine/Source/Runtime/MovieSceneTracks/Public/TrackInstances/MovieSceneMaterialTrackInstance.h b/Engine/Source/Runtime/MovieSceneTracks/Public/TrackInstances/MovieSceneMaterialTrackInstance.h
--- a/Engine/Source/Runtime/MovieSceneTracks/Public/TrackInstances/MovieSceneMaterialTrackInstance.h
+++ b/Engine/Source/Runtime/MovieSceneTracks/Public/TrackInstances/MovieSceneMaterialTrackInstance.h
@@ -19,11 +19,11 @@ public:
 	FMovieSceneMaterialTrackInstance( UMovieSceneMaterialTrack& InMaterialTrack );
 
 	/** IMovieSceneTrackInstance interface */
-	virtual void SaveState (const TArray<UObject*>& RuntimeObjects, IMovieScenePlayer& Player, FMovieSceneSequenceInstance& SequenceInstance) override;
-	virtual void RestoreState (const TArray<UObject*>& RuntimeObjects, IMovieScenePlayer& Player, FMovieSceneSequenceInstance& SequenceInstance) override;
-	virtual void Update(EMovieSceneUpdateData& UpdateData, const TArray<UObject*>& RuntimeObjects, IMovieScenePlayer& Player, FMovieSceneSequenceInstance& SequenceInstance) override;
-	virtual void RefreshInstance( const TArray<UObject*>& RuntimeObjects, IMovieScenePlayer& Player, FMovieSceneSequenceInstance& SequenceInstance ) override;
-	virtual void ClearInstance( IMovieScenePlayer& Player, FMovieSceneSequenceInstance& SequenceInstance ) override;
+	virtual void SaveState(const TArray<TWeakObjectPtr<UObject>>& RuntimeObjects, IMovieScenePlayer& Player, FMovieSceneSequenceInstance& SequenceInstance) override;
+	virtual void RestoreState(const TArray<TWeakObjectPtr<UObject>>& RuntimeObjects, IMovieScenePlayer& Player, FMovieSceneSequenceInstance& SequenceInstance) override;
+	virtual void Update(EMovieSceneUpdateData& UpdateData, const TArray<TWeakObjectPtr<UObject>>& RuntimeObjects, IMovieScenePlayer& Player, FMovieSceneSequenceInstance& SequenceInstance) override;
+	virtual void RefreshInstance(const TArray<TWeakObjectPtr<UObject>>& RuntimeObjects, IMovieScenePlayer& Player, FMovieSceneSequenceInstance& SequenceInstance ) override;
+	virtual void ClearInstance(IMovieScenePlayer& Player, FMovieSceneSequenceInstance& SequenceInstance ) override;
 
 protected:
 	/**
diff --git a/Engine/Source/Runtime/MovieSceneTracks/Public/TrackInstances/MovieSceneParticleParameterTrackInstance.h b/Engine/Source/Runtime/MovieSceneTracks/Public/TrackInstances/MovieSceneParticleParameterTrackInstance.h
--- a/Engine/Source/Runtime/MovieSceneTracks/Public/TrackInstances/MovieSceneParticleParameterTrackInstance.h
+++ b/Engine/Source/Runtime/MovieSceneTracks/Public/TrackInstances/MovieSceneParticleParameterTrackInstance.h
@@ -27,11 +27,11 @@ public:
 	FMovieSceneParticleParameterTrackInstance( UMovieSceneParticleParameterTrack& InParticleParameterTrack );
 
 	/** IMovieSceneTrackInstance interface */
-	virtual void SaveState (const TArray<UObject*>& RuntimeObjects, IMovieScenePlayer& Player, FMovieSceneSequenceInstance& SequenceInstance) override;
-	virtual void RestoreState (const TArray<UObject*>& RuntimeObjects, IMovieScenePlayer& Player, FMovieSceneSequenceInstance& SequenceInstance) override;
-	virtual void Update(EMovieSceneUpdateData& UpdateData, const TArray<UObject*>& RuntimeObjects, IMovieScenePlayer& Player, FMovieSceneSequenceInstance& SequenceInstance) override;
-	virtual void RefreshInstance( const TArray<UObject*>& RuntimeObjects, IMovieScenePlayer& Player, FMovieSceneSequenceInstance& SequenceInstance ) override;
-	virtual void ClearInstance( IMovieScenePlayer& Player, FMovieSceneSequenceInstance& SequenceInstance ) override;
+	virtual void SaveState(const TArray<TWeakObjectPtr<UObject>>& RuntimeObjects, IMovieScenePlayer& Player, FMovieSceneSequenceInstance& SequenceInstance) override;
+	virtual void RestoreState(const TArray<TWeakObjectPtr<UObject>>& RuntimeObjects, IMovieScenePlayer& Player, FMovieSceneSequenceInstance& SequenceInstance) override;
+	virtual void Update(EMovieSceneUpdateData& UpdateData, const TArray<TWeakObjectPtr<UObject>>& RuntimeObjects, IMovieScenePlayer& Player, FMovieSceneSequenceInstance& SequenceInstance) override;
+	virtual void RefreshInstance(const TArray<TWeakObjectPtr<UObject>>& RuntimeObjects, IMovieScenePlayer& Player, FMovieSceneSequenceInstance& SequenceInstance ) override;
+	virtual void ClearInstance(IMovieScenePlayer& Player, FMovieSceneSequenceInstance& SequenceInstance ) override;
 
 private:
 
diff --git a/Engine/Source/Runtime/MovieSceneTracks/Public/TrackInstances/MovieSceneParticleTrackInstance.h b/Engine/Source/Runtime/MovieSceneTracks/Public/TrackInstances/MovieSceneParticleTrackInstance.h
--- a/Engine/Source/Runtime/MovieSceneTracks/Public/TrackInstances/MovieSceneParticleTrackInstance.h
+++ b/Engine/Source/Runtime/MovieSceneTracks/Public/TrackInstances/MovieSceneParticleTrackInstance.h
@@ -23,11 +23,11 @@ public:
 	virtual ~FMovieSceneParticleTrackInstance();
 
 	/** IMovieSceneTrackInstance interface */
-	virtual void SaveState(const TArray<UObject*>& RuntimeObjects, IMovieScenePlayer& Player, FMovieSceneSequenceInstance& SequenceInstance) override {}
-	virtual void RestoreState(const TArray<UObject*>& RuntimeObjects, IMovieScenePlayer& Player, FMovieSceneSequenceInstance& SequenceInstance) override {}
-	virtual void Update(EMovieSceneUpdateData& UpdateData, const TArray<UObject*>& RuntimeObjects, IMovieScenePlayer& Player, FMovieSceneSequenceInstance& SequenceInstance) override;
-	virtual void RefreshInstance( const TArray<UObject*>& RuntimeObjects, IMovieScenePlayer& Player, FMovieSceneSequenceInstance& SequenceInstance ) override {}
-	virtual void ClearInstance( IMovieScenePlayer& Player, FMovieSceneSequenceInstance& SequenceInstance ) override {}
+	virtual void SaveState(const TArray<TWeakObjectPtr<UObject>>& RuntimeObjects, IMovieScenePlayer& Player, FMovieSceneSequenceInstance& SequenceInstance) override {}
+	virtual void RestoreState(const TArray<TWeakObjectPtr<UObject>>& RuntimeObjects, IMovieScenePlayer& Player, FMovieSceneSequenceInstance& SequenceInstance) override {}
+	virtual void Update(EMovieSceneUpdateData& UpdateData, const TArray<TWeakObjectPtr<UObject>>& RuntimeObjects, IMovieScenePlayer& Player, FMovieSceneSequenceInstance& SequenceInstance) override;
+	virtual void RefreshInstance(const TArray<TWeakObjectPtr<UObject>>& RuntimeObjects, IMovieScenePlayer& Player, FMovieSceneSequenceInstance& SequenceInstance ) override {}
+	virtual void ClearInstance(IMovieScenePlayer& Player, FMovieSceneSequenceInstance& SequenceInstance ) override {}
 
 private:
 
