Copying //UE4/Dev-Sequencer to //UE4/Main (Source: //UE4/Dev-Sequencer @ 2945541)

diff --git a/Engine/Source/Runtime/MovieSceneTracks/Private/TrackInstances/MovieSceneLevelVisibilityTrackInstance.cpp b/Engine/Source/Runtime/MovieSceneTracks/Private/TrackInstances/MovieSceneLevelVisibilityTrackInstance.cpp
--- a/Engine/Source/Runtime/MovieSceneTracks/Private/TrackInstances/MovieSceneLevelVisibilityTrackInstance.cpp
+++ b/Engine/Source/Runtime/MovieSceneTracks/Private/TrackInstances/MovieSceneLevelVisibilityTrackInstance.cpp
@@ -3,6 +3,9 @@
 #include "MovieSceneTracksPrivatePCH.h"
 #include "MovieSceneLevelVisibilityTrack.h"
 #include "MovieSceneLevelVisibilityTrackInstance.h"
+#if WITH_EDITOR
+#include "EditorLevelUtils.h"
+#endif
 
 
 FMovieSceneLevelVisibilityTrackInstance::FMovieSceneLevelVisibilityTrackInstance( UMovieSceneLevelVisibilityTrack& InLevelVisibilityTrack )
@@ -69,6 +72,34 @@ void GetAnimatedStreamingLevels( UMovieSceneLevelVisibilityTrack* LevelVisibilit
 	}
 }
 
+bool GetLevelVisibility( ULevelStreaming* Level )
+{
+#if WITH_EDITOR
+	if ( GIsEditor && Level->GetWorld()->IsPlayInEditor() == false )
+	{
+		return Level->bShouldBeVisibleInEditor;
+	}
+	else
+#endif
+	{
+		return Level->bShouldBeVisible;
+	}
+}
+
+void SetLevelVisibility( ULevelStreaming* Level, bool bVisible )
+{
+#if WITH_EDITOR
+	UE_LOG( LogMovieScene, Log, TEXT( "Setting Visibility: Level=%s Visible=%s" ), *Level->GetWorldAssetPackageName(), bVisible ? TEXT( "true" ) : TEXT( "false" ) );
+	if ( GIsEditor && Level->GetWorld()->IsPlayInEditor() == false )
+	{
+		EditorLevelUtils::SetLevelVisibility( Level->GetLoadedLevel(), bVisible, false );
+	}
+	else
+#endif
+	{
+		Level->bShouldBeVisible = bVisible;
+	}
+}
 
 void FMovieSceneLevelVisibilityTrackInstance::SaveState(const TArray<TWeakObjectPtr<UObject>>& RuntimeObjects, IMovieScenePlayer& Player, FMovieSceneSequenceInstance& SequenceInstance)
 {
@@ -92,7 +123,7 @@ void FMovieSceneLevelVisibilityTrackInstance::RestoreState(const TArray<TWeakObj
 		bool* SavedVisibility = SavedLevelVisibility.Find( FObjectKey( StreamingLevel ) );
 		if ( SavedVisibility != nullptr )
 		{
-			SetEditorLevelVisibilityState( StreamingLevel, *SavedVisibility );
+			SetLevelVisibility( StreamingLevel, *SavedVisibility );
 		}
 	}
 
@@ -102,58 +133,104 @@ void FMovieSceneLevelVisibilityTrackInstance::RestoreState(const TArray<TWeakObj
 
 void FMovieSceneLevelVisibilityTrackInstance::Update(EMovieSceneUpdateData& UpdateData, const TArray<TWeakObjectPtr<UObject>>& RuntimeObjects, class IMovieScenePlayer& Player, FMovieSceneSequenceInstance& SequenceInstance ) 
 {
-	for ( UMovieSceneSection* Section : LevelVisibilityTrack->GetAllSections() )
+	bool bVisibilityModified = false;
+	UWorld* CurrentWorld = nullptr;
+
+	TArray<UMovieSceneLevelVisibilitySection*> CurrentActiveSections;
+	if ( UpdateData.bSubSceneDeactivate == false )
 	{
-		UMovieSceneLevelVisibilitySection* LevelVisibilitySection = Cast<UMovieSceneLevelVisibilitySection>(Section);
-		if (LevelVisibilitySection != nullptr &&
-			LevelVisibilitySection->GetStartTime() <= UpdateData.Position &&
-			LevelVisibilitySection->GetEndTime() > UpdateData.Position )
+		for ( UMovieSceneSection* Section : LevelVisibilityTrack->GetAllSections() )
 		{
-			bool bShouldBeVisible = LevelVisibilitySection->GetVisibility() == ELevelVisibility::Visible ? true : false;
-			bool bVisibilityModified = false;
-			for ( FName LevelName : *LevelVisibilitySection->GetLevelNames() )
+			UMovieSceneLevelVisibilitySection* LevelVisibilitySection = Cast<UMovieSceneLevelVisibilitySection>( Section );
+			if ( LevelVisibilitySection != nullptr &&
+				LevelVisibilitySection->GetStartTime() <= UpdateData.Position &&
+				LevelVisibilitySection->GetEndTime() > UpdateData.Position )
 			{
-				FName SafeLevelName ( *MakeSafeLevelName( LevelName ) );
-				ULevelStreaming* StreamingLevel = nullptr;
-				TWeakObjectPtr<ULevelStreaming>* FoundStreamingLevel = NameToLevelMap.Find( SafeLevelName );
+				CurrentActiveSections.Add( LevelVisibilitySection );
+			}
+		}
+	}
 
-				if ( FoundStreamingLevel != nullptr )
-				{
-					StreamingLevel = FoundStreamingLevel->Get();
-				}
-				else if ( GIsEditor && !GWorld->HasBegunPlay() )
+	TArray<UMovieSceneLevelVisibilitySection*> NewActiveSections = CurrentActiveSections;
+	for (int i = ActiveSectionDatas.Num() - 1; i >= 0; i--)
+	{
+		FActiveSectionData& ActiveSectionData = ActiveSectionDatas[i];
+		if ( ActiveSectionData.Section.IsValid() )
+		{
+			if ( CurrentActiveSections.Contains( ActiveSectionData.Section ) )
+			{
+				NewActiveSections.Remove( ActiveSectionData.Section.Get() );
+			}
+			else
+			{
+				for ( auto LevelNamePreviousVisibilityPair : ActiveSectionData.PreviousLevelVisibility )
 				{
-					// If we're in the editor, it possible the level list has been updated since the last refresh, so try to find the level.
-					StreamingLevel = GetStreamingLevel( LevelName );
-					if ( StreamingLevel != nullptr )
+					FName SafeLevelName( *MakeSafeLevelName( LevelNamePreviousVisibilityPair.Key ) );
+					TWeakObjectPtr<ULevelStreaming>* FoundStreamingLevel = NameToLevelMap.Find( SafeLevelName );
+					if ( FoundStreamingLevel != nullptr && FoundStreamingLevel->IsValid() )
 					{
-						SaveLevelVisibilityState( StreamingLevel );
-						NameToLevelMap.Add( LevelName, TWeakObjectPtr<ULevelStreaming>( StreamingLevel ) );
+						if ( GetLevelVisibility( FoundStreamingLevel->Get() ) != LevelNamePreviousVisibilityPair.Value )
+						{
+							SetLevelVisibility( FoundStreamingLevel->Get(), LevelNamePreviousVisibilityPair.Value );
+							bVisibilityModified = true;
+							CurrentWorld = (*FoundStreamingLevel)->GetWorld();
+						}
 					}
 				}
+				ActiveSectionDatas.RemoveAt( i );
+			}
+		}
+		else
+		{
+			ActiveSectionDatas.RemoveAt( i );
+		}
+	}
+
+	for ( UMovieSceneLevelVisibilitySection* LevelVisibilitySection : NewActiveSections )
+	{
+		FActiveSectionData ActiveSectionData;
+		ActiveSectionData.Section = LevelVisibilitySection;
+		bool bShouldBeVisible = LevelVisibilitySection->GetVisibility() == ELevelVisibility::Visible ? true : false;
+		
+
+		for ( FName LevelName : *LevelVisibilitySection->GetLevelNames() )
+		{
+			FName SafeLevelName ( *MakeSafeLevelName( LevelName ) );
+			ULevelStreaming* StreamingLevel = nullptr;
+			TWeakObjectPtr<ULevelStreaming>* FoundStreamingLevel = NameToLevelMap.Find( SafeLevelName );
 
+			if ( FoundStreamingLevel != nullptr )
+			{
+				StreamingLevel = FoundStreamingLevel->Get();
+			}
+			else if ( GIsEditor )
+			{
+				// If we're in the editor, it possible the level list has been updated since the last refresh, so try to find the level.
+				StreamingLevel = GetStreamingLevel( LevelName );
 				if ( StreamingLevel != nullptr )
 				{
-					if ( GIsEditor && !GWorld->HasBegunPlay() )
-					{
-						SetEditorLevelVisibilityState( StreamingLevel, bShouldBeVisible );
-					}
-					else
-					{
-						if ( StreamingLevel->bShouldBeVisible != bShouldBeVisible )
-						{
-							StreamingLevel->bShouldBeVisible = bShouldBeVisible;
-							bVisibilityModified = true;
-						}
-					}
-					
+					SaveLevelVisibilityState( StreamingLevel );
+					NameToLevelMap.Add( LevelName, TWeakObjectPtr<ULevelStreaming>( StreamingLevel ) );
 				}
 			}
-			if ( bVisibilityModified )
+
+			if ( StreamingLevel != nullptr )
 			{
-				GWorld->FlushLevelStreaming( EFlushLevelStreamingType::Visibility );
+				ActiveSectionData.PreviousLevelVisibility.Add(LevelName, GetLevelVisibility( StreamingLevel ) );
+				if ( GetLevelVisibility( StreamingLevel ) != bShouldBeVisible )
+				{
+					SetLevelVisibility( StreamingLevel, bShouldBeVisible );
+					bVisibilityModified = true;
+					CurrentWorld = StreamingLevel->GetWorld();
+				}
 			}
 		}
+		ActiveSectionDatas.Emplace(ActiveSectionData);
+	}
+
+	if ( bVisibilityModified && CurrentWorld != nullptr )
+	{
+		CurrentWorld->FlushLevelStreaming( EFlushLevelStreamingType::Visibility );
 	}
 }
 
@@ -173,26 +250,5 @@ void FMovieSceneLevelVisibilityTrackInstance::RefreshInstance( const TArray<TWea
 
 void FMovieSceneLevelVisibilityTrackInstance::SaveLevelVisibilityState( ULevelStreaming* StreamingLevel )
 {
-	SavedLevelVisibility.Add( FObjectKey( StreamingLevel ), StreamingLevel->bShouldBeVisibleInEditor );
-}
-
-void FMovieSceneLevelVisibilityTrackInstance::SetEditorLevelVisibilityState( ULevelStreaming* StreamingLevel, bool bShouldBeVisible )
-{
-	if ( StreamingLevel->bShouldBeVisibleInEditor != bShouldBeVisible )
-	{
-		StreamingLevel->bShouldBeVisibleInEditor = bShouldBeVisible;
-		ULevel* LoadedLevel = StreamingLevel->GetLoadedLevel();
-		if ( LoadedLevel != nullptr )
-		{
-			for ( AActor* Actor : LoadedLevel->Actors )
-			{
-				if ( Actor != nullptr )
-				{
-#if WITH_EDITOR
-					Actor->SetIsTemporarilyHiddenInEditor( !bShouldBeVisible );
-#endif
-				}
-			}
-		}
-	}
+	SavedLevelVisibility.Add( FObjectKey( StreamingLevel ), GetLevelVisibility(StreamingLevel) );
 }
diff --git a/Engine/Source/Runtime/MovieSceneTracks/Private/TrackInstances/MovieSceneMaterialTrackInstance.cpp b/Engine/Source/Runtime/MovieSceneTracks/Private/TrackInstances/MovieSceneMaterialTrackInstance.cpp
--- a/Engine/Source/Runtime/MovieSceneTracks/Private/TrackInstances/MovieSceneMaterialTrackInstance.cpp
+++ b/Engine/Source/Runtime/MovieSceneTracks/Private/TrackInstances/MovieSceneMaterialTrackInstance.cpp
@@ -18,21 +18,28 @@ void FMovieSceneMaterialTrackInstance::SaveState(const TArray<TWeakObjectPtr<UOb
 		if (RuntimeObjectToOriginalMaterialMap.Find(FObjectKey(Object)) == nullptr)
 		{
 			UMaterialInterface* CurrentMaterial = GetMaterialForObject(Object);
-			if ( CurrentMaterial->IsA( UMaterialInstanceDynamic::StaticClass() ) )
+			if ( CurrentMaterial != nullptr )
 			{
-				TWeakObjectPtr<UMaterialInterface>* OriginalMaterial = DynamicMaterialToOriginalMaterialMap.Find( FObjectKey( CurrentMaterial ) );
-				if ( OriginalMaterial != nullptr )
+				UMaterialInterface* OriginalMaterial = nullptr;
+				if ( CurrentMaterial->IsA( UMaterialInstanceDynamic::StaticClass() ) )
 				{
-					CurrentMaterial = OriginalMaterial->Get();
+					// If the current material is a dynamic material see if we can find the original material.
+					TWeakObjectPtr<UMaterialInterface>* OriginalMaterialPtr = DynamicMaterialToOriginalMaterialMap.Find( FObjectKey( CurrentMaterial ) );
+					if ( OriginalMaterialPtr != nullptr )
+					{
+						OriginalMaterial = OriginalMaterialPtr->Get();
+					}
 				}
 				else
 				{
-					CurrentMaterial = nullptr;
+					// If the material isn't dynamic than the original material is the current material.
+					OriginalMaterial = CurrentMaterial;
+				}
+
+				if ( OriginalMaterial != nullptr )
+				{
+					RuntimeObjectToOriginalMaterialMap.Add( FObjectKey( Object ), OriginalMaterial );
 				}
-			}
-			if (CurrentMaterial != nullptr)
-			{
-				RuntimeObjectToOriginalMaterialMap.Add( FObjectKey(Object), CurrentMaterial );
 			}
 		}
 	}
@@ -93,16 +100,20 @@ void FMovieSceneMaterialTrackInstance::RefreshInstance( const TArray<TWeakObject
 	{
 		UObject* Object = ObjectPtr.Get();
 		UMaterialInterface* Material = GetMaterialForObject( Object );
-		UMaterialInstanceDynamic* DynamicMaterialInstance = Cast<UMaterialInstanceDynamic>(Material);
 
-		if (DynamicMaterialInstance == nullptr )
+		if ( Material != nullptr)
 		{
-			DynamicMaterialInstance = UMaterialInstanceDynamic::Create(Material, Object);
-			SetMaterialForObject( Object, DynamicMaterialInstance );
-			DynamicMaterialToOriginalMaterialMap.Add( FObjectKey( DynamicMaterialInstance ), Material );
-		}
+			UMaterialInstanceDynamic* DynamicMaterialInstance = Cast<UMaterialInstanceDynamic>( Material );
+
+			if ( DynamicMaterialInstance == nullptr )
+			{
+				DynamicMaterialInstance = UMaterialInstanceDynamic::Create( Material, Object );
+				SetMaterialForObject( Object, DynamicMaterialInstance );
+				DynamicMaterialToOriginalMaterialMap.Add( FObjectKey( DynamicMaterialInstance ), Material );
+			}
 
-		DynamicMaterialInstances.Add( DynamicMaterialInstance );
+			DynamicMaterialInstances.Add( DynamicMaterialInstance );
+		}
 	}
 }
 
diff --git a/Engine/Source/Runtime/MovieSceneTracks/Private/TrackInstances/MovieSceneSkeletalAnimationTrackInstance.cpp b/Engine/Source/Runtime/MovieSceneTracks/Private/TrackInstances/MovieSceneSkeletalAnimationTrackInstance.cpp
--- a/Engine/Source/Runtime/MovieSceneTracks/Private/TrackInstances/MovieSceneSkeletalAnimationTrackInstance.cpp
+++ b/Engine/Source/Runtime/MovieSceneTracks/Private/TrackInstances/MovieSceneSkeletalAnimationTrackInstance.cpp
@@ -62,7 +62,7 @@ void FMovieSceneSkeletalAnimationTrackInstance::SaveState(const TArray<TWeakObje
 
 			if (SkeletalMeshComponent)
 			{
-				if (GIsEditor && !RuntimeObject->GetWorld()->HasBegunPlay())
+				if (ShouldUsePreviewPlayback(Player, RuntimeObject))
 				{
 					PreviewBeginAnimControl(SkeletalMeshComponent);
 				}
@@ -99,7 +99,7 @@ void FMovieSceneSkeletalAnimationTrackInstance::RestoreState(const TArray<TWeakO
 
 			if (SkeletalMeshComponent)
 			{
-				if (GIsEditor && !RuntimeObject->GetWorld()->HasBegunPlay())
+				if (ShouldUsePreviewPlayback(Player, RuntimeObject))
 				{
 					PreviewFinishAnimControl(SkeletalMeshComponent);
 				}
@@ -162,9 +162,7 @@ void FMovieSceneSkeletalAnimationTrackInstance::Update(EMovieSceneUpdateData& Up
 
 						const bool bLooping = false;
 
-						// we also use PreviewSetAnimPosition in PIE when not playing, as we can preview in PIE
-						const bool bIsNotInPIEOrNotPlaying = !RuntimeObject->GetWorld()->HasBegunPlay() || Player.GetPlaybackStatus() != EMovieScenePlayerStatus::Playing;
-						if (GIsEditor && bIsNotInPIEOrNotPlaying)
+						if (ShouldUsePreviewPlayback(Player, RuntimeObject))
 						{
 							// If the playback status is jumping, ie. one such occurrence is setting the time for thumbnail generation, disable anim notifies updates because it could fire audio
 							const bool bFireNotifies = Player.GetPlaybackStatus() == EMovieScenePlayerStatus::Jumping || Player.GetPlaybackStatus() == EMovieScenePlayerStatus::Stopped ? false : true;
@@ -176,8 +174,13 @@ void FMovieSceneSkeletalAnimationTrackInstance::Update(EMovieSceneUpdateData& Up
 							{
 								DeltaTime = 0.f;
 							}
+
+							const bool bResetDynamics = Player.GetPlaybackStatus() == EMovieScenePlayerStatus::Stepping || 
+														Player.GetPlaybackStatus() == EMovieScenePlayerStatus::Jumping || 
+														Player.GetPlaybackStatus() == EMovieScenePlayerStatus::Scrubbing || 
+														(DeltaTime == 0.0f && Player.GetPlaybackStatus() != EMovieScenePlayerStatus::Stopped); 
 						
-							PreviewSetAnimPosition(SkeletalMeshComponent, AnimSection->GetSlotName(), ChannelIndex, AnimSequence, EvalTime, bLooping, bFireNotifies, DeltaTime, Player.GetPlaybackStatus() == EMovieScenePlayerStatus::Playing);
+							PreviewSetAnimPosition(SkeletalMeshComponent, AnimSection->GetSlotName(), ChannelIndex, AnimSequence, EvalTime, bLooping, bFireNotifies, DeltaTime, Player.GetPlaybackStatus() == EMovieScenePlayerStatus::Playing, bResetDynamics);
 						}
 						else
 						{
@@ -202,6 +205,8 @@ void FMovieSceneSkeletalAnimationTrackInstance::BeginAnimControl(USkeletalMeshCo
 			SkeletalMeshComponent->SetAnimationMode(EAnimationMode::Type::AnimationSingleNode);
 		}
 	}
+
+	CurrentlyPlayingMontage = nullptr;
 }
 
 bool FMovieSceneSkeletalAnimationTrackInstance::CanPlayAnimation(USkeletalMeshComponent* SkeletalMeshComponent, class UAnimSequenceBase* AnimAssetBase/*=nullptr*/) const
@@ -234,6 +239,8 @@ void FMovieSceneSkeletalAnimationTrackInstance::FinishAnimControl(USkeletalMeshC
 		SkeletalMeshComponent->RefreshSlaveComponents();
 		SkeletalMeshComponent->UpdateComponentToWorld();
 	}
+
+	CurrentlyPlayingMontage = nullptr;
 }
 
 
@@ -247,6 +254,8 @@ void FMovieSceneSkeletalAnimationTrackInstance::PreviewBeginAnimControl(USkeleta
 			SkeletalMeshComponent->SetAnimationMode(EAnimationMode::Type::AnimationSingleNode);
 		}
 	}
+
+	CurrentlyPlayingMontage = nullptr;
 }
 
 void FMovieSceneSkeletalAnimationTrackInstance::PreviewFinishAnimControl(USkeletalMeshComponent* SkeletalMeshComponent)
@@ -269,9 +278,11 @@ void FMovieSceneSkeletalAnimationTrackInstance::PreviewFinishAnimControl(USkelet
 		SkeletalMeshComponent->RefreshSlaveComponents();
 		SkeletalMeshComponent->UpdateComponentToWorld();
 	}
+
+	CurrentlyPlayingMontage = nullptr;
 }
 
-void FMovieSceneSkeletalAnimationTrackInstance::PreviewSetAnimPosition(USkeletalMeshComponent* SkeletalMeshComponent, FName SlotName, int32 ChannelIndex, UAnimSequence* InAnimSequence, float InPosition, bool bLooping, bool bFireNotifies, float DeltaTime, bool bPlaying)
+void FMovieSceneSkeletalAnimationTrackInstance::PreviewSetAnimPosition(USkeletalMeshComponent* SkeletalMeshComponent, FName SlotName, int32 ChannelIndex, UAnimSequence* InAnimSequence, float InPosition, bool bLooping, bool bFireNotifies, float DeltaTime, bool bPlaying, bool bResetDynamics)
 {
 	if(CanPlayAnimation(SkeletalMeshComponent, InAnimSequence))
 	{
@@ -297,6 +308,19 @@ void FMovieSceneSkeletalAnimationTrackInstance::PreviewSetAnimPosition(USkeletal
 					AnimInst->Montage_Pause(CurrentlyPlayingMontage.Get());
 				}
 			}
+
+			if(bResetDynamics)
+			{
+				// make sure we reset any simulations
+				AnimInst->ResetDynamics();	
+			}
 		}
 	}
-}
\ No newline at end of file
+}
+
+bool FMovieSceneSkeletalAnimationTrackInstance::ShouldUsePreviewPlayback(class IMovieScenePlayer& Player, UObject* RuntimeObject) const
+{
+	// we also use PreviewSetAnimPosition in PIE when not playing, as we can preview in PIE
+	bool bIsNotInPIEOrNotPlaying = (RuntimeObject && !RuntimeObject->GetWorld()->HasBegunPlay()) || Player.GetPlaybackStatus() != EMovieScenePlayerStatus::Playing;
+	return GIsEditor && bIsNotInPIEOrNotPlaying;
+}
diff --git a/Engine/Source/Runtime/MovieSceneTracks/Private/TrackInstances/MovieSceneSpawnTrackInstance.cpp b/Engine/Source/Runtime/MovieSceneTracks/Private/TrackInstances/MovieSceneSpawnTrackInstance.cpp
--- a/Engine/Source/Runtime/MovieSceneTracks/Private/TrackInstances/MovieSceneSpawnTrackInstance.cpp
+++ b/Engine/Source/Runtime/MovieSceneTracks/Private/TrackInstances/MovieSceneSpawnTrackInstance.cpp
@@ -29,7 +29,7 @@ void FMovieSceneSpawnTrackInstance::Update(EMovieSceneUpdateData& UpdateData, co
 #if WITH_EDITORONLY_DATA
 		bDestroy = !Spawnable->ShouldIgnoreOwnershipInEditor();
 		// Don't destroy cameras while previewing
-		if (bIsPreview && MovieSceneHelpers::CameraComponentFromActor(GetDefault<AActor>(Spawnable->GetClass())))
+		if (bIsPreview && MovieSceneHelpers::CameraComponentFromActor(Cast<AActor>(Spawnable->GetObjectTemplate())))
 		{
 			bDestroy = false;
 		}
diff --git a/Engine/Source/Runtime/MovieSceneTracks/Private/TrackInstances/MovieSceneSubTrackInstance.cpp b/Engine/Source/Runtime/MovieSceneTracks/Private/TrackInstances/MovieSceneSubTrackInstance.cpp
--- a/Engine/Source/Runtime/MovieSceneTracks/Private/TrackInstances/MovieSceneSubTrackInstance.cpp
+++ b/Engine/Source/Runtime/MovieSceneTracks/Private/TrackInstances/MovieSceneSubTrackInstance.cpp
@@ -188,49 +188,62 @@ void FMovieSceneSubTrackInstance::Update(EMovieSceneUpdateData& UpdateData, cons
 	float PreviousTime = CurrentTime;
 
 	TArray<UMovieSceneSection*> TraversedSections = GetTraversedSectionsWithPreroll(AllSections, CurrentTime, PreviousTime);
+	TArray<TWeakObjectPtr<UMovieSceneSubSection>>& LastTraversedSections = UpdatePassToLastTraversedSectionsMap.FindOrAdd(UpdateData.UpdatePass);
 
-	const float InitialUpdatePosition = UpdateData.Position;
+	for ( TWeakObjectPtr<UMovieSceneSubSection> LastSubSection : LastTraversedSections )
+	{
+		if ( TraversedSections.Contains( LastSubSection.Get() ) == false )
+		{
+			UpdateSection( UpdateData, Player, LastSubSection.Get(), true );
+		}
+	}
+	LastTraversedSections.Empty();
 
 	for (const auto Section : TraversedSections)
 	{
-		// skip sections with invalid time scale
-		UMovieSceneSubSection* SubSection = CastChecked<UMovieSceneSubSection>(Section);
+		UMovieSceneSubSection* SubSection = CastChecked<UMovieSceneSubSection>( Section );
+		UpdateSection(UpdateData, Player, SubSection, false);
+		LastTraversedSections.Add( TWeakObjectPtr<UMovieSceneSubSection>( SubSection ) );
+	}
+}
 
-		if (SubSection->TimeScale == 0.0f)
-		{
-			continue;
-		}
+void FMovieSceneSubTrackInstance::UpdateSection( EMovieSceneUpdateData& UpdateData, class IMovieScenePlayer& Player, UMovieSceneSubSection* SubSection, bool bSectionWasDeactivated )
+{
+	if (SubSection->TimeScale == 0.0f)
+	{
+		return;
+	}
 
-		// skip sections without valid instances
-		TSharedPtr<FMovieSceneSequenceInstance> Instance = SequenceInstancesBySection.FindRef(SubSection);
+	// skip sections without valid instances
+	TSharedPtr<FMovieSceneSequenceInstance> Instance = SequenceInstancesBySection.FindRef(SubSection);
 
-		if (!Instance.IsValid())
-		{
-			continue;
-		}
+	if (!Instance.IsValid())
+	{
+		return;
+	}
 
-		// calculate section's local time
-		const float InstanceOffset = SubSection->StartOffset + Instance->GetTimeRange().GetLowerBoundValue() - SubSection->PrerollTime;
-		const float InstanceLastPosition = InstanceOffset + (UpdateData.LastPosition - (SubSection->GetStartTime()- SubSection->PrerollTime)) / SubSection->TimeScale;
-		const float InstancePosition = InstanceOffset + (UpdateData.Position - (SubSection->GetStartTime()- SubSection->PrerollTime)) / SubSection->TimeScale;
+	// calculate section's local time
+	const float InstanceOffset = SubSection->StartOffset + Instance->GetTimeRange().GetLowerBoundValue() - SubSection->PrerollTime;
+	const float InstanceLastPosition = InstanceOffset + (UpdateData.LastPosition - (SubSection->GetStartTime()- SubSection->PrerollTime)) / SubSection->TimeScale;
+	const float InstancePosition = InstanceOffset + (UpdateData.Position - (SubSection->GetStartTime()- SubSection->PrerollTime)) / SubSection->TimeScale;
 
-		EMovieSceneUpdateData SubUpdateData(InstancePosition, InstanceLastPosition);
-		SubUpdateData.bJumpCut = UpdateData.LastPosition < SubSection->GetStartTime() || UpdateData.LastPosition > SubSection->GetEndTime();
-		SubUpdateData.UpdatePass = UpdateData.UpdatePass;
-		SubUpdateData.bPreroll = InitialUpdatePosition < SubSection->GetStartTime();
+	EMovieSceneUpdateData SubUpdateData(InstancePosition, InstanceLastPosition);
+	SubUpdateData.bJumpCut = UpdateData.LastPosition < SubSection->GetStartTime() || UpdateData.LastPosition > SubSection->GetEndTime();
+	SubUpdateData.UpdatePass = UpdateData.UpdatePass;
+	SubUpdateData.bPreroll = UpdateData.Position < SubSection->GetStartTime();
+	SubUpdateData.bSubSceneDeactivate = bSectionWasDeactivated;
 
-		// update sub sections
+	// update sub sections
 
-		if (SubUpdateData.UpdatePass == MSUP_PreUpdate)
-		{
-			Instance->PreUpdate(Player);
-		}
+	if (SubUpdateData.UpdatePass == MSUP_PreUpdate)
+	{
+		Instance->PreUpdate(Player);
+	}
 		
-		Instance->UpdatePassSingle(SubUpdateData, Player);
+	Instance->UpdatePassSingle(SubUpdateData, Player);
 
-		if (SubUpdateData.UpdatePass == MSUP_PostUpdate)
-		{
-			Instance->PostUpdate(Player);
-		}
+	if (SubUpdateData.UpdatePass == MSUP_PostUpdate)
+	{
+		Instance->PostUpdate(Player);
 	}
 }
diff --git a/Engine/Source/Runtime/MovieSceneTracks/Private/Tracks/MovieSceneCameraCutTrack.cpp b/Engine/Source/Runtime/MovieSceneTracks/Private/Tracks/MovieSceneCameraCutTrack.cpp
--- a/Engine/Source/Runtime/MovieSceneTracks/Private/Tracks/MovieSceneCameraCutTrack.cpp
+++ b/Engine/Source/Runtime/MovieSceneTracks/Private/Tracks/MovieSceneCameraCutTrack.cpp
@@ -26,14 +26,36 @@ void UMovieSceneCameraCutTrack::AddNewCameraCut(FGuid CameraHandle, float StartT
 {
 	Modify();
 
-	UMovieSceneCameraCutSection* NewSection = NewObject<UMovieSceneCameraCutSection>(this, NAME_None, RF_Transactional);
+
+	float NewSectionEndTime = FindEndTimeForCameraCut(StartTime);
+
+	// If there's an existing section, just swap the camera guid
+	UMovieSceneCameraCutSection* ExistingSection = nullptr;
+	for (auto Section : Sections)
+	{
+		if (Section->GetStartTime() == StartTime && Section->GetEndTime() == NewSectionEndTime)
+		{
+			ExistingSection = Cast<UMovieSceneCameraCutSection>(Section);
+			break;
+		}
+	}
+
+	UMovieSceneCameraCutSection* NewSection = ExistingSection;
+	if (ExistingSection != nullptr)
 	{
-		NewSection->SetStartTime(StartTime);
-		NewSection->SetEndTime(FindEndTimeForCameraCut(StartTime));
-		NewSection->SetCameraGuid(CameraHandle);
+		ExistingSection->SetCameraGuid(CameraHandle);
 	}
+	else
+	{
+		NewSection = NewObject<UMovieSceneCameraCutSection>(this, NAME_None, RF_Transactional);
+		{
+			NewSection->SetStartTime(StartTime);
+			NewSection->SetEndTime(FindEndTimeForCameraCut(StartTime));
+			NewSection->SetCameraGuid(CameraHandle);
+		}
 
-	Sections.Add(NewSection);
+		Sections.Add(NewSection);
+	}
 
 	// When a new CameraCut is added, sort all CameraCuts to ensure they are in the correct order
 	MovieSceneHelpers::SortConsecutiveSections(Sections);
diff --git a/Engine/Source/Runtime/MovieSceneTracks/Private/Tracks/MovieSceneEventTrack.cpp b/Engine/Source/Runtime/MovieSceneTracks/Private/Tracks/MovieSceneEventTrack.cpp
--- a/Engine/Source/Runtime/MovieSceneTracks/Private/Tracks/MovieSceneEventTrack.cpp
+++ b/Engine/Source/Runtime/MovieSceneTracks/Private/Tracks/MovieSceneEventTrack.cpp
@@ -37,7 +37,7 @@ bool UMovieSceneEventTrack::AddKeyToSection(float Time, FName EventName, FKeyPar
 }
 
 
-void UMovieSceneEventTrack::TriggerEvents(float Position, float LastPosition, UObject* EventContextObject)
+void UMovieSceneEventTrack::TriggerEvents(float Position, float LastPosition, TArray<UObject*> EventContextObjects)
 {
 	if ((Sections.Num() == 0) || (Position == LastPosition))
 	{
@@ -52,7 +52,7 @@ void UMovieSceneEventTrack::TriggerEvents(float Position, float LastPosition, UO
 		return;
 	}
 
-	if (EventContextObject == nullptr)
+	if (EventContextObjects.Num() == 0)
 	{
 		return;
 	}
@@ -60,7 +60,7 @@ void UMovieSceneEventTrack::TriggerEvents(float Position, float LastPosition, UO
 	TArray<UMovieSceneSection*> TraversedSections = MovieSceneHelpers::GetTraversedSections(Sections, Position, LastPosition);
 	for (auto EventSection : TraversedSections)
 	{
-		CastChecked<UMovieSceneEventSection>(EventSection)->TriggerEvents(EventContextObject, Position, LastPosition);
+		CastChecked<UMovieSceneEventSection>(EventSection)->TriggerEvents(EventContextObjects, Position, LastPosition);
 	}
 }
 
diff --git a/Engine/Source/Runtime/MovieSceneTracks/Public/Sections/MovieScene3DPathSection.h b/Engine/Source/Runtime/MovieSceneTracks/Public/Sections/MovieScene3DPathSection.h
--- a/Engine/Source/Runtime/MovieSceneTracks/Public/Sections/MovieScene3DPathSection.h
+++ b/Engine/Source/Runtime/MovieSceneTracks/Public/Sections/MovieScene3DPathSection.h
@@ -32,7 +32,7 @@ public:
 	/** MovieSceneSection interface */
 	virtual void MoveSection( float DeltaPosition, TSet<FKeyHandle>& KeyHandles ) override;
 	virtual void DilateSection( float DilationFactor, float Origin, TSet<FKeyHandle>& KeyHandles ) override;
-	virtual void GetKeyHandles(TSet<FKeyHandle>& KeyHandles) const override;
+	virtual void GetKeyHandles(TSet<FKeyHandle>& OutKeyHandles, TRange<float> TimeRange) const override;
 
 	/**
 	 * Evaluates the path track
diff --git a/Engine/Source/Runtime/MovieSceneTracks/Public/Sections/MovieScene3DTransformSection.h b/Engine/Source/Runtime/MovieSceneTracks/Public/Sections/MovieScene3DTransformSection.h
--- a/Engine/Source/Runtime/MovieSceneTracks/Public/Sections/MovieScene3DTransformSection.h
+++ b/Engine/Source/Runtime/MovieSceneTracks/Public/Sections/MovieScene3DTransformSection.h
@@ -222,7 +222,7 @@ public:
 
 	virtual void MoveSection( float DeltaPosition, TSet<FKeyHandle>& KeyHandles ) override;
 	virtual void DilateSection( float DilationFactor, float Origin, TSet<FKeyHandle>& KeyHandles ) override;
-	virtual void GetKeyHandles(TSet<FKeyHandle>& KeyHandles) const override;
+	virtual void GetKeyHandles(TSet<FKeyHandle>& OutKeyHandles, TRange<float> TimeRange) const override;
 	virtual TSharedPtr<FStructOnScope> GetKeyStruct(const TArray<FKeyHandle>& KeyHandles) override;
 
 public:
diff --git a/Engine/Source/Runtime/MovieSceneTracks/Public/Sections/MovieSceneAudioSection.h b/Engine/Source/Runtime/MovieSceneTracks/Public/Sections/MovieSceneAudioSection.h
--- a/Engine/Source/Runtime/MovieSceneTracks/Public/Sections/MovieSceneAudioSection.h
+++ b/Engine/Source/Runtime/MovieSceneTracks/Public/Sections/MovieSceneAudioSection.h
@@ -44,6 +44,21 @@ public:
 	 */
 	float GetAudioDilationFactor() const {return AudioDilationFactor;}
 
+	/**
+	* Sets the dilation multiplier of the sound
+	*/
+	void SetAudioDilationFactor( float InAudioDilationFactor ) { AudioDilationFactor = InAudioDilationFactor; }
+
+	/**
+	* @return The volume the sound will be played with.
+	*/
+	float GetAudioVolume() const { return AudioVolume; }
+
+	/**
+	* Sets the volume the sound will be played with.
+	*/
+	void SetAudioVolume( float InAudioVolume ) { AudioVolume = InAudioVolume; }
+
 	/**
 	 * Returns whether or not a provided position in time is within the timespan of the audio range
 	 *
@@ -72,7 +87,7 @@ public:
 
 	virtual void MoveSection( float DeltaPosition, TSet<FKeyHandle>& KeyHandles ) override;
 	virtual void DilateSection( float DilationFactor, float Origin, TSet<FKeyHandle>& KeyHandles ) override;
-	virtual void GetKeyHandles(TSet<FKeyHandle>& KeyHandles) const override;
+	virtual void GetKeyHandles(TSet<FKeyHandle>& OutKeyHandles, TRange<float> TimeRange) const override;
 	virtual void GetSnapTimes(TArray<float>& OutSnapTimes, bool bGetSectionBorders) const override;
 
 private:
@@ -89,6 +104,10 @@ private:
 	UPROPERTY(EditAnywhere, Category="Audio")
 	float AudioDilationFactor;
 
+	/** The volume the sound will be played with. */
+	UPROPERTY( EditAnywhere, Category = "Audio" )
+	float AudioVolume;
+
 #if WITH_EDITORONLY_DATA
 	/** Whether to show the actual intensity of the wave on the thumbnail, as well as the smoothed RMS */
 	UPROPERTY(EditAnywhere, Category="Audio")
diff --git a/Engine/Source/Runtime/MovieSceneTracks/Public/Sections/MovieSceneBoolSection.h b/Engine/Source/Runtime/MovieSceneTracks/Public/Sections/MovieSceneBoolSection.h
--- a/Engine/Source/Runtime/MovieSceneTracks/Public/Sections/MovieSceneBoolSection.h
+++ b/Engine/Source/Runtime/MovieSceneTracks/Public/Sections/MovieSceneBoolSection.h
@@ -2,6 +2,7 @@
 
 #pragma once
 
+#include "MovieSceneSection.h"
 #include "IKeyframeSection.h"
 #include "MovieSceneBoolSection.generated.h"
 
@@ -41,7 +42,7 @@ public:
 	 */
 	virtual void MoveSection(float DeltaPosition, TSet<FKeyHandle>& KeyHandles) override;
 	virtual void DilateSection(float DilationFactor, float Origin, TSet<FKeyHandle>& KeyHandles) override;
-	virtual void GetKeyHandles(TSet<FKeyHandle>& KeyHandles) const override;
+	virtual void GetKeyHandles(TSet<FKeyHandle>& OutKeyHandles, TRange<float> TimeRange) const override;
 
 	/** Gets all the keys of this boolean section */
 	FIntegralCurve& GetCurve() { return BoolCurve; }
diff --git a/Engine/Source/Runtime/MovieSceneTracks/Public/Sections/MovieSceneByteSection.h b/Engine/Source/Runtime/MovieSceneTracks/Public/Sections/MovieSceneByteSection.h
--- a/Engine/Source/Runtime/MovieSceneTracks/Public/Sections/MovieSceneByteSection.h
+++ b/Engine/Source/Runtime/MovieSceneTracks/Public/Sections/MovieSceneByteSection.h
@@ -2,6 +2,7 @@
 
 #pragma once
 
+#include "MovieSceneSection.h"
 #include "IKeyframeSection.h"
 #include "MovieSceneByteSection.generated.h"
 
@@ -33,7 +34,7 @@ public:
 	 */
 	virtual void MoveSection(float DeltaPosition, TSet<FKeyHandle>& KeyHandles) override;
 	virtual void DilateSection(float DilationFactor, float Origin, TSet<FKeyHandle>& KeyHandles) override;
-	virtual void GetKeyHandles(TSet<FKeyHandle>& KeyHandles) const override;
+	virtual void GetKeyHandles(TSet<FKeyHandle>& OutKeyHandles, TRange<float> TimeRange) const override;
 
 	/** Gets all the keys of this byte section */
 	FIntegralCurve& GetCurve() { return ByteCurve; }
diff --git a/Engine/Source/Runtime/MovieSceneTracks/Public/Sections/MovieSceneEventSection.h b/Engine/Source/Runtime/MovieSceneTracks/Public/Sections/MovieSceneEventSection.h
--- a/Engine/Source/Runtime/MovieSceneTracks/Public/Sections/MovieSceneEventSection.h
+++ b/Engine/Source/Runtime/MovieSceneTracks/Public/Sections/MovieSceneEventSection.h
@@ -42,18 +42,18 @@ public:
 	/**
 	 * Trigger the events that fall into the given time range.
 	 *
-	 * @param EventContextObject The object to trigger the events on.
+	 * @param EventContextObjects The objects to trigger the events on.
 	 * @param Position The current position in time.
 	 * @param LastPosition The time at the last update.
 	 */
-	void TriggerEvents(UObject* EventContextObject, float Position, float LastPosition);
+	void TriggerEvents(TArray<UObject*> EventContextObjects, float Position, float LastPosition);
 
 public:
 
 	// UMovieSceneSection interface
 
 	virtual void DilateSection(float DilationFactor, float Origin, TSet<FKeyHandle>& KeyHandles) override;
-	virtual void GetKeyHandles(TSet<FKeyHandle>& KeyHandles) const override;
+	virtual void GetKeyHandles(TSet<FKeyHandle>& KeyHandles, TRange<float> TimeRange) const override;
 	virtual void MoveSection(float DeltaPosition, TSet<FKeyHandle>& KeyHandles) override;
 
 protected:
@@ -62,9 +62,9 @@ protected:
 	 * Trigger event for the specified name.
 	 *
 	 * @param Event The Name to trigger.
-	 * @param Object The object to trigger the events on.
+	 * @param EventContextObjects The objects to trigger the events on.
 	 */
-	void TriggerEvent(const FName& Event, UObject* EventContextObject);
+	void TriggerEvent(const FName& Event, TArray<UObject*> EventContextObjects);
 
 private:
 
diff --git a/Engine/Source/Runtime/MovieSceneTracks/Public/Sections/MovieSceneParticleSection.h b/Engine/Source/Runtime/MovieSceneTracks/Public/Sections/MovieSceneParticleSection.h
--- a/Engine/Source/Runtime/MovieSceneTracks/Public/Sections/MovieSceneParticleSection.h
+++ b/Engine/Source/Runtime/MovieSceneTracks/Public/Sections/MovieSceneParticleSection.h
@@ -39,7 +39,7 @@ class UMovieSceneParticleSection
 	*/
 	virtual void MoveSection( float DeltaPosition, TSet<FKeyHandle>& KeyHandles ) override;
 	virtual void DilateSection( float DilationFactor, float Origin, TSet<FKeyHandle>& KeyHandles ) override;
-	virtual void GetKeyHandles( TSet<FKeyHandle>& KeyHandles ) const override;
+	virtual void GetKeyHandles( TSet<FKeyHandle>& OutKeyHandles, TRange<float> TimeRange ) const override;
 
 private:
 	/** Curve containing the particle keys. */
diff --git a/Engine/Source/Runtime/MovieSceneTracks/Public/Sections/MovieSceneSubSection.h b/Engine/Source/Runtime/MovieSceneTracks/Public/Sections/MovieSceneSubSection.h
--- a/Engine/Source/Runtime/MovieSceneTracks/Public/Sections/MovieSceneSubSection.h
+++ b/Engine/Source/Runtime/MovieSceneTracks/Public/Sections/MovieSceneSubSection.h
@@ -82,6 +82,15 @@ public:
 		TargetPathToRecordTo.Path = Path;
 	}
 
+	/** Set the target actor name to record */
+	void SetNameOfActorToRecord(const FString& ActorName)
+	{
+		NameOfActorToRecord = ActorName;
+	}
+
+	virtual UMovieSceneSection* SplitSection( float SplitTime ) override;
+	virtual void TrimSection( float TrimTime, bool bTrimLeft ) override;
+
 public:
 
 	/** Number of seconds to skip at the beginning of the sub-sequence. */
diff --git a/Engine/Source/Runtime/MovieSceneTracks/Public/TrackInstances/MovieSceneCameraAnimTrackInstance.h b/Engine/Source/Runtime/MovieSceneTracks/Public/TrackInstances/MovieSceneCameraAnimTrackInstance.h
--- a/Engine/Source/Runtime/MovieSceneTracks/Public/TrackInstances/MovieSceneCameraAnimTrackInstance.h
+++ b/Engine/Source/Runtime/MovieSceneTracks/Public/TrackInstances/MovieSceneCameraAnimTrackInstance.h
@@ -12,6 +12,8 @@ struct FMovieSceneCameraAnimSectionInstanceData
 	TWeakObjectPtr<UCameraAnimInst> CameraAnimInst;
 	FTransform AdditiveOffset;		// AdditiveCamToBaseCam
 	float AdditiveFOVOffset;
+	float PostProcessingBlendWeight;
+	FPostProcessSettings PostProcessingSettings;
 };
 
 /**
diff --git a/Engine/Source/Runtime/MovieSceneTracks/Public/TrackInstances/MovieSceneCameraShakeTrackInstance.h b/Engine/Source/Runtime/MovieSceneTracks/Public/TrackInstances/MovieSceneCameraShakeTrackInstance.h
--- a/Engine/Source/Runtime/MovieSceneTracks/Public/TrackInstances/MovieSceneCameraShakeTrackInstance.h
+++ b/Engine/Source/Runtime/MovieSceneTracks/Public/TrackInstances/MovieSceneCameraShakeTrackInstance.h
@@ -7,12 +7,13 @@
 class UMovieSceneCameraShakeTrack;
 class UMovieSceneCameraShakeSection;
 
-
 struct FMovieSceneCameraShakeSectionInstanceData
 {
 	TWeakObjectPtr<UCameraShake> CameraShakeInst;
 	FTransform AdditiveOffset;		// AdditiveCamToBaseCam
 	float AdditiveFOVOffset;
+	float PostProcessingBlendWeight;
+	FPostProcessSettings PostProcessingSettings;
 };
 
 /**
diff --git a/Engine/Source/Runtime/MovieSceneTracks/Public/TrackInstances/MovieSceneColorTrackInstance.h b/Engine/Source/Runtime/MovieSceneTracks/Public/TrackInstances/MovieSceneColorTrackInstance.h
--- a/Engine/Source/Runtime/MovieSceneTracks/Public/TrackInstances/MovieSceneColorTrackInstance.h
+++ b/Engine/Source/Runtime/MovieSceneTracks/Public/TrackInstances/MovieSceneColorTrackInstance.h
@@ -33,7 +33,7 @@ private:
 		/** FLinearColor */
 		Linear,
 		/** FColor */
-		RegularColor,
+		Color,
 	};
 
 	/** The track being instanced */
diff --git a/Engine/Source/Runtime/MovieSceneTracks/Public/TrackInstances/MovieSceneLevelVisibilityTrackInstance.h b/Engine/Source/Runtime/MovieSceneTracks/Public/TrackInstances/MovieSceneLevelVisibilityTrackInstance.h
--- a/Engine/Source/Runtime/MovieSceneTracks/Public/TrackInstances/MovieSceneLevelVisibilityTrackInstance.h
+++ b/Engine/Source/Runtime/MovieSceneTracks/Public/TrackInstances/MovieSceneLevelVisibilityTrackInstance.h
@@ -14,6 +14,13 @@ class UMovieSceneLevelVisibilityTrack;
 class FMovieSceneLevelVisibilityTrackInstance
 	: public IMovieSceneTrackInstance
 {
+private:
+	struct FActiveSectionData
+	{
+		TWeakObjectPtr<UMovieSceneLevelVisibilitySection> Section;
+		TMap<FName, bool> PreviousLevelVisibility;
+	};
+
 public:
 
 	FMovieSceneLevelVisibilityTrackInstance( UMovieSceneLevelVisibilityTrack& InLevelVisibilityTrack );
@@ -22,13 +29,13 @@ public:
 	virtual void SaveState (const TArray<TWeakObjectPtr<UObject>>& RuntimeObjects, IMovieScenePlayer& Player, FMovieSceneSequenceInstance& SequenceInstance) override;
 	virtual void RestoreState (const TArray<TWeakObjectPtr<UObject>>& RuntimeObjects, IMovieScenePlayer& Player, FMovieSceneSequenceInstance& SequenceInstance) override;
 	virtual void Update(EMovieSceneUpdateData& UpdateData, const TArray<TWeakObjectPtr<UObject>>& RuntimeObjects, IMovieScenePlayer& Player, FMovieSceneSequenceInstance& SequenceInstance) override;
+	virtual bool RequiresUpdateForSubSceneDeactivate() { return true; }
 	virtual void RefreshInstance( const TArray<TWeakObjectPtr<UObject>>& RuntimeObjects, IMovieScenePlayer& Player, FMovieSceneSequenceInstance& SequenceInstance ) override;
 	virtual void ClearInstance( IMovieScenePlayer& Player, FMovieSceneSequenceInstance& SequenceInstance ) override {}
 
 private:
 
 	void SaveLevelVisibilityState( ULevelStreaming* StreamingLevel );
-	void SetEditorLevelVisibilityState( ULevelStreaming* StreamingLevel, bool bShouldBeVisible );
 
 private:
 
@@ -37,6 +44,8 @@ private:
 
 	TMap<FObjectKey, bool> SavedLevelVisibility;
 
+	TArray<FActiveSectionData> ActiveSectionDatas;
+
 	TMap<FName, TWeakObjectPtr<ULevelStreaming>> NameToLevelMap;
 
 };
diff --git a/Engine/Source/Runtime/MovieSceneTracks/Public/TrackInstances/MovieSceneMaterialTrackInstance.h b/Engine/Source/Runtime/MovieSceneTracks/Public/TrackInstances/MovieSceneMaterialTrackInstance.h
--- a/Engine/Source/Runtime/MovieSceneTracks/Public/TrackInstances/MovieSceneMaterialTrackInstance.h
+++ b/Engine/Source/Runtime/MovieSceneTracks/Public/TrackInstances/MovieSceneMaterialTrackInstance.h
@@ -12,7 +12,7 @@ class UMovieSceneComponentMaterialTrack;
 /**
  * A movie scene track instance for material tracks.
  */
-class FMovieSceneMaterialTrackInstance
+class MOVIESCENETRACKS_API FMovieSceneMaterialTrackInstance
 	: public IMovieSceneTrackInstance
 {
 public:
diff --git a/Engine/Source/Runtime/MovieSceneTracks/Public/TrackInstances/MovieSceneSkeletalAnimationTrackInstance.h b/Engine/Source/Runtime/MovieSceneTracks/Public/TrackInstances/MovieSceneSkeletalAnimationTrackInstance.h
--- a/Engine/Source/Runtime/MovieSceneTracks/Public/TrackInstances/MovieSceneSkeletalAnimationTrackInstance.h
+++ b/Engine/Source/Runtime/MovieSceneTracks/Public/TrackInstances/MovieSceneSkeletalAnimationTrackInstance.h
@@ -40,7 +40,9 @@ private:
 
 	void PreviewFinishAnimControl(USkeletalMeshComponent* SkeletalMeshComponent);
 
-	void PreviewSetAnimPosition(USkeletalMeshComponent* SkeletalMeshComponent, FName SlotName, int32 ChannelIndex, UAnimSequence* InAnimSequence, float InPosition, bool bLooping, bool bFireNotifies, float DeltaTime, bool bPlaying);
+	void PreviewSetAnimPosition(USkeletalMeshComponent* SkeletalMeshComponent, FName SlotName, int32 ChannelIndex, UAnimSequence* InAnimSequence, float InPosition, bool bLooping, bool bFireNotifies, float DeltaTime, bool bPlaying, bool bResetDynamics);
+
+	bool ShouldUsePreviewPlayback(class IMovieScenePlayer& Player, UObject* RuntimeObject) const;
 
 private:
 
diff --git a/Engine/Source/Runtime/MovieSceneTracks/Public/TrackInstances/MovieSceneSubTrackInstance.h b/Engine/Source/Runtime/MovieSceneTracks/Public/TrackInstances/MovieSceneSubTrackInstance.h
--- a/Engine/Source/Runtime/MovieSceneTracks/Public/TrackInstances/MovieSceneSubTrackInstance.h
+++ b/Engine/Source/Runtime/MovieSceneTracks/Public/TrackInstances/MovieSceneSubTrackInstance.h
@@ -37,6 +37,7 @@ protected:
 	TArray<UMovieSceneSection*> GetAllTraversedSectionsWithPreroll( const TArray<UMovieSceneSection*>& Sections, float CurrentTime, float PreviousTime );
 	TArray<UMovieSceneSection*> GetTraversedSectionsWithPreroll( const TArray<UMovieSceneSection*>& Sections, float CurrentTime, float PreviousTime );
 	virtual bool ShouldEvaluateIfOverlapping(const TArray<UMovieSceneSection*>& Sections, UMovieSceneSection* Section) const;
+	void UpdateSection( EMovieSceneUpdateData& UpdateData, class IMovieScenePlayer& Player, UMovieSceneSubSection* SubSection, bool bSectionWasDeactivated );
 	
 protected:
 
@@ -45,4 +46,6 @@ protected:
 
 	/** Mapping of section lookups to instances.  Each section has a movie scene which must be instanced */
 	TMap<TWeakObjectPtr<UMovieSceneSubSection>, TSharedPtr<FMovieSceneSequenceInstance>> SequenceInstancesBySection;
+
+	TMap<EMovieSceneUpdatePass, TArray<TWeakObjectPtr<UMovieSceneSubSection>>> UpdatePassToLastTraversedSectionsMap;
 };
diff --git a/Engine/Source/Runtime/MovieSceneTracks/Public/Tracks/MovieSceneEventTrack.h b/Engine/Source/Runtime/MovieSceneTracks/Public/Tracks/MovieSceneEventTrack.h
--- a/Engine/Source/Runtime/MovieSceneTracks/Public/Tracks/MovieSceneEventTrack.h
+++ b/Engine/Source/Runtime/MovieSceneTracks/Public/Tracks/MovieSceneEventTrack.h
@@ -46,7 +46,7 @@ public:
 	 * @param Position The current position in time.
 	 * @param LastPosition The time at the last update.
 	 */
-	void TriggerEvents(float Position, float LastPosition, UObject* EventContextObject);
+	void TriggerEvents(float Position, float LastPosition, TArray<UObject*> EventContextObjects);
 
 public:
 
diff --git a/Engine/Source/Runtime/MovieSceneTracks/Public/Tracks/MovieSceneMaterialTrack.h b/Engine/Source/Runtime/MovieSceneTracks/Public/Tracks/MovieSceneMaterialTrack.h
--- a/Engine/Source/Runtime/MovieSceneTracks/Public/Tracks/MovieSceneMaterialTrack.h
+++ b/Engine/Source/Runtime/MovieSceneTracks/Public/Tracks/MovieSceneMaterialTrack.h
@@ -14,8 +14,8 @@ class UMovieSceneSection;
 /**
  * Handles manipulation of material parameters in a movie scene.
  */
-UCLASS(MinimalAPI)
-class UMovieSceneMaterialTrack
+UCLASS()
+class MOVIESCENETRACKS_API UMovieSceneMaterialTrack
 	: public UMovieSceneNameableTrack
 {
 	GENERATED_UCLASS_BODY()
@@ -41,7 +41,7 @@ public:
 	 * @param Time The time to add the new key.
 	 * @param The value for the new key.
 	 */
-	void MOVIESCENETRACKS_API AddScalarParameterKey(FName ParameterName, float Position, float Value);
+	void AddScalarParameterKey(FName ParameterName, float Position, float Value);
 
 	/**
 	* Adds a color parameter key to the track.
@@ -49,7 +49,7 @@ public:
 	* @param Time The time to add the new key.
 	* @param The value for the new key.
 	*/
-	void MOVIESCENETRACKS_API AddColorParameterKey(FName ParameterName, float Position, FLinearColor Value);
+	void AddColorParameterKey(FName ParameterName, float Position, FLinearColor Value);
 
 	/**
 	 * Gets the animated values for this track.
@@ -81,6 +81,7 @@ public:
 	// UMovieSceneTrack interface
 
 	virtual TSharedPtr<IMovieSceneTrackInstance> CreateInstance() override;
+	virtual FName GetTrackName() const { return FName( *FString::FromInt(MaterialIndex) ); }
 
 #if WITH_EDITORONLY_DATA
 	virtual FText GetDefaultDisplayName() const override;
diff --git a/Engine/Source/Runtime/Sockets/Private/BSDIPv6Sockets/SocketsBSDIPv6.cpp b/Engine/Source/Runtime/Sockets/Private/BSDIPv6Sockets/SocketsBSDIPv6.cpp
--- a/Engine/Source/Runtime/Sockets/Private/BSDIPv6Sockets/SocketsBSDIPv6.cpp
+++ b/Engine/Source/Runtime/Sockets/Private/BSDIPv6Sockets/SocketsBSDIPv6.cpp
@@ -49,7 +49,7 @@ EIPv6SocketInternalState::Return FSocketBSDIPv6::HasState(EIPv6SocketInternalSta
 	// convert WaitTime to a timeval
 	timeval Time;
 	Time.tv_sec = (int32)WaitTime.GetTotalSeconds();
-	Time.tv_usec = WaitTime.GetMilliseconds() * 1000;
+	Time.tv_usec = WaitTime.GetMicroseconds();
 
 	fd_set SocketSet;
 
diff --git a/Engine/Source/Runtime/Sockets/Private/BSDSockets/SocketsBSD.cpp b/Engine/Source/Runtime/Sockets/Private/BSDSockets/SocketsBSD.cpp
--- a/Engine/Source/Runtime/Sockets/Private/BSDSockets/SocketsBSD.cpp
+++ b/Engine/Source/Runtime/Sockets/Private/BSDSockets/SocketsBSD.cpp
@@ -426,7 +426,7 @@ ESocketBSDReturn FSocketBSD::HasState(ESocketBSDParam State, FTimespan WaitTime)
 	// convert WaitTime to a timeval
 	timeval Time;
 	Time.tv_sec = (int32)WaitTime.GetTotalSeconds();
-	Time.tv_usec = WaitTime.GetMilliseconds() * 1000;
+	Time.tv_usec = WaitTime.GetMicroseconds();
 
 	fd_set SocketSet;
 
diff --git a/Engine/Source/Runtime/UMG/Private/Animation/MovieScene2DTransformSection.cpp b/Engine/Source/Runtime/UMG/Private/Animation/MovieScene2DTransformSection.cpp
--- a/Engine/Source/Runtime/UMG/Private/Animation/MovieScene2DTransformSection.cpp
+++ b/Engine/Source/Runtime/UMG/Private/Animation/MovieScene2DTransformSection.cpp
@@ -36,14 +36,19 @@ void UMovieScene2DTransformSection::DilateSection( float DilationFactor, float O
 }
 
 
-void UMovieScene2DTransformSection::GetKeyHandles(TSet<FKeyHandle>& KeyHandles) const
+void UMovieScene2DTransformSection::GetKeyHandles(TSet<FKeyHandle>& OutKeyHandles, TRange<float> TimeRange) const
 {
+	if (!TimeRange.Overlaps(GetRange()))
+	{
+		return;
+	}
+
 	for (auto It(Rotation.GetKeyHandleIterator()); It; ++It)
 	{
 		float Time = Rotation.GetKeyTime(It.Key());
-		if (IsTimeWithinSection(Time))
+		if (TimeRange.Contains(Time))
 		{
-			KeyHandles.Add(It.Key());
+			OutKeyHandles.Add(It.Key());
 		}
 	}
 
@@ -52,25 +57,25 @@ void UMovieScene2DTransformSection::GetKeyHandles(TSet<FKeyHandle>& KeyHandles)
 		for (auto It(Translation[Axis].GetKeyHandleIterator()); It; ++It)
 		{
 			float Time = Translation[Axis].GetKeyTime(It.Key());
-			if (IsTimeWithinSection(Time))
+			if (TimeRange.Contains(Time))
 			{
-				KeyHandles.Add(It.Key());
+				OutKeyHandles.Add(It.Key());
 			}
 		}
 		for (auto It(Scale[Axis].GetKeyHandleIterator()); It; ++It)
 		{
 			float Time = Scale[Axis].GetKeyTime(It.Key());
-			if (IsTimeWithinSection(Time))
+			if (TimeRange.Contains(Time))
 			{
-				KeyHandles.Add(It.Key());
+				OutKeyHandles.Add(It.Key());
 			}
 		}
 		for (auto It(Shear[Axis].GetKeyHandleIterator()); It; ++It)
 		{
 			float Time = Shear[Axis].GetKeyTime(It.Key());
-			if (IsTimeWithinSection(Time))
+			if (TimeRange.Contains(Time))
 			{
-				KeyHandles.Add(It.Key());
+				OutKeyHandles.Add(It.Key());
 			}
 		}
 	}
diff --git a/Engine/Source/Runtime/UMG/Private/Animation/MovieSceneMarginSection.cpp b/Engine/Source/Runtime/UMG/Private/Animation/MovieSceneMarginSection.cpp
--- a/Engine/Source/Runtime/UMG/Private/Animation/MovieSceneMarginSection.cpp
+++ b/Engine/Source/Runtime/UMG/Private/Animation/MovieSceneMarginSection.cpp
@@ -30,38 +30,43 @@ void UMovieSceneMarginSection::DilateSection( float DilationFactor, float Origin
 	BottomCurve.ScaleCurve(Origin, DilationFactor, KeyHandles);
 }
 
-void UMovieSceneMarginSection::GetKeyHandles(TSet<FKeyHandle>& KeyHandles) const
+void UMovieSceneMarginSection::GetKeyHandles(TSet<FKeyHandle>& OutKeyHandles, TRange<float> TimeRange) const
 {
+	if (!TimeRange.Overlaps(GetRange()))
+	{
+		return;
+	}
+
 	for (auto It(LeftCurve.GetKeyHandleIterator()); It; ++It)
 	{
 		float Time = LeftCurve.GetKeyTime(It.Key());
-		if (IsTimeWithinSection(Time))
+		if (TimeRange.Contains(Time))
 		{
-			KeyHandles.Add(It.Key());
+			OutKeyHandles.Add(It.Key());
 		}
 	}
 	for (auto It(TopCurve.GetKeyHandleIterator()); It; ++It)
 	{
 		float Time = TopCurve.GetKeyTime(It.Key());
-		if (IsTimeWithinSection(Time))
+		if (TimeRange.Contains(Time))
 		{
-			KeyHandles.Add(It.Key());
+			OutKeyHandles.Add(It.Key());
 		}
 	}
 	for (auto It(RightCurve.GetKeyHandleIterator()); It; ++It)
 	{
 		float Time = RightCurve.GetKeyTime(It.Key());
-		if (IsTimeWithinSection(Time))
+		if (TimeRange.Contains(Time))
 		{
-			KeyHandles.Add(It.Key());
+			OutKeyHandles.Add(It.Key());
 		}
 	}
 	for (auto It(BottomCurve.GetKeyHandleIterator()); It; ++It)
 	{
 		float Time = BottomCurve.GetKeyTime(It.Key());
-		if (IsTimeWithinSection(Time))
+		if (TimeRange.Contains(Time))
 		{
-			KeyHandles.Add(It.Key());
+			OutKeyHandles.Add(It.Key());
 		}
 	}
 }
diff --git a/Engine/Source/Runtime/UMG/Private/Animation/MovieSceneWidgetMaterialTrack.cpp b/Engine/Source/Runtime/UMG/Private/Animation/MovieSceneWidgetMaterialTrack.cpp
new file mode 100644
--- /dev/null
+++ b/Engine/Source/Runtime/UMG/Private/Animation/MovieSceneWidgetMaterialTrack.cpp
@@ -0,0 +1,25 @@
+// Copyright 1998-2016 Epic Games, Inc. All Rights Reserved.
+
+#include "UMGPrivatePCH.h"
+#include "MovieSceneWidgetMaterialTrack.h"
+#include "MovieSceneWidgetMaterialTrackInstance.h"
+
+
+UMovieSceneWidgetMaterialTrack::UMovieSceneWidgetMaterialTrack( const FObjectInitializer& ObjectInitializer )
+	: Super(ObjectInitializer)
+{
+}
+
+
+TSharedPtr<IMovieSceneTrackInstance> UMovieSceneWidgetMaterialTrack::CreateInstance()
+{
+	return MakeShareable(new FMovieSceneWidgetMaterialTrackInstance(*this));
+}
+
+
+#if WITH_EDITORONLY_DATA
+FText UMovieSceneWidgetMaterialTrack::GetDefaultDisplayName() const
+{
+	return FText::Format(NSLOCTEXT("UMGAnimation", "MaterialTrackFormat", "{0} Material"), FText::FromName(BrushPropertyName));
+}
+#endif
diff --git a/Engine/Source/Runtime/UMG/Private/Animation/MovieSceneWidgetMaterialTrackInstance.cpp b/Engine/Source/Runtime/UMG/Private/Animation/MovieSceneWidgetMaterialTrackInstance.cpp
new file mode 100644
--- /dev/null
+++ b/Engine/Source/Runtime/UMG/Private/Animation/MovieSceneWidgetMaterialTrackInstance.cpp
@@ -0,0 +1,36 @@
+// Copyright 1998-2016 Epic Games, Inc. All Rights Reserved.
+
+#include "UMGPrivatePCH.h"
+#include "MovieSceneWidgetMaterialTrack.h"
+#include "MovieSceneWidgetMaterialTrackInstance.h"
+#include "WidgetMaterialTrackUtilities.h"
+
+
+FMovieSceneWidgetMaterialTrackInstance::FMovieSceneWidgetMaterialTrackInstance( UMovieSceneWidgetMaterialTrack& InMaterialTrack )
+	: FMovieSceneMaterialTrackInstance( InMaterialTrack )
+	, BrushPropertyName( InMaterialTrack.GetBrushPropertyName() )
+{
+}
+
+
+UMaterialInterface* FMovieSceneWidgetMaterialTrackInstance::GetMaterialForObject( UObject* Object ) const 
+{
+	UWidget* Widget = Cast<UWidget>(Object);
+	if ( Widget != nullptr )
+	{
+		FSlateBrush* Brush = WidgetMaterialTrackUtilities::GetBrush( Widget, BrushPropertyName );
+		return Cast<UMaterialInterface>( Brush->GetResourceObject() );
+	}
+	return nullptr;
+}
+
+
+void FMovieSceneWidgetMaterialTrackInstance::SetMaterialForObject( UObject* Object, UMaterialInterface* Material )
+{
+	UWidget* Widget = Cast<UWidget>( Object );
+	if ( Widget != nullptr )
+	{
+		FSlateBrush* Brush = WidgetMaterialTrackUtilities::GetBrush( Widget, BrushPropertyName );
+		Brush->SetResourceObject( Material );
+	}
+}
diff --git a/Engine/Source/Runtime/UMG/Private/Animation/MovieSceneWidgetMaterialTrackInstance.h b/Engine/Source/Runtime/UMG/Private/Animation/MovieSceneWidgetMaterialTrackInstance.h
new file mode 100644
--- /dev/null
+++ b/Engine/Source/Runtime/UMG/Private/Animation/MovieSceneWidgetMaterialTrackInstance.h
@@ -0,0 +1,23 @@
+// Copyright 1998-2016 Epic Games, Inc. All Rights Reserved.
+
+#pragma once
+
+#include "MovieSceneMaterialTrackInstance.h"
+
+class UMovieSceneWidgetMaterialTrack;
+
+/** A material track instance specialized for widget materials. */
+class FMovieSceneWidgetMaterialTrackInstance : public FMovieSceneMaterialTrackInstance
+{
+public:
+	/** Creates a new track instance with a specific widget material track. */
+	FMovieSceneWidgetMaterialTrackInstance( UMovieSceneWidgetMaterialTrack& InMaterialTrack );
+
+protected:
+	/** FMovieSceneMaterialTrackInstance interface. */
+	virtual UMaterialInterface* GetMaterialForObject( UObject* Object ) const override;
+	virtual void SetMaterialForObject( UObject* Object, UMaterialInterface* Material ) override;
+
+private:
+	FName BrushPropertyName;
+};
diff --git a/Engine/Source/Runtime/UMG/Private/Animation/UMGSequencePlayer.cpp b/Engine/Source/Runtime/UMG/Private/Animation/UMGSequencePlayer.cpp
--- a/Engine/Source/Runtime/UMG/Private/Animation/UMGSequencePlayer.cpp
+++ b/Engine/Source/Runtime/UMG/Private/Animation/UMGSequencePlayer.cpp
@@ -242,13 +242,14 @@ EMovieScenePlayerStatus::Type UUMGSequencePlayer::GetPlaybackStatus() const
 	return PlayerStatus;
 }
 
-UObject* UUMGSequencePlayer::GetEventContext() const
+TArray<UObject*> UUMGSequencePlayer::GetEventContexts() const
 {
+	TArray<UObject*> EventContexts;
 	if (UserWidget.IsValid())
 	{
-		return UserWidget.Get();
+		EventContexts.Add(UserWidget.Get());
 	}
-	return nullptr;
+	return EventContexts;
 }
 
 void UUMGSequencePlayer::SetPlaybackStatus(EMovieScenePlayerStatus::Type InPlaybackStatus)
diff --git a/Engine/Source/Runtime/UMG/Private/Animation/WidgetMaterialTrackUtilities.cpp b/Engine/Source/Runtime/UMG/Private/Animation/WidgetMaterialTrackUtilities.cpp
new file mode 100644
--- /dev/null
+++ b/Engine/Source/Runtime/UMG/Private/Animation/WidgetMaterialTrackUtilities.cpp
@@ -0,0 +1,41 @@
+// Copyright 1998-2016 Epic Games, Inc. All Rights Reserved.
+
+
+#include "UMGPrivatePCH.h"
+#include "WidgetMaterialTrackUtilities.h"
+
+
+FSlateBrush* WidgetMaterialTrackUtilities::GetBrush( UWidget* Widget, FName BrushPropertyName )
+{
+	for ( TFieldIterator<UProperty> PropertyIterator( Widget->GetClass() ); PropertyIterator; ++PropertyIterator )
+	{
+		UProperty* Property = *PropertyIterator;
+		if ( Property->GetFName() == BrushPropertyName )
+		{
+			return Property->ContainerPtrToValuePtr<FSlateBrush>( Widget );
+		}
+	}
+	return nullptr;
+}
+
+
+void WidgetMaterialTrackUtilities::GetMaterialBrushProperties( UWidget* Widget, TArray<UProperty*>& MaterialBrushProperties )
+{
+	for ( TFieldIterator<UProperty> PropertyIterator( Widget->GetClass() ); PropertyIterator; ++PropertyIterator )
+	{
+		UProperty* Property = *PropertyIterator;
+		if ( Property && !Property->HasAnyPropertyFlags( CPF_Deprecated ) )
+		{
+			UStructProperty* StructProperty = Cast<UStructProperty>( Property );
+			if ( StructProperty != nullptr && StructProperty->Struct->GetName() == "SlateBrush" )
+			{
+				FSlateBrush* SlateBrush = Property->ContainerPtrToValuePtr<FSlateBrush>( Widget );
+				UMaterialInterface* MaterialInterface = Cast<UMaterialInterface>( SlateBrush->GetResourceObject() );
+				if ( MaterialInterface != nullptr )
+				{
+					MaterialBrushProperties.Add( Property );
+				}
+			}
+		}
+	}
+}
\ No newline at end of file
diff --git a/Engine/Source/Runtime/UMG/Public/Animation/MovieSceneWidgetMaterialTrack.h b/Engine/Source/Runtime/UMG/Public/Animation/MovieSceneWidgetMaterialTrack.h
new file mode 100644
--- /dev/null
+++ b/Engine/Source/Runtime/UMG/Public/Animation/MovieSceneWidgetMaterialTrack.h
@@ -0,0 +1,45 @@
+// Copyright 1998-2016 Epic Games, Inc. All Rights Reserved.
+
+#pragma once
+
+#include "MovieSceneMaterialTrack.h"
+#include "MovieSceneWidgetMaterialTrack.generated.h"
+
+
+/**
+ * A material track which is specialized for materials which are owned by widget brushes.
+ */
+UCLASS(MinimalAPI)
+class UMovieSceneWidgetMaterialTrack
+	: public UMovieSceneMaterialTrack
+{
+	GENERATED_UCLASS_BODY()
+
+public:
+
+	// UMovieSceneTrack interface
+
+	virtual TSharedPtr<IMovieSceneTrackInstance> CreateInstance() override;
+	virtual FName GetTrackName() const override { return BrushPropertyName; }
+
+#if WITH_EDITORONLY_DATA
+	virtual FText GetDefaultDisplayName() const override;
+#endif
+
+public:
+
+	/** Gets name of the brush property which has the material to animate. */
+	FName GetBrushPropertyName() const { return BrushPropertyName; }
+
+	/** Sets the name of the brush property which has the material to animate. */
+	void SetBrushPropertyName(FName InBrushPropertyName) 
+	{
+		BrushPropertyName = InBrushPropertyName;
+	}
+
+private:
+
+	/** The name of the brush property which will be animated by this track. */
+	UPROPERTY()
+	FName BrushPropertyName;
+};
diff --git a/Engine/Source/Runtime/UMG/Public/Animation/UMGSequencePlayer.h b/Engine/Source/Runtime/UMG/Public/Animation/UMGSequencePlayer.h
--- a/Engine/Source/Runtime/UMG/Public/Animation/UMGSequencePlayer.h
+++ b/Engine/Source/Runtime/UMG/Public/Animation/UMGSequencePlayer.h
@@ -52,7 +52,7 @@ public:
 	virtual void RemoveMovieSceneInstance( class UMovieSceneSection& MovieSceneSection, TSharedRef<FMovieSceneSequenceInstance> InstanceToRemove ) override {}
 	virtual TSharedRef<FMovieSceneSequenceInstance> GetRootMovieSceneSequenceInstance() const override { return RootMovieSceneInstance.ToSharedRef(); }
 	virtual EMovieScenePlayerStatus::Type GetPlaybackStatus() const override;
-	virtual UObject* GetEventContext() const override;
+	virtual TArray<UObject*> GetEventContexts() const override;
 	virtual void SetPlaybackStatus(EMovieScenePlayerStatus::Type InPlaybackStatus) override;
 
 	DECLARE_EVENT_OneParam(UUMGSequencePlayer, FOnSequenceFinishedPlaying, UUMGSequencePlayer&);
diff --git a/Engine/Source/Runtime/UMG/Public/Animation/WidgetMaterialTrackUtilities.h b/Engine/Source/Runtime/UMG/Public/Animation/WidgetMaterialTrackUtilities.h
new file mode 100644
--- /dev/null
+++ b/Engine/Source/Runtime/UMG/Public/Animation/WidgetMaterialTrackUtilities.h
@@ -0,0 +1,12 @@
+// Copyright 1998-2016 Epic Games, Inc. All Rights Reserved.
+
+#pragma once
+
+namespace WidgetMaterialTrackUtilities
+{
+	/** Gets a slate brush from a property on a widget by the properties FName. */
+	UMG_API FSlateBrush* GetBrush( UWidget* Widget, FName BrushPropertyName );
+
+	/** Gets the properties on a widget which are slate brush properties, and who's slate brush has a valid material. */
+	UMG_API void GetMaterialBrushProperties( UWidget* Widget, TArray<UProperty*>& MaterialBrushProperties );
+}
\ No newline at end of file
