Sequencer: Revert 3057233 because it breaks sequence recording.

Copy from Dev-Sequencer

diff --git a/Engine/Source/Runtime/MovieSceneTracks/Private/TrackInstances/MovieScene3DTransformTrackInstance.cpp b/Engine/Source/Runtime/MovieSceneTracks/Private/TrackInstances/MovieScene3DTransformTrackInstance.cpp
--- a/Engine/Source/Runtime/MovieSceneTracks/Private/TrackInstances/MovieScene3DTransformTrackInstance.cpp
+++ b/Engine/Source/Runtime/MovieSceneTracks/Private/TrackInstances/MovieScene3DTransformTrackInstance.cpp
@@ -105,10 +105,8 @@ void FMovieScene3DTransformTrackInstance::Update(EMovieSceneUpdateData& UpdateDa
 			if (SceneComponent != nullptr)
 			{
 				if (UpdateData.UpdatePass == MSUP_PreUpdate)
-				{
-					SceneComponent->RelativeLocation = FVector(0, 0, 0);
-					SceneComponent->RelativeRotation = FRotator(0, 0, 0);
-					SceneComponent->RelativeScale3D = FVector(1.f, 1.f, 1.f);
+				{		
+					SceneComponent->ResetRelativeTransform();
 				}
 				else if (UpdateData.UpdatePass == MSUP_Update)
 				{
