Copying //UE4/Dev-Sequencer to //UE4/Dev-Main (Source: //UE4/Dev-Sequencer @ 3116759)

diff --git a/Engine/Source/Runtime/MovieSceneTracks/Private/TrackInstances/MovieSceneByteTrackInstance.cpp b/Engine/Source/Runtime/MovieSceneTracks/Private/TrackInstances/MovieSceneByteTrackInstance.cpp
--- a/Engine/Source/Runtime/MovieSceneTracks/Private/TrackInstances/MovieSceneByteTrackInstance.cpp
+++ b/Engine/Source/Runtime/MovieSceneTracks/Private/TrackInstances/MovieSceneByteTrackInstance.cpp
@@ -52,13 +52,16 @@ void FMovieSceneByteTrackInstance::RestoreState(const TArray<TWeakObjectPtr<UObj
 
 void FMovieSceneByteTrackInstance::Update(EMovieSceneUpdateData& UpdateData, const TArray<TWeakObjectPtr<UObject>>& RuntimeObjects, class IMovieScenePlayer& Player, FMovieSceneSequenceInstance& SequenceInstance) 
 {
-	uint8 ByteValue = 0;
-	if( ByteTrack->Eval( UpdateData.Position, UpdateData.LastPosition, ByteValue ) )
+	for (auto ObjectPtr : RuntimeObjects )
 	{
-		for (auto ObjectPtr : RuntimeObjects )
+		UObject* Object = ObjectPtr.Get();
+		if (Object != nullptr)
 		{
-			UObject* Object = ObjectPtr.Get();
-			PropertyBindings->CallFunction<uint8>(Object, &ByteValue);
+			uint8 ByteValue = PropertyBindings->GetCurrentValue<uint8>(Object);
+			if (ByteTrack->Eval(UpdateData.Position, UpdateData.LastPosition, ByteValue))
+			{
+				PropertyBindings->CallFunction<uint8>(Object, &ByteValue);
+			}
 		}
 	}
 }
diff --git a/Engine/Source/Runtime/MovieSceneTracks/Private/TrackInstances/MovieSceneFloatTrackInstance.cpp b/Engine/Source/Runtime/MovieSceneTracks/Private/TrackInstances/MovieSceneFloatTrackInstance.cpp
--- a/Engine/Source/Runtime/MovieSceneTracks/Private/TrackInstances/MovieSceneFloatTrackInstance.cpp
+++ b/Engine/Source/Runtime/MovieSceneTracks/Private/TrackInstances/MovieSceneFloatTrackInstance.cpp
@@ -52,13 +52,13 @@ void FMovieSceneFloatTrackInstance::RestoreState(const TArray<TWeakObjectPtr<UOb
 
 void FMovieSceneFloatTrackInstance::Update(EMovieSceneUpdateData& UpdateData, const TArray<TWeakObjectPtr<UObject>>& RuntimeObjects, class IMovieScenePlayer& Player, FMovieSceneSequenceInstance& SequenceInstance ) 
 {
-	float FloatValue = 0.0f;
-	if( FloatTrack->Eval( UpdateData.Position, UpdateData.LastPosition, FloatValue ) )
+	for (auto ObjectPtr : RuntimeObjects)
 	{
-		for(auto ObjectPtr : RuntimeObjects)
+		UObject* Object = ObjectPtr.Get();
+		float FloatValue = PropertyBindings->GetCurrentValue<float>(Object);
+		if (FloatTrack->Eval(UpdateData.Position, UpdateData.LastPosition, FloatValue))
 		{
-			UObject* Object = ObjectPtr.Get();
-			PropertyBindings->CallFunction<float>( Object, &FloatValue );
+			PropertyBindings->CallFunction<float>(Object, &FloatValue);
 		}
 	}
 }
diff --git a/Engine/Source/Runtime/MovieSceneTracks/Private/TrackInstances/MovieSceneSkeletalAnimationTrackInstance.cpp b/Engine/Source/Runtime/MovieSceneTracks/Private/TrackInstances/MovieSceneSkeletalAnimationTrackInstance.cpp
--- a/Engine/Source/Runtime/MovieSceneTracks/Private/TrackInstances/MovieSceneSkeletalAnimationTrackInstance.cpp
+++ b/Engine/Source/Runtime/MovieSceneTracks/Private/TrackInstances/MovieSceneSkeletalAnimationTrackInstance.cpp
@@ -135,6 +135,12 @@ USkeletalMeshComponent* GetSkeletalMeshComponentFromRuntimeObjectPtr( TWeakObjec
 
 void FMovieSceneSkeletalAnimationTrackInstance::Update( EMovieSceneUpdateData& UpdateData, const TArray<TWeakObjectPtr<UObject>>& RuntimeObjects, class IMovieScenePlayer& Player, FMovieSceneSequenceInstance& SequenceInstance )
 {
+	if (UpdateData.UpdatePass == MSUP_PreUpdate)
+	{
+		UpdateRefreshBones(RuntimeObjects);
+		return;
+	}
+
 	// @todo Sequencer gameplay update has a different code path than editor update for animation
 	for ( TWeakObjectPtr<UObject> RuntimeObjectPtr : RuntimeObjects )
 	{
@@ -204,6 +210,8 @@ void FMovieSceneSkeletalAnimationTrackInstance::Update( EMovieSceneUpdateData& U
 
 void FMovieSceneSkeletalAnimationTrackInstance::RefreshInstance( const TArray<TWeakObjectPtr<UObject>>& RuntimeObjects, IMovieScenePlayer& Player, FMovieSceneSequenceInstance& SequenceInstance )
 {
+	UpdateRefreshBones(RuntimeObjects);
+
 	// When not in preview playback we need to stop any running animations to prevent the animations from being advanced a frame when
 	// they are ticked by the engine.
 	for ( TWeakObjectPtr<UObject> RuntimeObjectPtr : RuntimeObjects )
@@ -369,3 +377,16 @@ bool FMovieSceneSkeletalAnimationTrackInstance::ShouldUsePreviewPlayback(class I
 	bool bIsNotInPIEOrNotPlaying = (RuntimeObject && RuntimeObject->GetWorld() && !RuntimeObject->GetWorld()->HasBegunPlay()) || Player.GetPlaybackStatus() != EMovieScenePlayerStatus::Playing;
 	return GIsEditor && bIsNotInPIEOrNotPlaying;
 }
+
+void FMovieSceneSkeletalAnimationTrackInstance::UpdateRefreshBones(const TArray<TWeakObjectPtr<UObject>>& RuntimeObjects)
+{
+	for ( TWeakObjectPtr<UObject> RuntimeObjectPtr : RuntimeObjects )
+	{
+		USkeletalMeshComponent* SkeletalMeshComponent = GetSkeletalMeshComponentFromRuntimeObjectPtr( RuntimeObjectPtr );
+
+		if (SkeletalMeshComponent && SkeletalMeshComponent->MeshComponentUpdateFlag != EMeshComponentUpdateFlag::AlwaysTickPoseAndRefreshBones)
+		{
+			SkeletalMeshComponent->MeshComponentUpdateFlag = EMeshComponentUpdateFlag::AlwaysTickPoseAndRefreshBones;
+		}
+	}
+}
\ No newline at end of file
diff --git a/Engine/Source/Runtime/MovieSceneTracks/Private/TrackInstances/MovieSceneStringTrackInstance.cpp b/Engine/Source/Runtime/MovieSceneTracks/Private/TrackInstances/MovieSceneStringTrackInstance.cpp
--- a/Engine/Source/Runtime/MovieSceneTracks/Private/TrackInstances/MovieSceneStringTrackInstance.cpp
+++ b/Engine/Source/Runtime/MovieSceneTracks/Private/TrackInstances/MovieSceneStringTrackInstance.cpp
@@ -65,14 +65,16 @@ void FMovieSceneStringTrackInstance::SaveState(const TArray<TWeakObjectPtr<UObje
 
 void FMovieSceneStringTrackInstance::Update(EMovieSceneUpdateData& UpdateData, const TArray<TWeakObjectPtr<UObject>>& RuntimeObjects, IMovieScenePlayer& Player, FMovieSceneSequenceInstance& SequenceInstance)
 {
-	FString StringValue;
-
-	if (StringTrack->Eval(UpdateData.Position, UpdateData.LastPosition, StringValue))
+	for (auto ObjectPtr : RuntimeObjects)
 	{
-		for (auto ObjectPtr : RuntimeObjects)
+		UObject* Object = ObjectPtr.Get();
+		if (Object != nullptr)
 		{
-			UObject* Object = ObjectPtr.Get();
-			PropertyBindings->CallFunction<FString>(Object, &StringValue);
+			FString StringValue = PropertyBindings->GetCurrentValue<FString>(Object);
+			if (StringTrack->Eval(UpdateData.Position, UpdateData.LastPosition, StringValue))
+			{
+				PropertyBindings->CallFunction<FString>(Object, &StringValue);
+			}
 		}
 	}
 }
diff --git a/Engine/Source/Runtime/MovieSceneTracks/Private/TrackInstances/MovieSceneVectorTrackInstance.cpp b/Engine/Source/Runtime/MovieSceneTracks/Private/TrackInstances/MovieSceneVectorTrackInstance.cpp
--- a/Engine/Source/Runtime/MovieSceneTracks/Private/TrackInstances/MovieSceneVectorTrackInstance.cpp
+++ b/Engine/Source/Runtime/MovieSceneTracks/Private/TrackInstances/MovieSceneVectorTrackInstance.cpp
@@ -103,45 +103,69 @@ void FMovieSceneVectorTrackInstance::RestoreState(const TArray<TWeakObjectPtr<UO
 }
 
 
-void FMovieSceneVectorTrackInstance::Update( EMovieSceneUpdateData& UpdateData, const TArray<TWeakObjectPtr<UObject>>& RuntimeObjects, class IMovieScenePlayer& Player, FMovieSceneSequenceInstance& SequenceInstance ) 
+void FMovieSceneVectorTrackInstance::Update(EMovieSceneUpdateData& UpdateData, const TArray<TWeakObjectPtr<UObject>>& RuntimeObjects, class IMovieScenePlayer& Player, FMovieSceneSequenceInstance& SequenceInstance)
 {
-	FVector4 Vector;
-	if( VectorTrack->Eval( UpdateData.Position, UpdateData.LastPosition, Vector ) )
+	int32 NumChannelsUsed = VectorTrack->GetNumChannelsUsed();
+	switch (NumChannelsUsed)
 	{
-		int32 NumChannelsUsed = VectorTrack->GetNumChannelsUsed();
-		switch( NumChannelsUsed )
+	case 2:
+	{
+		for (auto ObjectPtr : RuntimeObjects)
 		{
-			case 2:
+			UObject* Object = ObjectPtr.Get();
+			if(Object != nullptr)
 			{
-				FVector2D Value(Vector.X, Vector.Y);
-				for(auto Object : RuntimeObjects)
+				FVector2D Vector2D = PropertyBindings->GetCurrentValue<FVector2D>(Object);
+				FVector4 Vector4(Vector2D.X, Vector2D.Y);
+				if (VectorTrack->Eval(UpdateData.Position, UpdateData.LastPosition, Vector4))
 				{
-					PropertyBindings->CallFunction<FVector2D>(Object.Get(), &Value);
+					Vector2D.X = Vector4.X;
+					Vector2D.Y = Vector4.Y;
+					PropertyBindings->CallFunction<FVector2D>(Object, &Vector2D);
 				}
-				break;
 			}
-			case 3:
+		}
+		break;
+	}
+	case 3:
+	{
+		for (auto ObjectPtr : RuntimeObjects)
+		{
+			UObject* Object = ObjectPtr.Get();
+			if (Object != nullptr)
 			{
-				FVector Value(Vector.X, Vector.Y, Vector.Z);
-				for(auto Object : RuntimeObjects)
+				FVector Vector = PropertyBindings->GetCurrentValue<FVector>(Object);
+				FVector4 Vector4(Vector.X, Vector.Y, Vector.Z);
+				if (VectorTrack->Eval(UpdateData.Position, UpdateData.LastPosition, Vector4))
 				{
-					PropertyBindings->CallFunction<FVector>(Object.Get(), &Value);
+					Vector.X = Vector4.X;
+					Vector.Y = Vector4.Y;
+					Vector.Z = Vector4.Z;
+					PropertyBindings->CallFunction<FVector>(Object, &Vector);
 				}
-				break;
 			}
-			case 4:
+		}
+		break;
+	}
+	case 4:
+	{
+		for (auto ObjectPtr : RuntimeObjects)
+		{
+			UObject* Object = ObjectPtr.Get();
+			if (Object != nullptr)
 			{
-				for(auto Object : RuntimeObjects)
+				FVector4 Vector4 = PropertyBindings->GetCurrentValue<FVector4>(Object);
+				if (VectorTrack->Eval(UpdateData.Position, UpdateData.LastPosition, Vector4))
 				{
-					PropertyBindings->CallFunction<FVector4>(Object.Get(), &Vector);
+					PropertyBindings->CallFunction<FVector4>(Object, &Vector4);
 				}
-				break;
 			}
-			default:
-				UE_LOG(LogMovieScene, Warning, TEXT("Invalid number of channels(%d) for vector track"), NumChannelsUsed );
-				break;
 		}
-		
+		break;
+	}
+	default:
+		UE_LOG(LogMovieScene, Warning, TEXT("Invalid number of channels(%d) for vector track"), NumChannelsUsed);
+		break;
 	}
 }
 
diff --git a/Engine/Source/Runtime/MovieSceneTracks/Private/Tracks/MovieScene3DTransformTrack.cpp b/Engine/Source/Runtime/MovieSceneTracks/Private/Tracks/MovieScene3DTransformTrack.cpp
--- a/Engine/Source/Runtime/MovieSceneTracks/Private/Tracks/MovieScene3DTransformTrack.cpp
+++ b/Engine/Source/Runtime/MovieSceneTracks/Private/Tracks/MovieScene3DTransformTrack.cpp
@@ -31,7 +31,7 @@ TSharedPtr<IMovieSceneTrackInstance> UMovieScene3DTransformTrack::CreateInstance
 }
 
 
-bool UMovieScene3DTransformTrack::Eval( float Position, float LastPosition, FVector& OutTranslation, FRotator& OutRotation, FVector& OutScale ) const
+bool UMovieScene3DTransformTrack::Eval( float Position, float LastPosition, FVector& InOutTranslation, FRotator& InOutRotation, FVector& InOutScale ) const
 {
 	const UMovieSceneSection* Section = MovieSceneHelpers::FindNearestSectionAtTime( Sections, Position );
 
@@ -45,9 +45,9 @@ bool UMovieScene3DTransformTrack::Eval( float Position, float LastPosition, FVec
 		}
 
 		// Evaluate translation,rotation, and scale curves.  If no keys were found on one of these, that component of the transform will remain unchained
-		TransformSection->EvalTranslation( Position, OutTranslation );
-		TransformSection->EvalRotation( Position, OutRotation );
-		TransformSection->EvalScale( Position, OutScale );
+		TransformSection->EvalTranslation( Position, InOutTranslation );
+		TransformSection->EvalRotation( Position, InOutRotation );
+		TransformSection->EvalScale( Position, InOutScale );
 	}
 
 	return Section != nullptr;
diff --git a/Engine/Source/Runtime/MovieSceneTracks/Private/Tracks/MovieSceneBoolTrack.cpp b/Engine/Source/Runtime/MovieSceneTracks/Private/Tracks/MovieSceneBoolTrack.cpp
--- a/Engine/Source/Runtime/MovieSceneTracks/Private/Tracks/MovieSceneBoolTrack.cpp
+++ b/Engine/Source/Runtime/MovieSceneTracks/Private/Tracks/MovieSceneBoolTrack.cpp
@@ -19,7 +19,7 @@ TSharedPtr<IMovieSceneTrackInstance> UMovieSceneBoolTrack::CreateInstance()
 }
 
 
-bool UMovieSceneBoolTrack::Eval( float Position, float LastPostion, bool& OutBool ) const
+bool UMovieSceneBoolTrack::Eval( float Position, float LastPostion, bool& InOutBool ) const
 {	
 	const UMovieSceneSection* Section = MovieSceneHelpers::FindNearestSectionAtTime( Sections, Position );
 
@@ -30,7 +30,7 @@ bool UMovieSceneBoolTrack::Eval( float Position, float LastPostion, bool& OutBoo
 			Position = FMath::Clamp(Position, Section->GetStartTime(), Section->GetEndTime());
 		}
 
-		OutBool = CastChecked<UMovieSceneBoolSection>( Section )->Eval( Position );
+		InOutBool = CastChecked<UMovieSceneBoolSection>( Section )->Eval( Position, InOutBool );
 	}
 
 	return (Section != nullptr);
diff --git a/Engine/Source/Runtime/MovieSceneTracks/Private/Tracks/MovieSceneByteTrack.cpp b/Engine/Source/Runtime/MovieSceneTracks/Private/Tracks/MovieSceneByteTrack.cpp
--- a/Engine/Source/Runtime/MovieSceneTracks/Private/Tracks/MovieSceneByteTrack.cpp
+++ b/Engine/Source/Runtime/MovieSceneTracks/Private/Tracks/MovieSceneByteTrack.cpp
@@ -24,7 +24,7 @@ TSharedPtr<IMovieSceneTrackInstance> UMovieSceneByteTrack::CreateInstance()
 }
 
 
-bool UMovieSceneByteTrack::Eval( float Position, float LastPostion, uint8& OutByte ) const
+bool UMovieSceneByteTrack::Eval( float Position, float LastPostion, uint8& InOutByte ) const
 {	
 	const UMovieSceneSection* Section = MovieSceneHelpers::FindNearestSectionAtTime( Sections, Position );
 
@@ -35,7 +35,7 @@ bool UMovieSceneByteTrack::Eval( float Position, float LastPostion, uint8& OutBy
 			Position = FMath::Clamp(Position, Section->GetStartTime(), Section->GetEndTime());
 		}
 
-		OutByte = CastChecked<UMovieSceneByteSection>( Section )->Eval( Position );
+		InOutByte = CastChecked<UMovieSceneByteSection>( Section )->Eval( Position, InOutByte );
 	}
 
 	return Section != nullptr;
diff --git a/Engine/Source/Runtime/MovieSceneTracks/Private/Tracks/MovieSceneColorTrack.cpp b/Engine/Source/Runtime/MovieSceneTracks/Private/Tracks/MovieSceneColorTrack.cpp
--- a/Engine/Source/Runtime/MovieSceneTracks/Private/Tracks/MovieSceneColorTrack.cpp
+++ b/Engine/Source/Runtime/MovieSceneTracks/Private/Tracks/MovieSceneColorTrack.cpp
@@ -24,7 +24,7 @@ TSharedPtr<IMovieSceneTrackInstance> UMovieSceneColorTrack::CreateInstance()
 }
 
 
-bool UMovieSceneColorTrack::Eval( float Position, float LastPosition, FLinearColor& OutColor ) const
+bool UMovieSceneColorTrack::Eval( float Position, float LastPosition, FLinearColor& InOutColor ) const
 {
 	const UMovieSceneSection* Section = MovieSceneHelpers::FindNearestSectionAtTime( Sections, Position );
 
@@ -35,7 +35,7 @@ bool UMovieSceneColorTrack::Eval( float Position, float LastPosition, FLinearCol
 			Position = FMath::Clamp(Position, Section->GetStartTime(), Section->GetEndTime());
 		}
 
-		OutColor = CastChecked<UMovieSceneColorSection>( Section )->Eval( Position, OutColor );
+		InOutColor = CastChecked<UMovieSceneColorSection>( Section )->Eval( Position, InOutColor );
 	}
 
 	return Section != nullptr;
diff --git a/Engine/Source/Runtime/MovieSceneTracks/Private/Tracks/MovieSceneEventTrack.cpp b/Engine/Source/Runtime/MovieSceneTracks/Private/Tracks/MovieSceneEventTrack.cpp
--- a/Engine/Source/Runtime/MovieSceneTracks/Private/Tracks/MovieSceneEventTrack.cpp
+++ b/Engine/Source/Runtime/MovieSceneTracks/Private/Tracks/MovieSceneEventTrack.cpp
@@ -11,7 +11,7 @@
 /* UMovieSceneEventTrack interface
  *****************************************************************************/
 
-bool UMovieSceneEventTrack::AddKeyToSection(float Time, FName EventName, FKeyParams KeyParams)
+bool UMovieSceneEventTrack::AddKeyToSection(float Time, FName EventName)
 {
 	UMovieSceneSection* TargetSection = MovieSceneHelpers::FindNearestSectionAtTime(Sections, Time);
 
@@ -31,7 +31,7 @@ bool UMovieSceneEventTrack::AddKeyToSection(float Time, FName EventName, FKeyPar
 		return false;
 	}
 
-	EventSection->AddKey(Time, EventName, KeyParams);
+	EventSection->AddKey(Time, EventName);
 
 	return true;
 }
diff --git a/Engine/Source/Runtime/MovieSceneTracks/Private/Tracks/MovieSceneFloatTrack.cpp b/Engine/Source/Runtime/MovieSceneTracks/Private/Tracks/MovieSceneFloatTrack.cpp
--- a/Engine/Source/Runtime/MovieSceneTracks/Private/Tracks/MovieSceneFloatTrack.cpp
+++ b/Engine/Source/Runtime/MovieSceneTracks/Private/Tracks/MovieSceneFloatTrack.cpp
@@ -24,7 +24,7 @@ TSharedPtr<IMovieSceneTrackInstance> UMovieSceneFloatTrack::CreateInstance()
 }
 
 
-bool UMovieSceneFloatTrack::Eval( float Position, float LastPosition, float& OutFloat ) const
+bool UMovieSceneFloatTrack::Eval( float Position, float LastPosition, float& InOutFloat ) const
 {
 	const UMovieSceneSection* Section = MovieSceneHelpers::FindNearestSectionAtTime( Sections, Position );
 
@@ -35,7 +35,7 @@ bool UMovieSceneFloatTrack::Eval( float Position, float LastPosition, float& Out
 			Position = FMath::Clamp(Position, Section->GetStartTime(), Section->GetEndTime());
 		}
 
-		OutFloat = CastChecked<UMovieSceneFloatSection>( Section )->Eval( Position );
+		InOutFloat = CastChecked<UMovieSceneFloatSection>( Section )->Eval( Position, InOutFloat );
 	}
 
 	return Section != nullptr;
diff --git a/Engine/Source/Runtime/MovieSceneTracks/Private/Tracks/MovieSceneSpawnTrack.cpp b/Engine/Source/Runtime/MovieSceneTracks/Private/Tracks/MovieSceneSpawnTrack.cpp
--- a/Engine/Source/Runtime/MovieSceneTracks/Private/Tracks/MovieSceneSpawnTrack.cpp
+++ b/Engine/Source/Runtime/MovieSceneTracks/Private/Tracks/MovieSceneSpawnTrack.cpp
@@ -37,7 +37,7 @@ bool UMovieSceneSpawnTrack::Eval(float Position, float LastPostion, bool& bOutSp
 		Position = FMath::Clamp(Position, Section->GetStartTime(), Section->GetEndTime());
 	}
 
-	bOutSpawned = Section->Eval(Position);
+	bOutSpawned = Section->Eval(Position, bOutSpawned);
 	return true;
 }
 
diff --git a/Engine/Source/Runtime/MovieSceneTracks/Private/Tracks/MovieSceneStringTrack.cpp b/Engine/Source/Runtime/MovieSceneTracks/Private/Tracks/MovieSceneStringTrack.cpp
--- a/Engine/Source/Runtime/MovieSceneTracks/Private/Tracks/MovieSceneStringTrack.cpp
+++ b/Engine/Source/Runtime/MovieSceneTracks/Private/Tracks/MovieSceneStringTrack.cpp
@@ -12,7 +12,7 @@
 /* UMovieSceneStringTrack interface
  *****************************************************************************/
 
-bool UMovieSceneStringTrack::AddKeyToSection(float Time, const FString& String, FKeyParams KeyParams)
+bool UMovieSceneStringTrack::AddKeyToSection(float Time, const FString& String)
 {
 	UMovieSceneSection* TargetSection = MovieSceneHelpers::FindNearestSectionAtTime(Sections, Time);
 
@@ -38,7 +38,7 @@ bool UMovieSceneStringTrack::AddKeyToSection(float Time, const FString& String,
 }
 
 
-bool UMovieSceneStringTrack::Eval(float Position, float LastPostion, FString& OutString) const
+bool UMovieSceneStringTrack::Eval(float Position, float LastPostion, FString& InOutString) const
 {
 	const UMovieSceneSection* Section = MovieSceneHelpers::FindNearestSectionAtTime(Sections, Position);
 
@@ -49,7 +49,7 @@ bool UMovieSceneStringTrack::Eval(float Position, float LastPostion, FString& Ou
 			Position = FMath::Clamp(Position, Section->GetStartTime(), Section->GetEndTime());
 		}
 
-		OutString = CastChecked<UMovieSceneStringSection>(Section)->Eval(Position, OutString);
+		InOutString = CastChecked<UMovieSceneStringSection>(Section)->Eval(Position, InOutString);
 	}
 
 	return Section != nullptr;
diff --git a/Engine/Source/Runtime/MovieSceneTracks/Public/Sections/IKeyframeSection.h b/Engine/Source/Runtime/MovieSceneTracks/Public/Sections/IKeyframeSection.h
--- a/Engine/Source/Runtime/MovieSceneTracks/Public/Sections/IKeyframeSection.h
+++ b/Engine/Source/Runtime/MovieSceneTracks/Public/Sections/IKeyframeSection.h
@@ -6,8 +6,20 @@ template<typename KeyDataType>
 class IKeyframeSection
 {
 public:
+
+	/** Returns true if the supplied key data is different than the section value at the supplied time, otherwise false. */
 	virtual bool NewKeyIsNewData( float Time, const KeyDataType& KeyData ) const = 0;
+
+	/** Returns true if the section has any keys which correspond to the supplied key data, otherwise false. */
 	virtual bool HasKeys( const KeyDataType& KeyData ) const = 0;
+
+	/** Adds a key to the section at the supplied time, with the supplied data. */
 	virtual void AddKey( float Time, const KeyDataType& KeyData, EMovieSceneKeyInterpolation KeyInterpolation ) = 0;
+
+	/** Sets the default value for the section which corresponds to the supplied key data. */
 	virtual void SetDefault( const KeyDataType& KeyData ) = 0;
+
+	/** Removes all default data from this section. */
+	virtual void ClearDefaults() = 0;
+
 };
\ No newline at end of file
diff --git a/Engine/Source/Runtime/MovieSceneTracks/Public/Sections/MovieScene3DTransformSection.h b/Engine/Source/Runtime/MovieSceneTracks/Public/Sections/MovieScene3DTransformSection.h
--- a/Engine/Source/Runtime/MovieSceneTracks/Public/Sections/MovieScene3DTransformSection.h
+++ b/Engine/Source/Runtime/MovieSceneTracks/Public/Sections/MovieScene3DTransformSection.h
@@ -185,7 +185,7 @@ public:
 	 * @param Time				The position in time within the movie scene
 	 * @param OutTranslation	The evaluated translation.  Note: will remain unchanged if there were no keys to evaluate
 	 */
-	MOVIESCENETRACKS_API void EvalTranslation( float Time, FVector& OutTranslation ) const;
+	MOVIESCENETRACKS_API void EvalTranslation( float Time, FVector& InOutTranslation ) const;
 
 	/**
 	 * Evaluates the rotation component of the transform
@@ -193,7 +193,7 @@ public:
 	 * @param Time				The position in time within the movie scene
 	 * @param OutRotation		The evaluated rotation.  Note: will remain unchanged if there were no keys to evaluate
 	 */
-	MOVIESCENETRACKS_API void EvalRotation( float Time, FRotator& OutRotation ) const;
+	MOVIESCENETRACKS_API void EvalRotation( float Time, FRotator& InOutRotation ) const;
 
 	/**
 	 * Evaluates the scale component of the transform
@@ -201,7 +201,7 @@ public:
 	 * @param Time				The position in time within the movie scene
 	 * @param OutScale			The evaluated scale.  Note: will remain unchanged if there were no keys to evaluate
 	 */
-	MOVIESCENETRACKS_API void EvalScale( float Time, FVector& OutScale ) const;
+	MOVIESCENETRACKS_API void EvalScale( float Time, FVector& InOutScale ) const;
 
 	/** 
 	 * Returns the translation curve for a specific axis
@@ -253,6 +253,7 @@ public:
 	virtual bool HasKeys( const FTransformKey& KeyData ) const override;
 	virtual void AddKey( float Time, const FTransformKey& KeyData, EMovieSceneKeyInterpolation KeyInterpolation ) override;
 	virtual void SetDefault( const FTransformKey& KeyData ) override;
+	virtual void ClearDefaults() override;
 
 private:
 
diff --git a/Engine/Source/Runtime/MovieSceneTracks/Public/Sections/MovieSceneActorReferenceSection.h b/Engine/Source/Runtime/MovieSceneTracks/Public/Sections/MovieSceneActorReferenceSection.h
--- a/Engine/Source/Runtime/MovieSceneTracks/Public/Sections/MovieSceneActorReferenceSection.h
+++ b/Engine/Source/Runtime/MovieSceneTracks/Public/Sections/MovieSceneActorReferenceSection.h
@@ -25,14 +25,15 @@ public:
 	 *
 	 * @param Position	The position in time within the movie scene
 	 */
-	FGuid Eval(float Position) const;
+	MOVIESCENETRACKS_API FGuid Eval(float Position) const;
 
 	// IKeyframeSection interface.
 
-	void AddKey(float Time, const FGuid& Value, EMovieSceneKeyInterpolation KeyInterpolation);
-	bool NewKeyIsNewData(float Time, const FGuid& Value) const;
-	bool HasKeys(const FGuid& Value) const;
-	void SetDefault(const FGuid& Value);
+	virtual void AddKey(float Time, const FGuid& Value, EMovieSceneKeyInterpolation KeyInterpolation) override;
+	virtual bool NewKeyIsNewData(float Time, const FGuid& Value) const override;
+	virtual bool HasKeys(const FGuid& Value) const override;
+	virtual void SetDefault(const FGuid& Value) override;
+	virtual void ClearDefaults() override;
 
 	/**
 	 * @return The integral curve on this section
@@ -44,6 +45,8 @@ public:
 
 public:
 
+	MOVIESCENETRACKS_API FKeyHandle AddKey(float Time, const FGuid& Value);
+
 	//~ UMovieSceneSection interface
 
 	virtual void MoveSection(float DeltaPosition, TSet<FKeyHandle>& KeyHandles) override;
diff --git a/Engine/Source/Runtime/MovieSceneTracks/Public/Sections/MovieSceneBoolSection.h b/Engine/Source/Runtime/MovieSceneTracks/Public/Sections/MovieSceneBoolSection.h
--- a/Engine/Source/Runtime/MovieSceneTracks/Public/Sections/MovieSceneBoolSection.h
+++ b/Engine/Source/Runtime/MovieSceneTracks/Public/Sections/MovieSceneBoolSection.h
@@ -31,7 +31,7 @@ public:
 	 *
 	 * @param Position The position in time within the movie scene.
 	 */
-	virtual bool Eval(float Position) const;
+	virtual bool Eval(float Position, bool DefaultValue) const;
 
 	/** Gets all the keys of this boolean section. */
 	FIntegralCurve& GetCurve()
@@ -46,7 +46,8 @@ public:
 	virtual void AddKey(float Time, const bool& Value, EMovieSceneKeyInterpolation KeyInterpolation) override;
 	virtual void SetDefault(const bool& Value) override;
 	virtual bool NewKeyIsNewData(float Time, const bool& Value) const override;
-	virtual bool HasKeys(const bool& Value) const;
+	virtual bool HasKeys(const bool& Value) const override;
+	virtual void ClearDefaults() override;
 
 public:
 
diff --git a/Engine/Source/Runtime/MovieSceneTracks/Public/Sections/MovieSceneByteSection.h b/Engine/Source/Runtime/MovieSceneTracks/Public/Sections/MovieSceneByteSection.h
--- a/Engine/Source/Runtime/MovieSceneTracks/Public/Sections/MovieSceneByteSection.h
+++ b/Engine/Source/Runtime/MovieSceneTracks/Public/Sections/MovieSceneByteSection.h
@@ -23,7 +23,7 @@ public:
 	 *
 	 * @param Position The position in time within the movie scene
 	 */
-	virtual uint8 Eval(float Position) const;
+	virtual uint8 Eval(float Position, uint8 DefaultValue) const;
 
 	/** Gets all the keys of this byte section */
 	FIntegralCurve& GetCurve()
@@ -39,6 +39,7 @@ public:
 	virtual bool NewKeyIsNewData(float Time, const uint8& Value) const override;
 	virtual bool HasKeys(const uint8& Value) const override;
 	virtual void SetDefault(const uint8& Value) override;
+	virtual void ClearDefaults() override;
 
 public:
 
diff --git a/Engine/Source/Runtime/MovieSceneTracks/Public/Sections/MovieSceneColorSection.h b/Engine/Source/Runtime/MovieSceneTracks/Public/Sections/MovieSceneColorSection.h
--- a/Engine/Source/Runtime/MovieSceneTracks/Public/Sections/MovieSceneColorSection.h
+++ b/Engine/Source/Runtime/MovieSceneTracks/Public/Sections/MovieSceneColorSection.h
@@ -127,6 +127,7 @@ public:
 	virtual bool NewKeyIsNewData(float Time, const FColorKey& Key) const override;
 	virtual bool HasKeys(const FColorKey& Key) const override;
 	virtual void SetDefault(const FColorKey& Key) override;
+	virtual void ClearDefaults() override;
 
 private:
 
diff --git a/Engine/Source/Runtime/MovieSceneTracks/Public/Sections/MovieSceneEventSection.h b/Engine/Source/Runtime/MovieSceneTracks/Public/Sections/MovieSceneEventSection.h
--- a/Engine/Source/Runtime/MovieSceneTracks/Public/Sections/MovieSceneEventSection.h
+++ b/Engine/Source/Runtime/MovieSceneTracks/Public/Sections/MovieSceneEventSection.h
@@ -30,7 +30,7 @@ public:
 	 * @param EventName The name of the event to fire at the specified time.
 	 * @param KeyParams The keying parameters.
 	 */
-	void AddKey(float Time, const FName& EventName, FKeyParams KeyParams);
+	void AddKey(float Time, const FName& EventName);
 
 	/**
 	 * Get the section's event curve.
diff --git a/Engine/Source/Runtime/MovieSceneTracks/Public/Sections/MovieSceneFloatSection.h b/Engine/Source/Runtime/MovieSceneTracks/Public/Sections/MovieSceneFloatSection.h
--- a/Engine/Source/Runtime/MovieSceneTracks/Public/Sections/MovieSceneFloatSection.h
+++ b/Engine/Source/Runtime/MovieSceneTracks/Public/Sections/MovieSceneFloatSection.h
@@ -24,7 +24,7 @@ public:
 	 *
 	 * @param Position	The position in time within the movie scene
 	 */
-	virtual float Eval( float Position ) const;
+	virtual float Eval( float Position, float DefaultValue ) const;
 
 	/**
 	 * @return The float curve on this section
@@ -36,10 +36,11 @@ public:
 
 	//~ IKeyframeSection interface
 
-	void AddKey( float Time, const float& Value, EMovieSceneKeyInterpolation KeyInterpolation );
-	bool NewKeyIsNewData(float Time, const float& Value) const;
-	bool HasKeys( const float& Value ) const;
-	void SetDefault( const float& Value );
+	virtual void AddKey( float Time, const float& Value, EMovieSceneKeyInterpolation KeyInterpolation ) override;
+	virtual bool NewKeyIsNewData(float Time, const float& Value) const override;
+	virtual bool HasKeys( const float& Value ) const override;
+	virtual void SetDefault( const float& Value ) override;
+	virtual void ClearDefaults() override;
 
 public:
 
diff --git a/Engine/Source/Runtime/MovieSceneTracks/Public/Sections/MovieSceneStringSection.h b/Engine/Source/Runtime/MovieSceneTracks/Public/Sections/MovieSceneStringSection.h
--- a/Engine/Source/Runtime/MovieSceneTracks/Public/Sections/MovieSceneStringSection.h
+++ b/Engine/Source/Runtime/MovieSceneTracks/Public/Sections/MovieSceneStringSection.h
@@ -46,6 +46,7 @@ public:
 	virtual bool HasKeys( const FString& KeyData ) const;
 	virtual void AddKey( float Time, const FString& KeyData, EMovieSceneKeyInterpolation KeyInterpolation );
 	virtual void SetDefault( const FString& KeyData );
+	virtual void ClearDefaults() override;
 
 public:
 
diff --git a/Engine/Source/Runtime/MovieSceneTracks/Public/TrackInstances/MovieScene3DAttachTrackInstance.h b/Engine/Source/Runtime/MovieSceneTracks/Public/TrackInstances/MovieScene3DAttachTrackInstance.h
--- a/Engine/Source/Runtime/MovieSceneTracks/Public/TrackInstances/MovieScene3DAttachTrackInstance.h
+++ b/Engine/Source/Runtime/MovieSceneTracks/Public/TrackInstances/MovieScene3DAttachTrackInstance.h
@@ -17,5 +17,5 @@ public:
 	FMovieScene3DAttachTrackInstance( UMovieScene3DAttachTrack& InAttachTrack );
 
 	/** FMovieScene3DConstraintTrackInstance interface */
-	virtual void UpdateConstraint(float Position, const TArray<TWeakObjectPtr<UObject>>& RuntimeObjects, AActor* Actor, UMovieScene3DConstraintSection* ConstraintSection) override;
+	virtual void UpdateConstraint(EMovieSceneUpdatePass UpdatePass, float Position, const TArray<TWeakObjectPtr<UObject>>& RuntimeObjects, AActor* Actor, UMovieScene3DConstraintSection* ConstraintSection) override;
 };
diff --git a/Engine/Source/Runtime/MovieSceneTracks/Public/TrackInstances/MovieScene3DConstraintTrackInstance.h b/Engine/Source/Runtime/MovieSceneTracks/Public/TrackInstances/MovieScene3DConstraintTrackInstance.h
--- a/Engine/Source/Runtime/MovieSceneTracks/Public/TrackInstances/MovieScene3DConstraintTrackInstance.h
+++ b/Engine/Source/Runtime/MovieSceneTracks/Public/TrackInstances/MovieScene3DConstraintTrackInstance.h
@@ -23,10 +23,9 @@ public:
 	virtual void RefreshInstance( const TArray<TWeakObjectPtr<UObject>>& RuntimeObjects, IMovieScenePlayer& Player, FMovieSceneSequenceInstance& SequenceInstance ) override {}
 	virtual void ClearInstance( IMovieScenePlayer& Player, FMovieSceneSequenceInstance& SequenceInstance ) override {}
 
-	// Evaluate after TransformTrackInstance so that parents can be evaluated first
-	virtual EMovieSceneUpdatePass HasUpdatePasses() override { return MSUP_PostUpdate; }
+	virtual EMovieSceneUpdatePass HasUpdatePasses() override { return MSUP_All; }
 
-	virtual void UpdateConstraint(float Position, const TArray<TWeakObjectPtr<UObject>>& RuntimeObjects, AActor* Actor, UMovieScene3DConstraintSection* ConstraintSection) = 0;
+	virtual void UpdateConstraint(EMovieSceneUpdatePass UpdatePass, float Position, const TArray<TWeakObjectPtr<UObject>>& RuntimeObjects, AActor* Actor, UMovieScene3DConstraintSection* ConstraintSection) = 0;
 
 protected:
 	/** Track that is being instanced */
diff --git a/Engine/Source/Runtime/MovieSceneTracks/Public/TrackInstances/MovieScene3DPathTrackInstance.h b/Engine/Source/Runtime/MovieSceneTracks/Public/TrackInstances/MovieScene3DPathTrackInstance.h
--- a/Engine/Source/Runtime/MovieSceneTracks/Public/TrackInstances/MovieScene3DPathTrackInstance.h
+++ b/Engine/Source/Runtime/MovieSceneTracks/Public/TrackInstances/MovieScene3DPathTrackInstance.h
@@ -17,5 +17,5 @@ public:
 	FMovieScene3DPathTrackInstance( UMovieScene3DPathTrack& InPathTrack );
 
 	/** FMovieScene3DConstraintTrackInstance interface */
-	virtual void UpdateConstraint(float Position, const TArray<TWeakObjectPtr<UObject>>& RuntimeObjects, AActor* Actor, UMovieScene3DConstraintSection* ConstraintSection) override;
+	virtual void UpdateConstraint(EMovieSceneUpdatePass UpdatePass, float Position, const TArray<TWeakObjectPtr<UObject>>& RuntimeObjects, AActor* Actor, UMovieScene3DConstraintSection* ConstraintSection) override;
 };
diff --git a/Engine/Source/Runtime/MovieSceneTracks/Public/TrackInstances/MovieSceneSkeletalAnimationTrackInstance.h b/Engine/Source/Runtime/MovieSceneTracks/Public/TrackInstances/MovieSceneSkeletalAnimationTrackInstance.h
--- a/Engine/Source/Runtime/MovieSceneTracks/Public/TrackInstances/MovieSceneSkeletalAnimationTrackInstance.h
+++ b/Engine/Source/Runtime/MovieSceneTracks/Public/TrackInstances/MovieSceneSkeletalAnimationTrackInstance.h
@@ -23,6 +23,7 @@ public:
 	virtual void SaveState(const TArray<TWeakObjectPtr<UObject>>& RuntimeObjects, IMovieScenePlayer& Player, FMovieSceneSequenceInstance& SequenceInstance) override;
 	virtual void RestoreState(const TArray<TWeakObjectPtr<UObject>>& RuntimeObjects, IMovieScenePlayer& Player, FMovieSceneSequenceInstance& SequenceInstance) override;
 	virtual void Update(EMovieSceneUpdateData& UpdateData, const TArray<TWeakObjectPtr<UObject>>& RuntimeObjects, IMovieScenePlayer& Player, FMovieSceneSequenceInstance& SequenceInstance) override;
+	virtual EMovieSceneUpdatePass HasUpdatePasses() override { return (EMovieSceneUpdatePass)(MSUP_PreUpdate | MSUP_Update); }
 	virtual void RefreshInstance(const TArray<TWeakObjectPtr<UObject>>& RuntimeObjects, IMovieScenePlayer& Player, FMovieSceneSequenceInstance& SequenceInstance) override;
 	virtual void ClearInstance(IMovieScenePlayer& Player, FMovieSceneSequenceInstance& SequenceInstance ) override {}
 
@@ -44,6 +45,8 @@ private:
 
 	bool ShouldUsePreviewPlayback(class IMovieScenePlayer& Player, UObject* RuntimeObject) const;
 
+	void UpdateRefreshBones(const TArray<TWeakObjectPtr<UObject>>& RuntimeObjects);
+
 private:
 
 	/** Track that is being instanced */
diff --git a/Engine/Source/Runtime/MovieSceneTracks/Public/Tracks/MovieScene3DTransformTrack.h b/Engine/Source/Runtime/MovieSceneTracks/Public/Tracks/MovieScene3DTransformTrack.h
--- a/Engine/Source/Runtime/MovieSceneTracks/Public/Tracks/MovieScene3DTransformTrack.h
+++ b/Engine/Source/Runtime/MovieSceneTracks/Public/Tracks/MovieScene3DTransformTrack.h
@@ -29,7 +29,7 @@ public:
 	 * @param OutScale The evaluated scale component of the transform.
 	 * @return true if anything was evaluated. Note: if false is returned OutBool remains unchanged.
 	 */
-	virtual bool Eval(float Position, float LastPosition, FVector& OutTranslation, FRotator& OutRotation, FVector& OutScale) const;
+	virtual bool Eval(float Position, float LastPosition, FVector& InOutTranslation, FRotator& InOutRotation, FVector& InOutScale) const;
 
 public:
 
diff --git a/Engine/Source/Runtime/MovieSceneTracks/Public/Tracks/MovieSceneBoolTrack.h b/Engine/Source/Runtime/MovieSceneTracks/Public/Tracks/MovieSceneBoolTrack.h
--- a/Engine/Source/Runtime/MovieSceneTracks/Public/Tracks/MovieSceneBoolTrack.h
+++ b/Engine/Source/Runtime/MovieSceneTracks/Public/Tracks/MovieSceneBoolTrack.h
@@ -29,8 +29,8 @@ public:
 	 *
 	 * @param Position	The current playback position
 	 * @param LastPosition	The last plackback position
-	 * @param OutBool 	The value at the playback position
-	 * @return true if anything was evaluated. Note: if false is returned OutBool remains unchanged
+	 * @param InOutBool 	The current value at the playback position
+	 * @return true if anything was evaluated. Note: if false is returned InOutBool remains unchanged
 	 */
-	virtual bool Eval( float Position, float LastPostion, bool& OutBool ) const;
+	virtual bool Eval( float Position, float LastPostion, bool& InOutBool ) const;
 };
diff --git a/Engine/Source/Runtime/MovieSceneTracks/Public/Tracks/MovieSceneByteTrack.h b/Engine/Source/Runtime/MovieSceneTracks/Public/Tracks/MovieSceneByteTrack.h
--- a/Engine/Source/Runtime/MovieSceneTracks/Public/Tracks/MovieSceneByteTrack.h
+++ b/Engine/Source/Runtime/MovieSceneTracks/Public/Tracks/MovieSceneByteTrack.h
@@ -25,10 +25,10 @@ public:
 	 *
 	 * @param Position	The current playback position
 	 * @param LastPosition	The last plackback position
-	 * @param OutByte 	The value at the playback position
-	 * @return true if anything was evaluated. Note: if false is returned OutByte remains unchanged
+	 * @param InOutByte 	The value at the playback position
+	 * @return true if anything was evaluated. Note: if false is returned InOutByte remains unchanged
 	 */
-	virtual bool Eval( float Position, float LastPostion, uint8& OutByte ) const;
+	virtual bool Eval( float Position, float LastPostion, uint8& InOutByte ) const;
 
 	void SetEnum(UEnum* Enum);
 
diff --git a/Engine/Source/Runtime/MovieSceneTracks/Public/Tracks/MovieSceneColorTrack.h b/Engine/Source/Runtime/MovieSceneTracks/Public/Tracks/MovieSceneColorTrack.h
--- a/Engine/Source/Runtime/MovieSceneTracks/Public/Tracks/MovieSceneColorTrack.h
+++ b/Engine/Source/Runtime/MovieSceneTracks/Public/Tracks/MovieSceneColorTrack.h
@@ -25,10 +25,10 @@ public:
 	 *
 	 * @param Position	The current playback position
 	 * @param LastPosition	The last plackback position
-	 * @param The color at the playback position
-	 * @return true if anything was evaluated. Note: if false is returned OutColor remains unchanged
+	 * @param InOutColor The color at the playback position
+	 * @return true if anything was evaluated. Note: if false is returned InOutColor remains unchanged
 	 */
-	virtual bool Eval( float Position, float LastPostion, FLinearColor& OutColor ) const;
+	virtual bool Eval( float Position, float LastPostion, FLinearColor& InOutColor ) const;
 
 private:
 	UPROPERTY()
diff --git a/Engine/Source/Runtime/MovieSceneTracks/Public/Tracks/MovieSceneEventTrack.h b/Engine/Source/Runtime/MovieSceneTracks/Public/Tracks/MovieSceneEventTrack.h
--- a/Engine/Source/Runtime/MovieSceneTracks/Public/Tracks/MovieSceneEventTrack.h
+++ b/Engine/Source/Runtime/MovieSceneTracks/Public/Tracks/MovieSceneEventTrack.h
@@ -38,7 +38,7 @@ public:
 	 * @param EventName The name of the event to be triggered.
 	 * @param KeyParams The keying parameters
 	 */
-	bool AddKeyToSection(float Time, FName EventName, FKeyParams KeyParams);
+	bool AddKeyToSection(float Time, FName EventName);
 
 	/**
 	 * Trigger the events that fall into the given time range.
diff --git a/Engine/Source/Runtime/MovieSceneTracks/Public/Tracks/MovieSceneFloatTrack.h b/Engine/Source/Runtime/MovieSceneTracks/Public/Tracks/MovieSceneFloatTrack.h
--- a/Engine/Source/Runtime/MovieSceneTracks/Public/Tracks/MovieSceneFloatTrack.h
+++ b/Engine/Source/Runtime/MovieSceneTracks/Public/Tracks/MovieSceneFloatTrack.h
@@ -25,8 +25,8 @@ public:
 	 *
 	 * @param Position	The current playback position
 	 * @param LastPosition	The last plackback position
-	 * @param OutFloat 	The value at the playback position
-	 * @return true if anything was evaluated. Note: if false is returned OutFloat remains unchanged
+	 * @param InOutFloat 	The value at the playback position
+	 * @return true if anything was evaluated. Note: if false is returned InOutFloat remains unchanged
 	 */
-	virtual bool Eval( float Position, float LastPosition, float& OutFloat ) const;
+	virtual bool Eval( float Position, float LastPosition, float& InOutFloat ) const;
 };
diff --git a/Engine/Source/Runtime/MovieSceneTracks/Public/Tracks/MovieSceneStringTrack.h b/Engine/Source/Runtime/MovieSceneTracks/Public/Tracks/MovieSceneStringTrack.h
--- a/Engine/Source/Runtime/MovieSceneTracks/Public/Tracks/MovieSceneStringTrack.h
+++ b/Engine/Source/Runtime/MovieSceneTracks/Public/Tracks/MovieSceneStringTrack.h
@@ -36,17 +36,17 @@ public:
 	 * @param String The string to add.
 	 * @param KeyParams The keying parameters
 	 */
-	bool AddKeyToSection(float Time, const FString& String, FKeyParams KeyParams);
+	bool AddKeyToSection(float Time, const FString& String);
 
 	/**
 	 * Evaluates the track at the playback position
 	 *
 	 * @param Position The psootion at which to evaluate this track.
 	 * @param LastPositionThe last playback position.
-	 * @param OutString The string at the evaluation time.
-	 * @return true if anything was evaluated. Note: if false is returned OutColor remains unchanged
+	 * @param InOutString The string at the evaluation time.
+	 * @return true if anything was evaluated. Note: if false is returned InOutString remains unchanged
 	 */
-	virtual bool Eval(float Position, float LastPostion, FString& OutString) const;
+	virtual bool Eval(float Position, float LastPostion, FString& InOutString) const;
 
 public:
 
diff --git a/Engine/Source/Runtime/UMG/Private/Animation/MovieScene2DTransformSection.cpp b/Engine/Source/Runtime/UMG/Private/Animation/MovieScene2DTransformSection.cpp
--- a/Engine/Source/Runtime/UMG/Private/Animation/MovieScene2DTransformSection.cpp
+++ b/Engine/Source/Runtime/UMG/Private/Animation/MovieScene2DTransformSection.cpp
@@ -253,3 +253,13 @@ void UMovieScene2DTransformSection::SetDefault( const struct F2DTransformKey& Tr
 	FRichCurve* KeyCurve = GetCurveForChannelAndAxis( TransformKey.Channel, TransformKey.Axis, Translation, Scale, Shear, &Rotation );
 	SetCurveDefault( *KeyCurve, TransformKey.Value );
 }
+
+
+void UMovieScene2DTransformSection::ClearDefaults()
+{
+	Rotation.ClearDefaultValue();
+	Translation[0].ClearDefaultValue();
+	Translation[1].ClearDefaultValue();
+	Scale[0].ClearDefaultValue();
+	Scale[1].ClearDefaultValue();
+}
diff --git a/Engine/Source/Runtime/UMG/Private/Animation/MovieSceneMarginSection.cpp b/Engine/Source/Runtime/UMG/Private/Animation/MovieSceneMarginSection.cpp
--- a/Engine/Source/Runtime/UMG/Private/Animation/MovieSceneMarginSection.cpp
+++ b/Engine/Source/Runtime/UMG/Private/Animation/MovieSceneMarginSection.cpp
@@ -168,3 +168,10 @@ void UMovieSceneMarginSection::SetDefault( const FMarginKey& Key )
 	SetCurveDefault( *KeyCurve, Key.Value );
 }
 
+void UMovieSceneMarginSection::ClearDefaults()
+{
+	LeftCurve.ClearDefaultValue();
+	TopCurve.ClearDefaultValue();
+	RightCurve.ClearDefaultValue();
+	BottomCurve.ClearDefaultValue();
+}
diff --git a/Engine/Source/Runtime/Windows/XAudio2/Private/XAudio2Device.cpp b/Engine/Source/Runtime/Windows/XAudio2/Private/XAudio2Device.cpp
--- a/Engine/Source/Runtime/Windows/XAudio2/Private/XAudio2Device.cpp
+++ b/Engine/Source/Runtime/Windows/XAudio2/Private/XAudio2Device.cpp
@@ -197,7 +197,7 @@ bool FXAudio2Device::InitializeHardware()
 	}
 
 #if XAUDIO_SUPPORTS_DEVICE_DETAILS
-	UE_LOG(LogInit, Log, TEXT("XAudio2 using '%s' : %d channels at %g kHz using %d bits per sample (channel mask 0x%x)"),
+	UE_LOG(LogInit, Log, TEXT( "XAudio2 using '%s' : %d channels at %g kHz using %d bits per sample (channel mask 0x%x)" ), 
 		FXAudioDeviceProperties::DeviceDetails.DisplayName,
 		FXAudioDeviceProperties::NumSpeakers, 
 		( float )SampleRate / 1000.0f, 
@@ -249,6 +249,9 @@ bool FXAudio2Device::InitializeHardware()
 		CommonAudioPoolFreeBytes = 0;
 	}
 
+	// Now initialize the audio clock voice after xaudio2 is initialized
+	DeviceProperties->InitAudioClockVoice();
+
 	return true;
 }
 
@@ -289,6 +292,9 @@ void FXAudio2Device::UpdateHardware()
 		}
 #endif
 	}
+
+	// Update the audio clock time
+	AudioClock = DeviceProperties->GetAudioClockTime();
 }
 
 FAudioEffectsManager* FXAudio2Device::CreateEffectsManager()
diff --git a/Engine/Source/Runtime/Windows/XAudio2/Private/XAudio2Support.h b/Engine/Source/Runtime/Windows/XAudio2/Private/XAudio2Support.h
--- a/Engine/Source/Runtime/Windows/XAudio2/Private/XAudio2Support.h
+++ b/Engine/Source/Runtime/Windows/XAudio2/Private/XAudio2Support.h
@@ -742,6 +742,11 @@ struct FXAudioDeviceProperties : public IDeviceChangedListener
 	struct IXAudio2MasteringVoice*		MasteringVoice;
 	HMODULE								XAudio2Dll;
 
+	// Audio clock info
+	struct IXAudio2SourceVoice*			AudioClockVoice;
+	XAUDIO2_BUFFER						AudioClockXAudio2Buffer;
+	TArray<int16>						AudioClockPCMBufferData;
+
 	// These variables are static because they are common across all audio device instances
 	static int32						NumSpeakers;
 	static const float*					OutputMixMatrix;
@@ -767,6 +772,7 @@ struct FXAudioDeviceProperties : public IDeviceChangedListener
 
 	FXAudioDeviceProperties()
 		: XAudio2(nullptr)
+		, AudioClockVoice(nullptr)
 		, MasteringVoice(nullptr)
 		, XAudio2Dll(nullptr)
 		, NumActiveVoices(0)
@@ -775,7 +781,7 @@ struct FXAudioDeviceProperties : public IDeviceChangedListener
 		if (NotificationClient == nullptr)
 		{
 			NotificationClient = new FMMNotificationClient();
-		}
+	}
 		else
 		{
 			NotificationClient->AddRef();
@@ -838,6 +844,53 @@ struct FXAudioDeviceProperties : public IDeviceChangedListener
 		return bChanged;
 	}
 
+	void InitAudioClockVoice()
+	{
+		if (XAudio2)
+		{
+			WAVEFORMATEX PCMFormat;
+			PCMFormat.wFormatTag = WAVE_FORMAT_PCM;
+			PCMFormat.nChannels = 1;
+			PCMFormat.nSamplesPerSec = 44100;
+			PCMFormat.wBitsPerSample = 16;
+			PCMFormat.cbSize = 0;
+			PCMFormat.nBlockAlign = sizeof(int16);
+			PCMFormat.nAvgBytesPerSec = sizeof(int16) * 44100;
+
+			check(XAudio2 != nullptr);
+			bool bSuccess = Validate(TEXT("CreateSourceVoice, GetAudioClockTime"),
+				XAudio2->CreateSourceVoice(&AudioClockVoice, &PCMFormat));
+
+			if (bSuccess && AudioClockVoice)
+			{
+				// Setup the zeroed clock buffer
+				AudioClockPCMBufferData.Reset();
+				AudioClockPCMBufferData.AddZeroed(64);
+
+				// Setup the xaudio2 buffer of looping silence
+				FMemory::Memzero(&AudioClockXAudio2Buffer, sizeof(XAUDIO2_BUFFER));
+				AudioClockXAudio2Buffer.AudioBytes = AudioClockPCMBufferData.Num() * sizeof(int16);
+				AudioClockXAudio2Buffer.pAudioData = (uint8*)AudioClockPCMBufferData.GetData();
+				AudioClockXAudio2Buffer.LoopCount = XAUDIO2_LOOP_INFINITE;
+
+				// Submit it and start the voice... basically an always looping silent sound
+				AudioClockVoice->SubmitSourceBuffer(&AudioClockXAudio2Buffer);
+				AudioClockVoice->Start(0);
+			}
+		}
+	}
+
+	double GetAudioClockTime()
+	{
+		if (AudioClockVoice)
+		{
+			XAUDIO2_VOICE_STATE VoiceState;
+			AudioClockVoice->GetState(&VoiceState);
+			return (double)VoiceState.SamplesPlayed / 44100;
+		}
+		return 0.0;
+	}
+
 	bool Validate(const TCHAR* Function, uint32 ErrorCode) const
 	{
 		if (ErrorCode != S_OK)
