Copying //UE4/Dev-Sequencer to //UE4/Dev-Main (Source: //UE4/Dev-Sequencer @ 3178529)

diff --git a/Engine/Source/Runtime/MovieScene/Public/MovieSceneSpawnable.h b/Engine/Source/Runtime/MovieScene/Public/MovieSceneSpawnable.h
--- a/Engine/Source/Runtime/MovieScene/Public/MovieSceneSpawnable.h
+++ b/Engine/Source/Runtime/MovieScene/Public/MovieSceneSpawnable.h
@@ -31,7 +31,6 @@ struct FMovieSceneSpawnable
 		: ObjectTemplate(nullptr)
 		, Ownership(ESpawnOwnership::InnerSequence)
 #if WITH_EDITORONLY_DATA
-		, bIgnoreOwnershipInEditor(false)
 		, GeneratedClass_DEPRECATED(nullptr)
 #endif
 	{
@@ -44,7 +43,6 @@ struct FMovieSceneSpawnable
 		, ObjectTemplate(&InObjectTemplate)
 		, Ownership(ESpawnOwnership::InnerSequence)
 #if WITH_EDITORONLY_DATA
-		, bIgnoreOwnershipInEditor(false)
 		, GeneratedClass_DEPRECATED(nullptr)
 #endif
 	{
@@ -162,24 +160,6 @@ public:
 		Ownership = InOwnership;
 	}
 
-#if WITH_EDITORONLY_DATA
-	/**
-	 * Check whether this spawnable should remain spawned when outside the play-range, regardless of ownership
-	 */
-	bool ShouldIgnoreOwnershipInEditor() const
-	{
-		return bIgnoreOwnershipInEditor;
-	}
-
-	/**
-	 * Set whether this spawnable should remain spawned when outside the play-range, regardless of ownership
-	 */
-	void SetIgnoreOwnershipInEditor(bool bInIgnoreOwnershipInEditor)
-	{
-		bIgnoreOwnershipInEditor = bInIgnoreOwnershipInEditor;
-	}
-#endif
-
 private:
 
 	/** Unique identifier of the spawnable object. */
@@ -206,10 +186,6 @@ private:
 	ESpawnOwnership Ownership;
 
 #if WITH_EDITORONLY_DATA
-	/** When true, this spawnable will not respect its ownership sematics outside of the playback range, when spawned from inside a currently editing sequence */
-	UPROPERTY()
-	bool bIgnoreOwnershipInEditor;
-
 public:
 	/** Deprecated generated class */
 	UPROPERTY()
diff --git a/Engine/Source/Runtime/MovieScene/Public/MovieSceneTrack.h b/Engine/Source/Runtime/MovieScene/Public/MovieSceneTrack.h
--- a/Engine/Source/Runtime/MovieScene/Public/MovieSceneTrack.h
+++ b/Engine/Source/Runtime/MovieScene/Public/MovieSceneTrack.h
@@ -3,11 +3,16 @@
 #pragma once
 
 #include "EnumClassFlags.h"
+#include "MovieSceneFwd.h"
+#include "InlineValue.h"
+#include "MovieSceneSignedObject.h"
 #include "MovieSceneTrack.generated.h"
 
-
 class UMovieSceneSection;
-class FMovieSceneSequenceInstance;
+struct FMovieSceneEvaluationTrack;
+struct FMovieSceneEvalTemplatePtr;
+struct FMovieSceneSegmentCompilerRules;
+struct FMovieSceneSequenceTemplateStore;
 
 /** Flags used to perform cook-time optimization of movie scene data */
 enum class ECookOptimizationFlags
@@ -21,32 +26,189 @@ enum class ECookOptimizationFlags
 };
 ENUM_CLASS_FLAGS(ECookOptimizationFlags)
 
+/** Movie scene compilation parameters. Serialized items contribute to a compiled template's cached hash */
+USTRUCT()
+struct FMovieSceneTrackCompilationParams
+{
+	GENERATED_BODY()
+
+	FMovieSceneTrackCompilationParams()
+		: bForEditorPreview(false)
+	{
+	}
+
+	friend bool operator!=(const FMovieSceneTrackCompilationParams& A, const FMovieSceneTrackCompilationParams& B)
+	{
+		return !(A == B);
+	}
+
+	friend bool operator==(const FMovieSceneTrackCompilationParams& A, const FMovieSceneTrackCompilationParams& B)
+	{
+		return A.bForEditorPreview == B.bForEditorPreview;
+	}
+
+	/** Whether we're generating for an editor preview, or for efficient runtime evaluation */
+	UPROPERTY()
+	bool bForEditorPreview;
+};
+
+/** Track compiler arguments */
+struct FMovieSceneTrackCompilerArgs
+{
+	FMovieSceneTrackCompilerArgs(IMovieSceneTemplateGenerator& InGenerator, FMovieSceneSequenceTemplateStore& InSubSequenceStore)
+		: Generator(InGenerator)
+		, SubSequenceStore(InSubSequenceStore)
+	{
+	}
+
+	/** Compilation parameters */
+	FMovieSceneTrackCompilationParams Params;
+
+	/** The object binding ID that this track belongs to. */
+	FGuid ObjectBindingId;
+
+	/** The generator responsible for generating the template */
+	IMovieSceneTemplateGenerator& Generator;
+
+	/** Store that describes how to find sub sequence templates */
+	FMovieSceneSequenceTemplateStore& SubSequenceStore;
+};
+
+/** Generic evaluation options for any track */
+USTRUCT()
+struct FMovieSceneTrackEvalOptions
+{
+	GENERATED_BODY()
+	
+	FMovieSceneTrackEvalOptions()
+		: bCanEvaluateNearestSection(false)
+		, bEvaluateNearestSection(false)
+		, bEvaluateInPreroll(true)
+		, bEvaluateInPostroll(true)
+	{}
+
+	/** true when the value of bEvaluateNearestSection is to be considered for the track */
+	UPROPERTY()
+	uint32 bCanEvaluateNearestSection : 1;
+
+	/** When evaluating empty space on a track, will evaluate the last position of the previous section (if possible), or the first position of the next section, in that order of preference. */
+	UPROPERTY(EditAnywhere, Category="General", meta=(EditCondition=bCanEvaluateNearestSection))
+	uint32 bEvaluateNearestSection : 1;
+
+	/** Evaluate this track as part of its parent sub-section's pre-roll, if applicable */
+	UPROPERTY(EditAnywhere, Category="General")
+	uint32 bEvaluateInPreroll : 1;
+
+	/** Evaluate this track as part of its parent sub-section's post-roll, if applicable */
+	UPROPERTY(EditAnywhere, Category="General")
+	uint32 bEvaluateInPostroll : 1;
+};
+
+/** Enumeration specifying the result of a compilation */
+enum class EMovieSceneCompileResult : uint8
+{
+	/** The compilation was successful */
+	Success,
+	/** The compilation was not successful */
+	Failure,
+	/** No compilation routine was implemented */
+	Unimplemented
+};
+
 /**
  * Base class for a track in a Movie Scene
  */
 UCLASS(abstract, MinimalAPI)
 class UMovieSceneTrack
-	: public UObject
+	: public UMovieSceneSignedObject
 {
 	GENERATED_BODY()
 
 public:
 
-	UMovieSceneTrack(const FObjectInitializer& InInitializer)
-		: Super(InInitializer)
-	{
-#if WITH_EDITORONLY_DATA
-		TrackTint = FColor(127, 127, 127, 0);
-#endif
-	}
+	MOVIESCENE_API UMovieSceneTrack(const FObjectInitializer& InInitializer);
+
+public:
+
+	/** General evaluation options for a given track */
+	UPROPERTY(EditAnywhere, Category="General", meta=(ShowOnlyInnerProperties))
+	FMovieSceneTrackEvalOptions EvalOptions;
+
+public:
+
+	//~ Methods relating to compilation 
+
+	/** 
+	 * Get compiler rules to use when compiling sections that overlap on the same row.
+	 * These define how to deal with overlapping sections and empty space on a row
+	 */
+	MOVIESCENE_API virtual TInlineValue<FMovieSceneSegmentCompilerRules> GetRowCompilerRules() const;
+
+	/** 
+	 * Get compiler rules to use when compiling sections that overlap on different rows.
+	 * These define how to deal with overlapping sections and empty space at the track level
+	 */
+	MOVIESCENE_API virtual TInlineValue<FMovieSceneSegmentCompilerRules> GetTrackCompilerRules() const;
 
 	/**
-	 * Creates a runtime instance of this class.
-	 * 
-	 * @param SequenceInstance The sequence instance to create this track instance for
-	 * @return The created runtime instance. This must not be nullptr (note sharedptr not sharedref due to compatibility with uobjects not actually allowing pure virtuals)
+	 * Generate a template for this track
+	 *
+	 * @param Args 			Compilation arguments
+	 */
+	MOVIESCENE_API virtual void GenerateTemplate(const FMovieSceneTrackCompilerArgs& Args) const;
+
+protected:
+
+	/**
+	 * Overridable user defined custom compilation method
+	 *
+	 * @param Track 		Destination track to compile into
+	 * @param Args 			Compilation arguments
+	 * @return Compilation result
+	 */
+	virtual EMovieSceneCompileResult CustomCompile(FMovieSceneEvaluationTrack& Track, const FMovieSceneTrackCompilerArgs& Args) const { return EMovieSceneCompileResult::Unimplemented; }
+
+	/**
+	 * Called after this track has been compiled, regardless of whether it was compiled through CustomCompile, or the default logic
+	 *
+	 * @param Track 		Destination track to compile into
+	 * @param Args 			Compilation arguments
 	 */
-	virtual TSharedPtr<class IMovieSceneTrackInstance> CreateInstance() PURE_VIRTUAL(UMovieSceneTrack::CreateInstance, check(false); return nullptr;);
+	virtual void PostCompile(FMovieSceneEvaluationTrack& Track, const FMovieSceneTrackCompilerArgs& Args) const {}
+
+public:
+
+	/**
+	 * Internal function to create a legacy track instance
+	 */
+	TSharedPtr<class IMovieSceneTrackInstance> CreateLegacyInstance() const;
+
+protected:
+
+	/**
+	 * Create a movie scene eval template for the specified section
+	 *
+	 * @param InSection		The section to create a template for
+	 * @return A template, or null
+	 */
+	MOVIESCENE_API virtual FMovieSceneEvalTemplatePtr CreateTemplateForSection(const UMovieSceneSection& InSection) const;
+
+	/**
+	 * Compile this movie scene track into an efficient runtime structure
+	 *
+	 * @param Track 		Destination track to compile into
+	 * @param Args 			Compilation arguments
+	 * @return Compilation result
+	 */
+	MOVIESCENE_API EMovieSceneCompileResult Compile(FMovieSceneEvaluationTrack& Track, const FMovieSceneTrackCompilerArgs& Args) const;
+
+	/**
+	 * Deprecated method of creating track instances
+	 */
+	DEPRECATED(4.15, "Create Instance has been deprecated. Please provide an evaluation template through CreateTemplateForSection instead.")
+	virtual TSharedPtr<class IMovieSceneTrackInstance> CreateInstance() { return nullptr; }
+
+public:
 
 	/**
 	 * @return The name that makes this track unique from other track of the same class.
diff --git a/Engine/Source/Runtime/MovieSceneCapture/Private/MovieSceneCapture.cpp b/Engine/Source/Runtime/MovieSceneCapture/Private/MovieSceneCapture.cpp
--- a/Engine/Source/Runtime/MovieSceneCapture/Private/MovieSceneCapture.cpp
+++ b/Engine/Source/Runtime/MovieSceneCapture/Private/MovieSceneCapture.cpp
@@ -8,6 +8,7 @@
 #include "SceneViewExtension.h"
 #include "JsonObjectConverter.h"
 #include "RemoteConfigIni.h"
+#include "Scalability.h"
 #include "SceneViewport.h"
 
 #if WITH_EDITOR
@@ -42,6 +43,7 @@ FMovieSceneCaptureSettings::FMovieSceneCaptureSettings()
 	FrameRate = 24;
 	ZeroPadFrameNumbers = 4;
 	bEnableTextureStreaming = false;
+	bCinematicEngineScalability = true;
 	bCinematicMode = true;
 	bAllowMovement = false;
 	bAllowTurning = false;
@@ -134,6 +136,12 @@ void UMovieSceneCapture::Initialize(TSharedPtr<FSceneViewport> InSceneViewport,
 			Settings.HandleFrames = HandleFramesOverride;
 		}
 
+		bool bOverrideCinematicEngineScalabilityMode;
+		if( FParse::Bool( FCommandLine::Get(), TEXT( "-MovieEngineScalabilityMode=" ), bOverrideCinematicEngineScalabilityMode ) )
+		{
+			Settings.bCinematicEngineScalability = bOverrideCinematicEngineScalabilityMode;
+		}
+
 		bool bOverrideCinematicMode;
 		if( FParse::Bool( FCommandLine::Get(), TEXT( "-MovieCinematicMode=" ), bOverrideCinematicMode ) )
 		{
@@ -178,6 +186,14 @@ void UMovieSceneCapture::Initialize(TSharedPtr<FSceneViewport> InSceneViewport,
 
 void UMovieSceneCapture::StartWarmup()
 {
+	if (Settings.bCinematicEngineScalability)
+	{
+		CachedQualityLevels = Scalability::GetQualityLevels();
+		Scalability::FQualityLevels QualityLevels = CachedQualityLevels;
+		QualityLevels.SetFromSingleQualityLevelRelativeToMax(0);
+		Scalability::SetQualityLevels(QualityLevels);
+	}
+
 	check( !bCapturing );
 	if( !CaptureStrategy.IsValid() )
 	{
@@ -244,6 +260,11 @@ void UMovieSceneCapture::FinalizeWhenReady()
 
 void UMovieSceneCapture::Finalize()
 {
+	if (Settings.bCinematicEngineScalability)
+	{
+		Scalability::SetQualityLevels(CachedQualityLevels);
+	}
+
 	FActiveMovieSceneCaptures::Get().Remove(this);
 
 	if (bCapturing)
diff --git a/Engine/Source/Runtime/MovieSceneCapture/Private/MovieSceneCaptureHelpers.cpp b/Engine/Source/Runtime/MovieSceneCapture/Private/MovieSceneCaptureHelpers.cpp
--- a/Engine/Source/Runtime/MovieSceneCapture/Private/MovieSceneCaptureHelpers.cpp
+++ b/Engine/Source/Runtime/MovieSceneCapture/Private/MovieSceneCaptureHelpers.cpp
@@ -381,7 +381,7 @@ bool MovieSceneCaptureHelpers::ImportEDL(UMovieScene* InMovieScene, float InFram
 			if (ShotSection)
 			{
 				ShotSection->Modify();
-				ShotSection->StartOffset = ShotData.SourceInTime;
+				ShotSection->Parameters.StartOffset = ShotData.SourceInTime;
 				ShotSection->SetStartTime(ShotData.EditInTime);
 				ShotSection->SetEndTime(ShotData.EditOutTime);
 			}
diff --git a/Engine/Source/Runtime/MovieSceneCapture/Public/MovieSceneCapture.h b/Engine/Source/Runtime/MovieSceneCapture/Public/MovieSceneCapture.h
--- a/Engine/Source/Runtime/MovieSceneCapture/Public/MovieSceneCapture.h
+++ b/Engine/Source/Runtime/MovieSceneCapture/Public/MovieSceneCapture.h
@@ -10,6 +10,7 @@
 #include "IMovieSceneCaptureProtocol.h"
 #include "MovieSceneCaptureProtocolRegistry.h"
 #include "MovieSceneCaptureProtocolSettings.h"
+#include "Scalability.h"
 #include "MovieSceneCapture.generated.h"
 
 class FJsonObject;
@@ -183,6 +184,8 @@ protected:
 	FOnCaptureFinished OnCaptureFinishedDelegate;
 	/** Format string used for frame numbers */
 	FString FrameNumberFormat;
+	/** Cached quality levels */
+	Scalability::FQualityLevels CachedQualityLevels;
 };
 
 /** A strategy that employs a fixed frame time-step, and as such never drops a frame. Potentially accelerated. */
diff --git a/Engine/Source/Runtime/MovieSceneCapture/Public/MovieSceneCaptureSettings.h b/Engine/Source/Runtime/MovieSceneCapture/Public/MovieSceneCaptureSettings.h
--- a/Engine/Source/Runtime/MovieSceneCapture/Public/MovieSceneCaptureSettings.h
+++ b/Engine/Source/Runtime/MovieSceneCapture/Public/MovieSceneCaptureSettings.h
@@ -82,6 +82,10 @@ struct MOVIESCENECAPTURE_API FMovieSceneCaptureSettings
 	UPROPERTY(config, EditAnywhere, Category=CaptureSettings, AdvancedDisplay)
 	bool bEnableTextureStreaming;
 
+	/** Whether to enable cinematic engine scalability settings */
+	UPROPERTY(config, EditAnywhere, Category=Cinematic)
+	bool bCinematicEngineScalability;
+
 	/** Whether to enable cinematic mode whilst capturing */
 	UPROPERTY(config, EditAnywhere, Category=Cinematic)
 	bool bCinematicMode;
diff --git a/Engine/Source/Runtime/MovieSceneTracks/Private/Evaluation/MovieScene3DAttachTemplate.cpp b/Engine/Source/Runtime/MovieSceneTracks/Private/Evaluation/MovieScene3DAttachTemplate.cpp
new file mode 100644
--- /dev/null
+++ b/Engine/Source/Runtime/MovieSceneTracks/Private/Evaluation/MovieScene3DAttachTemplate.cpp
@@ -0,0 +1,185 @@
+// Copyright 1998-2016 Epic Games, Inc. All Rights Reserved.
+
+#include "MovieSceneTracksPrivatePCH.h"
+#include "MovieScene3DAttachTemplate.h"
+
+#include "MovieSceneExecutionToken.h"
+#include "MovieScene3DAttachSection.h"
+#include "MovieSceneTemplateCommon.h"
+
+DECLARE_CYCLE_STAT(TEXT("Attach Track Evaluate"), MovieSceneEval_AttachTrack_Evaluate, STATGROUP_MovieSceneEval);
+DECLARE_CYCLE_STAT(TEXT("Attach Track Token Execute"), MovieSceneEval_AttachTrack_TokenExecute, STATGROUP_MovieSceneEval);
+
+/** A token that sets a component's attach */
+struct F3DAttachTrackToken
+{
+	F3DAttachTrackToken(USceneComponent* InAttachParent, FName InAttachSocketName, bool bInShouldbeAttached)
+		: AttachParent(InAttachParent)
+		, AttachSocketName(InAttachSocketName)
+		, bShouldBeAttached(bInShouldbeAttached)
+	{}
+
+	void Apply(USceneComponent& SceneComponent) const
+	{
+		if (bShouldBeAttached)
+		{
+			if (SceneComponent.GetAttachParent() != AttachParent.Get() || SceneComponent.GetAttachSocketName() != AttachSocketName)
+			{
+				SceneComponent.AttachToComponent(AttachParent.Get(), FAttachmentTransformRules::KeepRelativeTransform, AttachSocketName);	
+			}
+		}
+		else
+		{
+			SceneComponent.DetachFromComponent(FDetachmentTransformRules::KeepRelativeTransform);
+		}
+	}
+
+	TWeakObjectPtr<USceneComponent> AttachParent;
+	FName AttachSocketName;
+	bool bShouldBeAttached;
+};
+
+/** A movie scene pre-animated token that stores a pre-animated component attachment */
+struct F3DAttachTrackPreAnimatedToken : F3DAttachTrackToken, IMovieScenePreAnimatedToken
+{
+	F3DAttachTrackPreAnimatedToken(USceneComponent& SceneComponent)
+		: F3DAttachTrackToken(SceneComponent.GetAttachParent(), SceneComponent.GetAttachSocketName(), true)
+	{}
+
+	virtual void RestoreState(UObject& InObject, IMovieScenePlayer& Player) override
+	{
+		USceneComponent* SceneComponent = CastChecked<USceneComponent>(&InObject);
+
+		Apply(*SceneComponent);
+
+		// Detach if there was no pre-existing parent
+		if (!AttachParent.IsValid())
+		{
+			SceneComponent->DetachFromComponent(FDetachmentTransformRules::KeepRelativeTransform);
+		}
+	}
+};
+
+struct F3DAttachTokenProducer : IMovieScenePreAnimatedTokenProducer
+{
+	/** Cache the existing state of an object before moving it */
+	virtual IMovieScenePreAnimatedTokenPtr CacheExistingState(UObject& Object) const override
+	{
+		USceneComponent* SceneComponent = CastChecked<USceneComponent>(&Object);
+		return F3DAttachTrackPreAnimatedToken(*SceneComponent);
+	}
+};
+
+/** A movie scene execution token that stores a specific attach, and an operand */
+struct F3DAttachExecutionToken
+	: IMovieSceneExecutionToken
+	, F3DAttachTrackToken
+{
+	F3DAttachExecutionToken(FGuid InAttachGuid, FName InAttachSocketName, FName InAttachComponentName, bool bInShouldBeAttached) 
+		: F3DAttachTrackToken(nullptr, InAttachSocketName, bInShouldBeAttached)
+		, AttachGuid(InAttachGuid)
+		, AttachComponentName(InAttachComponentName)
+	{}
+
+	USceneComponent* GetAttachComponent(AActor* InParentActor)
+	{
+		if (AttachSocketName != NAME_None)
+		{
+			if (AttachComponentName != NAME_None )
+			{
+				TInlineComponentArray<USceneComponent*> PotentialAttachComponents(InParentActor);
+				for (USceneComponent* PotentialAttachComponent : PotentialAttachComponents)
+				{
+					if (PotentialAttachComponent->GetFName() == AttachComponentName && PotentialAttachComponent->DoesSocketExist(AttachSocketName))
+					{
+						return PotentialAttachComponent;
+					}
+				}
+			}
+			else if (InParentActor->GetRootComponent()->DoesSocketExist(AttachSocketName))
+			{
+				return InParentActor->GetRootComponent();
+			}
+		}
+		else if (AttachComponentName != NAME_None )
+		{
+			TInlineComponentArray<USceneComponent*> PotentialAttachComponents(InParentActor);
+			for (USceneComponent* PotentialAttachComponent : PotentialAttachComponents)
+			{
+				if (PotentialAttachComponent->GetFName() == AttachComponentName)
+				{
+					return PotentialAttachComponent;
+				}
+			}
+		}
+
+		if (InParentActor->GetDefaultAttachComponent())
+		{
+			return InParentActor->GetDefaultAttachComponent();
+		}
+		else
+		{
+			return InParentActor->GetRootComponent();
+		}
+	}
+
+	/** Execute this token, operating on all objects referenced by 'Operand' */
+	virtual void Execute(const FMovieSceneContext& Context, const FMovieSceneEvaluationOperand& Operand, FPersistentEvaluationData& PersistentData, IMovieScenePlayer& Player) override
+	{
+		MOVIESCENE_DETAILED_SCOPE_CYCLE_COUNTER(MovieSceneEval_AttachTrack_TokenExecute)
+
+		FMovieSceneEvaluationOperand AttachOperand(Operand.SequenceID, AttachGuid);
+		
+		TArrayView<TWeakObjectPtr<>> Objects = Player.FindBoundObjects(AttachOperand);
+		if (!Objects.Num())
+		{
+			return;
+		}
+
+		// Only ever deal with one parent
+		UObject* ParentObject = Objects[0].Get();
+		AActor* ParentActor = Cast<AActor>(ParentObject);
+
+		if (!ParentActor)
+		{
+			return;
+		}
+
+		for (TWeakObjectPtr<> Object : Player.FindBoundObjects(Operand))
+		{
+			UObject* ObjectPtr = Object.Get();
+			USceneComponent* SceneComponent = ObjectPtr ? MovieSceneHelpers::SceneComponentFromRuntimeObject(ObjectPtr) : nullptr;
+
+			if (SceneComponent)
+			{
+				Player.SavePreAnimatedState(*SceneComponent, FMobilityTokenProducer::GetAnimTypeID(), FMobilityTokenProducer());
+				Player.SavePreAnimatedState(*SceneComponent, TMovieSceneAnimTypeID<F3DAttachExecutionToken>(), F3DAttachTokenProducer());
+
+				AttachParent = GetAttachComponent(ParentActor);
+
+				SceneComponent->SetMobility(EComponentMobility::Movable);
+
+				Apply(*SceneComponent);
+			}
+		}
+	}
+	
+	FGuid AttachGuid;
+	FName AttachComponentName;
+};
+
+FMovieScene3DAttachSectionTemplate::FMovieScene3DAttachSectionTemplate(const UMovieScene3DAttachSection& Section)
+	: AttachGuid(Section.GetConstraintId())
+	, AttachSocketName(Section.AttachSocketName)
+	, AttachComponentName(Section.AttachComponentName)
+{
+}
+
+void FMovieScene3DAttachSectionTemplate::Evaluate(const FMovieSceneEvaluationOperand& Operand, const FMovieSceneContext& Context, const FPersistentEvaluationData& PersistentData, FMovieSceneExecutionTokens& ExecutionTokens) const
+{
+	MOVIESCENE_DETAILED_SCOPE_CYCLE_COUNTER(MovieSceneEval_AttachTrack_Evaluate)
+
+	const bool bShouldBeAttached = true;
+
+	ExecutionTokens.Add(F3DAttachExecutionToken(AttachGuid, AttachSocketName, AttachComponentName, bShouldBeAttached));
+}
diff --git a/Engine/Source/Runtime/MovieSceneTracks/Private/Evaluation/MovieScene3DAttachTemplate.h b/Engine/Source/Runtime/MovieSceneTracks/Private/Evaluation/MovieScene3DAttachTemplate.h
new file mode 100644
--- /dev/null
+++ b/Engine/Source/Runtime/MovieSceneTracks/Private/Evaluation/MovieScene3DAttachTemplate.h
@@ -0,0 +1,34 @@
+// Copyright 1998-2016 Epic Games, Inc. All Rights Reserved.
+
+#pragma once
+
+#include "MovieSceneEvalTemplate.h"
+#include "MovieScene3DAttachTemplate.generated.h"
+
+class UMovieScene3DAttachSection;
+
+USTRUCT()
+struct FMovieScene3DAttachSectionTemplate : public FMovieSceneEvalTemplate
+{
+	GENERATED_BODY()
+	
+	FMovieScene3DAttachSectionTemplate() {}
+	FMovieScene3DAttachSectionTemplate(const UMovieScene3DAttachSection& Section);
+
+	/** GUID of the parent we should attach to */
+	UPROPERTY()
+	FGuid AttachGuid;
+
+	/** Name of the socket to attach to */
+	UPROPERTY()
+	FName AttachSocketName;
+
+	/** Name of the component to attach to */
+	UPROPERTY()
+	FName AttachComponentName;
+
+private:
+
+	virtual UScriptStruct& GetScriptStructImpl() const override { return *StaticStruct(); }
+	virtual void Evaluate(const FMovieSceneEvaluationOperand& Operand, const FMovieSceneContext& Context, const FPersistentEvaluationData& PersistentData, FMovieSceneExecutionTokens& ExecutionTokens) const override;
+};
\ No newline at end of file
diff --git a/Engine/Source/Runtime/MovieSceneTracks/Private/Evaluation/MovieScene3DTransformTemplate.cpp b/Engine/Source/Runtime/MovieSceneTracks/Private/Evaluation/MovieScene3DTransformTemplate.cpp
new file mode 100644
--- /dev/null
+++ b/Engine/Source/Runtime/MovieSceneTracks/Private/Evaluation/MovieScene3DTransformTemplate.cpp
@@ -0,0 +1,77 @@
+// Copyright 1998-2016 Epic Games, Inc. All Rights Reserved.
+
+#include "MovieSceneTracksPrivatePCH.h"
+#include "MovieScene3DTransformTemplate.h"
+
+#include "MovieSceneExecutionToken.h"
+#include "MovieScene3DTransformSection.h"
+#include "MovieSceneTemplateCommon.h"
+
+DECLARE_CYCLE_STAT(TEXT("Transform Track Evaluate"), MovieSceneEval_TransformTrack_Evaluate, STATGROUP_MovieSceneEval);
+DECLARE_CYCLE_STAT(TEXT("Transform Track Token Execute"), MovieSceneEval_TransformTrack_TokenExecute, STATGROUP_MovieSceneEval);
+
+/** A movie scene execution token that stores a specific transform, and an operand */
+struct F3DTransformTrackExecutionToken
+	: IMovieSceneExecutionToken
+	, F3DTransformTrackToken
+{
+	/** Execute this token, operating on all objects referenced by 'Operand' */
+	virtual void Execute(const FMovieSceneContext& Context, const FMovieSceneEvaluationOperand& Operand, FPersistentEvaluationData& PersistentData, IMovieScenePlayer& Player) override
+	{
+		MOVIESCENE_DETAILED_SCOPE_CYCLE_COUNTER(MovieSceneEval_TransformTrack_TokenExecute)
+
+		for (TWeakObjectPtr<> Object : Player.FindBoundObjects(Operand))
+		{
+			UObject* ObjectPtr = Object.Get();
+			USceneComponent* SceneComponent = ObjectPtr ? MovieSceneHelpers::SceneComponentFromRuntimeObject(ObjectPtr) : nullptr;
+
+			if (SceneComponent)
+			{
+				Player.SavePreAnimatedState(*SceneComponent, FMobilityTokenProducer::GetAnimTypeID(), FMobilityTokenProducer());
+				Player.SavePreAnimatedState(*SceneComponent, F3DTransformTokenProducer::GetAnimTypeID(), F3DTransformTokenProducer());
+	
+				SceneComponent->SetMobility(EComponentMobility::Movable);
+
+				Apply(*SceneComponent, Context.GetDelta());
+			}
+		}
+	}
+};
+
+FMovieScene3DTransformSectionTemplate::FMovieScene3DTransformSectionTemplate(const UMovieScene3DTransformSection& Section)
+{
+	TranslationCurve[0]	= Section.GetTranslationCurve(EAxis::X);
+	TranslationCurve[1]	= Section.GetTranslationCurve(EAxis::Y);
+	TranslationCurve[2]	= Section.GetTranslationCurve(EAxis::Z);
+
+	RotationCurve[0]	= Section.GetRotationCurve(EAxis::X);
+	RotationCurve[1]	= Section.GetRotationCurve(EAxis::Y);
+	RotationCurve[2]	= Section.GetRotationCurve(EAxis::Z);
+
+	ScaleCurve[0]		= Section.GetScaleCurve(EAxis::X);
+	ScaleCurve[1]		= Section.GetScaleCurve(EAxis::Y);
+	ScaleCurve[2]		= Section.GetScaleCurve(EAxis::Z);
+}
+
+void FMovieScene3DTransformSectionTemplate::Evaluate(const FMovieSceneEvaluationOperand& Operand, const FMovieSceneContext& Context, const FPersistentEvaluationData& PersistentData, FMovieSceneExecutionTokens& ExecutionTokens) const
+{
+	MOVIESCENE_DETAILED_SCOPE_CYCLE_COUNTER(MovieSceneEval_TransformTrack_Evaluate)
+
+	F3DTransformTrackExecutionToken Token;
+
+	const float Time = Context.GetTime();
+
+	Token.Translation.X = 	TranslationCurve[0].Eval(Time);
+	Token.Translation.Y = 	TranslationCurve[1].Eval(Time);
+	Token.Translation.Z = 	TranslationCurve[2].Eval(Time);
+
+	Token.Rotation.Roll = 	RotationCurve[0].Eval(Time);
+	Token.Rotation.Pitch = 	RotationCurve[1].Eval(Time);
+	Token.Rotation.Yaw = 	RotationCurve[2].Eval(Time);
+
+	Token.Scale.X = 		ScaleCurve[0].Eval(Time);
+	Token.Scale.Y = 		ScaleCurve[1].Eval(Time);
+	Token.Scale.Z = 		ScaleCurve[2].Eval(Time);
+
+	ExecutionTokens.Add(Token);
+}
\ No newline at end of file
diff --git a/Engine/Source/Runtime/MovieSceneTracks/Private/Evaluation/MovieScene3DTransformTemplate.h b/Engine/Source/Runtime/MovieSceneTracks/Private/Evaluation/MovieScene3DTransformTemplate.h
new file mode 100644
--- /dev/null
+++ b/Engine/Source/Runtime/MovieSceneTracks/Private/Evaluation/MovieScene3DTransformTemplate.h
@@ -0,0 +1,31 @@
+// Copyright 1998-2016 Epic Games, Inc. All Rights Reserved.
+
+#pragma once
+
+#include "MovieSceneEvalTemplate.h"
+#include "MovieScene3DTransformTemplate.generated.h"
+
+class UMovieScene3DTransformSection;
+
+USTRUCT()
+struct FMovieScene3DTransformSectionTemplate : public FMovieSceneEvalTemplate
+{
+	GENERATED_BODY()
+	
+	FMovieScene3DTransformSectionTemplate() {}
+	FMovieScene3DTransformSectionTemplate(const UMovieScene3DTransformSection& Section);
+
+	UPROPERTY()
+	FRichCurve TranslationCurve[3];
+
+	UPROPERTY()
+	FRichCurve RotationCurve[3];
+
+	UPROPERTY()
+	FRichCurve ScaleCurve[3];
+
+private:
+
+	virtual UScriptStruct& GetScriptStructImpl() const override { return *StaticStruct(); }
+	virtual void Evaluate(const FMovieSceneEvaluationOperand& Operand, const FMovieSceneContext& Context, const FPersistentEvaluationData& PersistentData, FMovieSceneExecutionTokens& ExecutionTokens) const override;
+};
\ No newline at end of file
diff --git a/Engine/Source/Runtime/MovieSceneTracks/Private/Evaluation/MovieSceneActorReferenceTemplate.cpp b/Engine/Source/Runtime/MovieSceneTracks/Private/Evaluation/MovieSceneActorReferenceTemplate.cpp
new file mode 100644
--- /dev/null
+++ b/Engine/Source/Runtime/MovieSceneTracks/Private/Evaluation/MovieSceneActorReferenceTemplate.cpp
@@ -0,0 +1,66 @@
+// Copyright 1998-2016 Epic Games, Inc. All Rights Reserved.
+
+#include "MovieSceneTracksPrivatePCH.h"
+#include "MovieSceneActorReferenceTemplate.h"
+
+#include "MovieSceneActorReferenceSection.h"
+#include "MovieSceneActorReferenceTrack.h"
+
+namespace PropertyTemplate
+{
+	template<>
+	AActor* ConvertFromIntermediateType<AActor*, FGuid>(const FGuid& InObjectGuid, const FMovieSceneEvaluationOperand& Operand, FPersistentEvaluationData& PersistentData, IMovieScenePlayer& Player)
+	{
+		FMovieSceneEvaluationOperand NewOperand = Operand;
+		NewOperand.ObjectBindingID = InObjectGuid;
+
+		for (TWeakObjectPtr<>& WeakObject : Player.FindBoundObjects(NewOperand))
+		{
+			if (AActor* Actor = Cast<AActor>(WeakObject.Get()))
+			{
+				return Actor;
+			}
+		}
+		return nullptr;
+	}
+
+	template<>
+	AActor* ConvertFromIntermediateType<AActor*, TWeakObjectPtr<AActor>>(const TWeakObjectPtr<AActor>& InWeakPtr, IMovieScenePlayer& Player)
+	{
+		return InWeakPtr.Get();
+	}
+
+	template<>
+	AActor* ConvertFromIntermediateType<AActor*, TWeakObjectPtr<AActor>>(const TWeakObjectPtr<AActor>& InWeakPtr, const FMovieSceneEvaluationOperand& Operand, FPersistentEvaluationData& PersistentData, IMovieScenePlayer& Player)
+	{
+		return InWeakPtr.Get();
+	}
+
+	static bool IsValueValid(AActor* InValue)
+	{
+		return InValue != nullptr;
+	}
+
+	template<> IMovieScenePreAnimatedTokenPtr CacheExistingState<AActor*, FGuid>(UObject& Object, FTrackInstancePropertyBindings& PropertyBindings)
+	{
+		return TCachedState<AActor*, TWeakObjectPtr<AActor>>(PropertyBindings.GetCurrentValue<AActor*>(Object), PropertyBindings);
+	}
+}
+
+FMovieSceneActorReferenceSectionTemplate::FMovieSceneActorReferenceSectionTemplate(const UMovieSceneActorReferenceSection& Section, const UMovieScenePropertyTrack& Track)
+	: PropertyData(Track.GetPropertyName(), Track.GetPropertyPath())
+	, ActorGuidIndexCurve(Section.GetActorReferenceCurve())
+	, ActorGuids(Section.GetActorGuids())
+{
+}
+
+void FMovieSceneActorReferenceSectionTemplate::Evaluate(const FMovieSceneEvaluationOperand& Operand, const FMovieSceneContext& Context, const FPersistentEvaluationData& PersistentData, FMovieSceneExecutionTokens& ExecutionTokens) const
+{
+	using namespace PropertyTemplate;
+	FSectionData& TrackData = PersistentData.GetSectionData<FSectionData>();
+
+	const int32 ActorGuidIndex = ActorGuidIndexCurve.Evaluate(Context.GetTime());
+	FGuid ObjectBindingID = ActorGuids.IsValidIndex(ActorGuidIndex) ? ActorGuids[ActorGuidIndex] : FGuid();
+
+	ExecutionTokens.Add(TPropertyTrackExecutionToken<AActor*, FGuid>(ObjectBindingID));
+}
\ No newline at end of file
diff --git a/Engine/Source/Runtime/MovieSceneTracks/Private/Evaluation/MovieSceneActorReferenceTemplate.h b/Engine/Source/Runtime/MovieSceneTracks/Private/Evaluation/MovieSceneActorReferenceTemplate.h
new file mode 100644
--- /dev/null
+++ b/Engine/Source/Runtime/MovieSceneTracks/Private/Evaluation/MovieSceneActorReferenceTemplate.h
@@ -0,0 +1,45 @@
+// Copyright 1998-2016 Epic Games, Inc. All Rights Reserved.
+
+#pragma once
+
+#include "MovieScenePropertyTemplate.h"
+#include "Curves/IntegralCurve.h"
+#include "MovieSceneActorReferenceTemplate.generated.h"
+
+class UMovieSceneActorReferenceSection;
+class UMovieScenePropertyTrack;
+
+USTRUCT()
+struct FMovieSceneActorReferenceSectionTemplate : public FMovieSceneEvalTemplate
+{
+	GENERATED_BODY()
+	
+	FMovieSceneActorReferenceSectionTemplate(){}
+	FMovieSceneActorReferenceSectionTemplate(const UMovieSceneActorReferenceSection& Section, const UMovieScenePropertyTrack& Track);
+
+private:
+
+	virtual UScriptStruct& GetScriptStructImpl() const override
+	{
+		return *StaticStruct();
+	}
+	virtual void SetupOverrides() override
+	{
+		EnableOverrides(RequiresSetupFlag);
+	}
+	virtual void Setup(FPersistentEvaluationData& PersistentData, IMovieScenePlayer& Player) const override
+	{
+		PropertyData.SetupTrack(PersistentData);
+	}
+
+	virtual void Evaluate(const FMovieSceneEvaluationOperand& Operand, const FMovieSceneContext& Context, const FPersistentEvaluationData& PersistentData, FMovieSceneExecutionTokens& ExecutionTokens) const override;
+
+	UPROPERTY()
+	FMovieScenePropertySectionData PropertyData;
+
+	UPROPERTY()
+	FIntegralCurve ActorGuidIndexCurve;
+
+	UPROPERTY()
+	TArray<FGuid> ActorGuids;
+};
\ No newline at end of file
diff --git a/Engine/Source/Runtime/MovieSceneTracks/Private/Evaluation/MovieSceneAudioTemplate.cpp b/Engine/Source/Runtime/MovieSceneTracks/Private/Evaluation/MovieSceneAudioTemplate.cpp
new file mode 100644
--- /dev/null
+++ b/Engine/Source/Runtime/MovieSceneTracks/Private/Evaluation/MovieSceneAudioTemplate.cpp
@@ -0,0 +1,298 @@
+// Copyright 1998-2016 Epic Games, Inc. All Rights Reserved.
+
+#include "MovieSceneTracksPrivatePCH.h"
+#include "MovieSceneAudioTemplate.h"
+
+#include "MovieSceneExecutionToken.h"
+
+#include "SoundDefinitions.h"
+#include "Components/AudioComponent.h"
+
+#include "MovieSceneAudioSection.h"
+#include "MovieSceneAudioTrack.h"
+
+DECLARE_CYCLE_STAT(TEXT("Audio Track Evaluate"), MovieSceneEval_AudioTrack_Evaluate, STATGROUP_MovieSceneEval);
+DECLARE_CYCLE_STAT(TEXT("Audio Track Token Execute"), MovieSceneEval_AudioTrack_TokenExecute, STATGROUP_MovieSceneEval);
+
+/** Stop audio from playing */
+struct FStopAudioPreAnimatedToken : IMovieScenePreAnimatedToken
+{
+	static FMovieSceneAnimTypeID GetAnimTypeID()
+	{
+		return TMovieSceneAnimTypeID<FStopAudioPreAnimatedToken>();
+	}
+
+	virtual void RestoreState(UObject& InObject, IMovieScenePlayer& Player) override
+	{
+		UAudioComponent* AudioComponent = CastChecked<UAudioComponent>(&InObject);
+		AudioComponent->Stop();
+		AudioComponent->DestroyComponent();
+	}
+
+	struct FProducer : IMovieScenePreAnimatedTokenProducer
+	{
+		virtual IMovieScenePreAnimatedTokenPtr CacheExistingState(UObject& Object) const override
+		{
+			return FStopAudioPreAnimatedToken();
+		}
+	};
+};
+
+/** Destroy a transient audio component */
+struct FDestroyAudioPreAnimatedToken : IMovieScenePreAnimatedToken
+{
+	static FMovieSceneAnimTypeID GetAnimTypeID()
+	{
+		return TMovieSceneAnimTypeID<FDestroyAudioPreAnimatedToken>();
+	}
+
+	virtual void RestoreState(UObject& InObject, IMovieScenePlayer& Player) override
+	{
+		UAudioComponent* AudioComponent = CastChecked<UAudioComponent>(&InObject);
+		AudioComponent->DestroyComponent();
+	}
+
+	struct FProducer : IMovieScenePreAnimatedTokenProducer
+	{
+		virtual IMovieScenePreAnimatedTokenPtr CacheExistingState(UObject& Object) const override
+		{
+			return FDestroyAudioPreAnimatedToken();
+		}
+	};
+};
+
+struct FCachedAudioTrackData : IPersistentEvaluationData
+{
+	TArray<TMap<FObjectKey, TWeakObjectPtr<UAudioComponent>>> AudioComponentsByRow;
+
+	UAudioComponent* GetAudioComponentForRow(int32 RowIndex, FObjectKey Key)
+	{
+		if (AudioComponentsByRow.Num() > RowIndex)
+		{
+			return AudioComponentsByRow[RowIndex].FindRef(Key).Get();
+		}
+
+		return nullptr;
+	}
+
+	/** Only to be called on the game thread */
+	UAudioComponent* AddAudioComponentForRow(int32 RowIndex, AActor& PrincipalActor, IMovieScenePlayer& Player)
+	{
+		const int32 NumExtra = RowIndex + 1 - AudioComponentsByRow.Num();
+		if (NumExtra > 0)
+		{
+			AudioComponentsByRow.AddDefaulted(NumExtra);
+		}
+
+		FObjectKey ActorKey(&PrincipalActor);
+		UAudioComponent* ExistingComponent = AudioComponentsByRow[RowIndex].FindRef(ActorKey).Get();
+		if (!ExistingComponent)
+		{
+			USoundCue* TempPlaybackAudioCue = NewObject<USoundCue>();
+
+			FAudioDevice::FCreateComponentParams Params(PrincipalActor.GetWorld(), &PrincipalActor);
+			ExistingComponent = FAudioDevice::CreateComponent(TempPlaybackAudioCue, Params);
+
+			if (!ExistingComponent)
+			{
+				FString ActorName = 
+#if WITH_EDITOR
+					PrincipalActor.GetActorLabel();
+#else
+					PrincipalActor.GetName();
+#endif
+				UE_LOG(LogMovieScene, Warning, TEXT("Failed to create audio component for spatialized audio track (row %d on %s)."), RowIndex, *ActorName);
+				return nullptr;
+			}
+
+			Player.SavePreAnimatedState(*ExistingComponent, FMovieSceneAnimTypeID::Unique(), FDestroyAudioPreAnimatedToken::FProducer());
+
+			AudioComponentsByRow[RowIndex].Add(ActorKey, ExistingComponent);
+
+			ExistingComponent->SetFlags(RF_Transient);
+			ExistingComponent->AttachToComponent(PrincipalActor.GetRootComponent(), FAttachmentTransformRules::KeepRelativeTransform);
+		}
+
+		return ExistingComponent;
+	}
+
+	/** Only to be called on the game thread */
+	UAudioComponent* AddMasterAudioComponentForRow(int32 RowIndex, UWorld* World, IMovieScenePlayer& Player)
+	{
+		const int32 NumExtra = RowIndex + 1 - AudioComponentsByRow.Num();
+		if (NumExtra > 0)
+		{
+			AudioComponentsByRow.AddDefaulted(NumExtra);
+		}
+
+		UAudioComponent* ExistingComponent = AudioComponentsByRow[RowIndex].FindRef(FObjectKey()).Get();
+		if (!ExistingComponent)
+		{
+			USoundCue* TempPlaybackAudioCue = NewObject<USoundCue>();
+			ExistingComponent = FAudioDevice::CreateComponent(TempPlaybackAudioCue, FAudioDevice::FCreateComponentParams(World));
+
+			if (!ExistingComponent)
+			{
+				UE_LOG(LogMovieScene, Warning, TEXT("Failed to create audio component for master audio track (row %d)."), RowIndex);
+				return nullptr;
+			}
+
+			Player.SavePreAnimatedState(*ExistingComponent, FMovieSceneAnimTypeID::Unique(), FDestroyAudioPreAnimatedToken::FProducer());
+			
+			ExistingComponent->SetFlags(RF_Transient);
+			AudioComponentsByRow[RowIndex].Add(FObjectKey(), ExistingComponent);
+		}
+
+		return ExistingComponent;
+	}
+
+	void StopAllSounds()
+	{
+		for (TMap<FObjectKey, TWeakObjectPtr<UAudioComponent>>& Map : AudioComponentsByRow)
+		{
+			for (auto& Pair : Map)
+			{
+				if (UAudioComponent* AudioComponent = Pair.Value.Get())
+				{
+					AudioComponent->Stop();
+				}
+			}
+		}
+	}
+};
+
+
+struct FAudioSectionExecutionToken : IMovieSceneExecutionToken
+{
+	FAudioSectionExecutionToken(const FMovieSceneAudioSectionTemplateData& InAudioData)
+		: AudioData(InAudioData)
+	{}
+
+	virtual void Execute(const FMovieSceneContext& Context, const FMovieSceneEvaluationOperand& Operand, FPersistentEvaluationData& PersistentData, IMovieScenePlayer& Player)
+	{
+		FCachedAudioTrackData& TrackData = PersistentData.GetOrAddTrackData<FCachedAudioTrackData>();
+
+		if (Context.GetStatus() != EMovieScenePlayerStatus::Playing && Context.GetStatus() != EMovieScenePlayerStatus::Scrubbing)
+		{
+			// stopped, recording, etc
+			TrackData.StopAllSounds();
+		}
+
+		// Master audio track
+		else if (!Operand.ObjectBindingID.IsValid())
+		{
+			UObject* PlaybackContext = Player.GetPlaybackContext();
+
+			UAudioComponent* AudioComponent = TrackData.GetAudioComponentForRow(AudioData.RowIndex, FObjectKey());
+			if (!AudioComponent)
+			{
+				// Initialize the sound
+				AudioComponent = TrackData.AddMasterAudioComponentForRow(AudioData.RowIndex, PlaybackContext ? PlaybackContext->GetWorld() : nullptr, Player);
+			}
+
+			if (AudioComponent)
+			{
+				AudioData.EnsureAudioIsPlaying(*AudioComponent, PersistentData, Context, false, Player);
+			}
+		}
+
+		// Object binding audio track
+		else
+		{
+			for (TWeakObjectPtr<> Object : Player.FindBoundObjects(Operand))
+			{
+				AActor* Actor = Cast<AActor>(Object.Get());
+				if (!Actor)
+				{
+					continue;
+				}
+
+				UAudioComponent* AudioComponent = TrackData.GetAudioComponentForRow(AudioData.RowIndex, Actor);
+				if (!AudioComponent)
+				{
+					// Initialize the sound
+					AudioComponent = TrackData.AddAudioComponentForRow(AudioData.RowIndex, *Actor, Player);
+				}
+
+				
+				if (AudioComponent)
+				{
+					AudioData.EnsureAudioIsPlaying(*AudioComponent, PersistentData, Context, true, Player);
+				}
+			}
+		}
+	}
+
+	FMovieSceneAudioSectionTemplateData AudioData;
+};
+
+FMovieSceneAudioSectionTemplateData::FMovieSceneAudioSectionTemplateData(const UMovieSceneAudioSection& Section)
+	: Sound(Section.GetSound())
+	, AudioStartTime(Section.GetAudioStartTime())
+	, AudioDilationFactor(Section.GetAudioDilationFactor())
+	, AudioVolume(Section.GetAudioVolume())
+	, RowIndex(Section.GetRowIndex())
+{
+}
+
+void FMovieSceneAudioSectionTemplateData::EnsureAudioIsPlaying(UAudioComponent& AudioComponent, FPersistentEvaluationData& PersistentData, const FMovieSceneContext& Context, bool bAllowSpatialization, IMovieScenePlayer& Player) const
+{
+	Player.SavePreAnimatedState(AudioComponent, FStopAudioPreAnimatedToken::GetAnimTypeID(), FStopAudioPreAnimatedToken::FProducer());
+
+	bool bPlaySound = !AudioComponent.IsPlaying() || AudioComponent.Sound != Sound;
+	if (bPlaySound)
+	{
+		float PitchMultiplier = 1.f / AudioDilationFactor;
+
+		AudioComponent.bAllowSpatialization = bAllowSpatialization;
+		AudioComponent.bOverrideAttenuation = bAllowSpatialization;
+		AudioComponent.Stop();
+		AudioComponent.SetSound(Sound);
+		AudioComponent.SetVolumeMultiplier(AudioVolume);
+		AudioComponent.SetPitchMultiplier(PitchMultiplier);
+		AudioComponent.bIsUISound = true;
+		AudioComponent.Play(Context.GetTime() - AudioStartTime);
+
+		if (Context.GetStatus() == EMovieScenePlayerStatus::Scrubbing)
+		{
+			// Fade out the sound at the same volume in order to simply
+			// set a short duration on the sound, far from ideal soln
+			AudioComponent.FadeOut(AudioTrackConstants::ScrubDuration, 1.f);
+		}
+	}
+
+	if (bAllowSpatialization)
+	{
+		if (FAudioDevice* AudioDevice = AudioComponent.GetAudioDevice())
+		{
+			DECLARE_CYCLE_STAT(TEXT("FAudioThreadTask.MovieSceneUpdateAudioTransform"), STAT_MovieSceneUpdateAudioTransform, STATGROUP_TaskGraphTasks);
+
+			const FTransform ActorTransform = AudioComponent.GetOwner()->GetTransform();
+			const uint64 ActorComponentID = AudioComponent.GetAudioComponentID();
+			FAudioThread::RunCommandOnAudioThread([AudioDevice, ActorComponentID, ActorTransform]()
+			{
+				if (FActiveSound* ActiveSound = AudioDevice->FindActiveSound(ActorComponentID))
+				{
+					ActiveSound->bLocationDefined = true;
+					ActiveSound->Transform = ActorTransform;
+				}
+			}, GET_STATID(STAT_MovieSceneUpdateAudioTransform));
+		}
+	}
+}
+
+FMovieSceneAudioSectionTemplate::FMovieSceneAudioSectionTemplate(const UMovieSceneAudioSection& Section)
+	: AudioData(Section)
+{
+}
+
+
+void FMovieSceneAudioSectionTemplate::Evaluate(const FMovieSceneEvaluationOperand& Operand, const FMovieSceneContext& Context, const FPersistentEvaluationData& PersistentData, FMovieSceneExecutionTokens& ExecutionTokens) const
+{
+	MOVIESCENE_DETAILED_SCOPE_CYCLE_COUNTER(MovieSceneEval_AudioTrack_Evaluate)
+
+	if (GEngine && GEngine->UseSound() && !Context.HasJumped() && Context.GetStatus() != EMovieScenePlayerStatus::Jumping)
+	{
+		ExecutionTokens.Add(FAudioSectionExecutionToken(AudioData));
+	}
+}
\ No newline at end of file
diff --git a/Engine/Source/Runtime/MovieSceneTracks/Private/Evaluation/MovieSceneAudioTemplate.h b/Engine/Source/Runtime/MovieSceneTracks/Private/Evaluation/MovieSceneAudioTemplate.h
new file mode 100644
--- /dev/null
+++ b/Engine/Source/Runtime/MovieSceneTracks/Private/Evaluation/MovieSceneAudioTemplate.h
@@ -0,0 +1,58 @@
+// Copyright 1998-2016 Epic Games, Inc. All Rights Reserved.
+
+#pragma once
+
+#include "MovieSceneEvalTemplate.h"
+#include "MovieSceneAudioTemplate.generated.h"
+
+class UMovieSceneAudioSection;
+class UAudioComponent;
+
+USTRUCT()
+struct FMovieSceneAudioSectionTemplateData
+{
+	GENERATED_BODY()
+
+	FMovieSceneAudioSectionTemplateData() {}
+	FMovieSceneAudioSectionTemplateData(const UMovieSceneAudioSection& Section);
+
+	/** Ensure that the sound is playing for the specified audio component and data */
+	void EnsureAudioIsPlaying(UAudioComponent& AudioComponent, FPersistentEvaluationData& PersistentData, const FMovieSceneContext& Context, bool bAllowSpatialization, IMovieScenePlayer& Player) const;
+
+	/** The sound cue or wave that this template plays. Not to be dereferenced on a background thread */
+	UPROPERTY()
+	USoundBase* Sound;
+
+	/** The absolute time that the sound starts playing at */
+	UPROPERTY()
+	float AudioStartTime;
+	
+	/** The amount which this audio is time dilated by */
+	UPROPERTY()
+	float AudioDilationFactor;
+
+	/** The volume the sound will be played with. */
+	UPROPERTY()
+	float AudioVolume;
+
+	/** The row index of the section */
+	UPROPERTY()
+	int32 RowIndex;
+};
+
+USTRUCT()
+struct FMovieSceneAudioSectionTemplate : public FMovieSceneEvalTemplate
+{
+	GENERATED_BODY()
+	
+	FMovieSceneAudioSectionTemplate() {}
+	FMovieSceneAudioSectionTemplate(const UMovieSceneAudioSection& Section);
+
+	UPROPERTY()
+	FMovieSceneAudioSectionTemplateData AudioData;
+
+private:
+
+	virtual UScriptStruct& GetScriptStructImpl() const override { return *StaticStruct(); }
+	virtual void Evaluate(const FMovieSceneEvaluationOperand& Operand, const FMovieSceneContext& Context, const FPersistentEvaluationData& PersistentData, FMovieSceneExecutionTokens& ExecutionTokens) const override;
+};
\ No newline at end of file
diff --git a/Engine/Source/Runtime/MovieSceneTracks/Private/Evaluation/MovieSceneCameraAnimTemplate.cpp b/Engine/Source/Runtime/MovieSceneTracks/Private/Evaluation/MovieSceneCameraAnimTemplate.cpp
new file mode 100644
--- /dev/null
+++ b/Engine/Source/Runtime/MovieSceneTracks/Private/Evaluation/MovieSceneCameraAnimTemplate.cpp
@@ -0,0 +1,510 @@
+// Copyright 1998-2016 Epic Games, Inc. All Rights Reserved.
+
+#include "MovieSceneTracksPrivatePCH.h"
+#include "MovieSceneCameraAnimTemplate.h"
+#include "MovieSceneCameraAnimSection.h"
+
+DECLARE_CYCLE_STAT(TEXT("Camera Anim Track Token Execute"), MovieSceneEval_CameraAnimTrack_TokenExecute, STATGROUP_MovieSceneEval);
+
+/** Structure that holds blended post processing settings */
+struct FBlendedPostProcessSettings : FPostProcessSettings
+{
+	FBlendedPostProcessSettings() : Weight(0.f) {}
+	FBlendedPostProcessSettings(float InWeight, const FPostProcessSettings& InSettings) : FPostProcessSettings(InSettings), Weight(InWeight) {}
+
+	/** The weighting to apply to these settings */
+	float Weight;
+};
+
+
+/** Persistent data that exists as long as a given camera track is being evaluated */
+struct FMovieSceneAdditiveCameraData : IPersistentEvaluationData
+{
+	static FMovieSceneAdditiveCameraData& Get(const FMovieSceneEvaluationOperand& Operand, FPersistentEvaluationData& PersistentData)
+	{
+		return PersistentData.GetOrAdd<FMovieSceneAdditiveCameraData>(FSharedPersistentDataKey(FMovieSceneAdditiveCameraAnimationTrackTemplate::SharedDataId, Operand));
+	}
+
+	/** Reset the additive camera values */
+	void Reset()
+	{
+		TotalFOVOffset = 0.f;
+		TotalTransform = FTransform::Identity;
+		BlendedPostProcessSettings.Reset();
+
+		bApplyTransform = false;
+		bApplyPostProcessing = false;
+	}
+
+	/** Accumulate the given post processing settings for this frame */
+	void AccumulatePostProcessing(const FPostProcessSettings& InPostProcessSettings, float Weight)
+	{
+		if (Weight > 0.f)
+		{
+			BlendedPostProcessSettings.Add(FBlendedPostProcessSettings(Weight, InPostProcessSettings));
+		}
+
+		bApplyPostProcessing = true;
+	}
+
+	/** Accumulate the transform and FOV offset */
+	void AccumulateOffset(const FTransform& AdditiveOffset, float AdditiveFOVOffset)
+	{
+		TotalTransform = TotalTransform * AdditiveOffset;
+		TotalFOVOffset += AdditiveFOVOffset;
+
+		bApplyTransform = true;
+	}
+
+	/** Apply any cumulative animation states */
+	void ApplyCumulativeAnimation(UCameraComponent& CameraComponent) const
+	{
+		if (bApplyPostProcessing)
+		{
+			CameraComponent.ClearExtraPostProcessBlends();
+			for (const FBlendedPostProcessSettings& Settings : BlendedPostProcessSettings)
+			{
+				CameraComponent.AddExtraPostProcessBlend(Settings, Settings.Weight);
+			}
+		}
+
+		if (bApplyTransform)
+		{
+			CameraComponent.ClearAdditiveOffset();
+			CameraComponent.AddAdditiveOffset(TotalTransform, TotalFOVOffset);
+		}
+	}
+
+	ACameraActor* GetTempCameraActor(IMovieScenePlayer& Player)
+	{
+		if (!TempCameraActor.IsValid())
+		{
+			// spawn the temp CameraActor used for updating CameraAnims
+			FActorSpawnParameters SpawnInfo;
+			SpawnInfo.SpawnCollisionHandlingOverride = ESpawnActorCollisionHandlingMethod::AlwaysSpawn;
+
+			// We never want to save these temp actors into a map
+			SpawnInfo.ObjectFlags |= RF_Transient;
+			ACameraActor* Cam = Player.GetPlaybackContext()->GetWorld()->SpawnActor<ACameraActor>(SpawnInfo);
+			if (Cam)
+			{
+#if WITH_EDITOR
+				Cam->SetIsTemporarilyHiddenInEditor(true);
+#endif
+				TempCameraActor = Cam;
+
+				struct FDestroyTempObject : IMovieScenePreAnimatedGlobalToken
+				{
+					FDestroyTempObject(AActor& InActor) : TempActor(&InActor) {}
+
+					virtual void RestoreState(IMovieScenePlayer& InPlayer)
+					{
+						AActor* Actor = TempActor.Get();
+						if (Actor)
+						{
+							Actor->Destroy(false, false);
+						}
+					}
+
+					TWeakObjectPtr<AActor> TempActor;
+				};
+
+				Player.SavePreAnimatedState(*Cam, FMovieSceneAnimTypeID::Unique(), FTempCameraPreAnimatedStateProducer());
+			}
+		}
+
+		return TempCameraActor.Get();
+	}
+
+private:
+
+	struct FTempCameraPreAnimatedStateProducer : IMovieScenePreAnimatedTokenProducer
+	{
+		virtual IMovieScenePreAnimatedTokenPtr CacheExistingState(UObject& Object) const override
+		{
+			struct FTempCameraPreAnimatedState : IMovieScenePreAnimatedToken
+			{
+				virtual void RestoreState(UObject& InObject, IMovieScenePlayer& InPlayer)
+				{
+					AActor* Actor = CastChecked<AActor>(&InObject);
+					Actor->Destroy(false, false);
+				}
+			};
+
+			return FTempCameraPreAnimatedState();
+		}
+	};
+
+	bool bApplyTransform, bApplyPostProcessing;
+	TArray<FBlendedPostProcessSettings, TInlineAllocator<2>> BlendedPostProcessSettings;
+	FTransform TotalTransform;
+	float TotalFOVOffset;
+	TWeakObjectPtr<ACameraActor> TempCameraActor;
+};
+
+/** Pre animated token that restores a camera component's additive transform */
+struct FPreAnimatedCameraTransformTokenProducer : IMovieScenePreAnimatedTokenProducer
+{
+	virtual IMovieScenePreAnimatedTokenPtr CacheExistingState(UObject& Object) const override
+	{
+		struct FRestoreToken : IMovieScenePreAnimatedToken
+		{
+			virtual void RestoreState(UObject& InObject, IMovieScenePlayer& Player) override
+			{
+				UCameraComponent* CameraComponent = CastChecked<UCameraComponent>(&InObject);
+				CameraComponent->ClearAdditiveOffset();
+			}
+		};
+
+		return FRestoreToken();
+	}
+};
+
+/** Pre animated token that restores a camera component's blended post processing settings */
+struct FPreAnimatedPostProcessingBlendsTokenProducer : IMovieScenePreAnimatedTokenProducer
+{
+	virtual IMovieScenePreAnimatedTokenPtr CacheExistingState(UObject& Object) const override
+	{
+		struct FRestoreToken : IMovieScenePreAnimatedToken
+		{
+			virtual void RestoreState(UObject& InObject, IMovieScenePlayer& Player) override
+			{
+				UCameraComponent* CameraComponent = CastChecked<UCameraComponent>(&InObject);
+				CameraComponent->ClearExtraPostProcessBlends();
+			}
+		};
+
+		return FRestoreToken();
+	}
+};
+
+/** A movie scene execution token that applies camera cuts */
+struct FApplyCameraAnimExecutionToken : IMovieSceneExecutionToken
+{
+	static FMovieSceneAnimTypeID GetAnimTypeID()
+	{
+		return TMovieSceneAnimTypeID<FApplyCameraAnimExecutionToken>();
+	}
+
+	/** Execute this token, operating on all objects referenced by 'Operand' */
+	virtual void Execute(const FMovieSceneContext& Context, const FMovieSceneEvaluationOperand& Operand, FPersistentEvaluationData& PersistentData, IMovieScenePlayer& Player) override
+	{
+		FMovieSceneAdditiveCameraData& SharedData = FMovieSceneAdditiveCameraData::Get(Operand, PersistentData);
+
+		for (TWeakObjectPtr<> ObjectWP : Player.FindBoundObjects(Operand))
+		{
+			UCameraComponent* CameraComponent = MovieSceneHelpers::CameraComponentFromRuntimeObject(ObjectWP.Get());
+			if (CameraComponent)
+			{
+				SharedData.ApplyCumulativeAnimation(*CameraComponent);
+			}
+		}
+	}
+};
+
+const FMovieSceneSharedDataId FMovieSceneAdditiveCameraAnimationTrackTemplate::SharedDataId = FMovieSceneSharedDataId::Allocate();
+
+FMovieSceneAdditiveCameraAnimationTrackTemplate::FMovieSceneAdditiveCameraAnimationTrackTemplate()
+{
+}
+
+void FMovieSceneAdditiveCameraAnimationTrackTemplate::Initialize(const FMovieSceneEvaluationOperand& Operand, const FMovieSceneContext& Context, FPersistentEvaluationData& PersistentData, IMovieScenePlayer& Player) const
+{
+	FMovieSceneAdditiveCameraData::Get(Operand, PersistentData).Reset();
+}
+
+void FMovieSceneAdditiveCameraAnimationTrackTemplate::Evaluate(const FMovieSceneEvaluationOperand& Operand, const FMovieSceneContext& Context, const FPersistentEvaluationData& PersistentData, FMovieSceneExecutionTokens& ExecutionTokens) const
+{
+	// Add a final execution token that will apply the blended result of anything added to our shared data
+	ExecutionTokens.Add(FApplyCameraAnimExecutionToken());
+}
+
+/** A movie scene execution token that applies camera animations */
+struct FAccumulateCameraAnimExecutionToken : IMovieSceneExecutionToken
+{
+	FAccumulateCameraAnimExecutionToken(const FMovieSceneAdditiveCameraAnimationTemplate* InTemplate) : Template(InTemplate) {}
+
+	virtual void Execute(const FMovieSceneContext& Context, const FMovieSceneEvaluationOperand& Operand, FPersistentEvaluationData& PersistentData, IMovieScenePlayer& Player) override
+	{
+		if (!Template->EnsureSetup(Operand, PersistentData, Player))
+		{
+			return;
+		}
+
+		FMovieSceneAdditiveCameraData& SharedData = FMovieSceneAdditiveCameraData::Get(Operand, PersistentData);
+
+		ACameraActor* TempCameraActor = SharedData.GetTempCameraActor(Player);
+		if (!ensure(TempCameraActor))
+		{
+			return;
+		}
+
+		for (TWeakObjectPtr<> ObjectWP : Player.FindBoundObjects(Operand))
+		{
+			UCameraComponent* CameraComponent = MovieSceneHelpers::CameraComponentFromRuntimeObject(ObjectWP.Get());
+			if (!CameraComponent)
+			{
+				continue;
+			}
+
+			FMinimalViewInfo POV;
+			POV.Location = CameraComponent->GetComponentLocation();
+			POV.Rotation = CameraComponent->GetComponentRotation();
+			POV.FOV = CameraComponent->FieldOfView;
+
+			if (!Template->UpdateCamera(*TempCameraActor, POV, Context, PersistentData))
+			{
+				continue;
+			}
+
+			FTransform WorldToBaseCamera = CameraComponent->GetComponentToWorld().Inverse();
+			float BaseFOV = CameraComponent->FieldOfView;
+			FTransform NewCameraToWorld(POV.Rotation, POV.Location);
+			float NewFOV = POV.FOV;
+
+			FTransform NewCameraToBaseCamera = NewCameraToWorld * WorldToBaseCamera;
+
+			float NewFOVToBaseFOV = BaseFOV - NewFOV;
+
+			{
+				static FMovieSceneAnimTypeID TransformAnimTypeID = TMovieSceneAnimTypeID<FAccumulateCameraAnimExecutionToken, 0>();
+				Player.SavePreAnimatedState(*CameraComponent, TransformAnimTypeID, FPreAnimatedCameraTransformTokenProducer());
+
+				// Accumumulate the offsets into the track data for application as part of the track execution token
+				SharedData.AccumulateOffset(NewCameraToBaseCamera, NewFOVToBaseFOV);
+			}
+
+			// harvest PP changes
+			{
+				UCameraComponent* AnimCamComp = TempCameraActor->GetCameraComponent();
+				if (AnimCamComp && AnimCamComp->PostProcessBlendWeight > 0.f)
+				{
+					static FMovieSceneAnimTypeID PostAnimTypeID = TMovieSceneAnimTypeID<FAccumulateCameraAnimExecutionToken, 1>();
+					Player.SavePreAnimatedState(*CameraComponent, PostAnimTypeID, FPreAnimatedPostProcessingBlendsTokenProducer());
+
+					SharedData.AccumulatePostProcessing(AnimCamComp->PostProcessSettings, AnimCamComp->PostProcessBlendWeight);
+				}
+			}
+		}
+	}
+
+	const FMovieSceneAdditiveCameraAnimationTemplate* Template;
+};
+
+void FMovieSceneAdditiveCameraAnimationTemplate::Evaluate(const FMovieSceneEvaluationOperand& Operand, const FMovieSceneContext& Context, const FPersistentEvaluationData& PersistentData, FMovieSceneExecutionTokens& ExecutionTokens) const
+{
+	ExecutionTokens.Add(FAccumulateCameraAnimExecutionToken(this));
+}
+
+/** Persistent data that exists as long as a given camera anim section is being evaluated */
+struct FMovieSceneCameraAnimSectionInstanceData : IPersistentEvaluationData
+{
+	TWeakObjectPtr<UCameraAnimInst> CameraAnimInst;
+};
+
+/** Pre animated token that restores a camera animation */
+struct FPreAnimatedCameraAnimTokenProducer : IMovieScenePreAnimatedTokenProducer
+{
+	virtual IMovieScenePreAnimatedTokenPtr CacheExistingState(UObject& Object) const override
+	{
+		struct FRestoreToken : IMovieScenePreAnimatedToken
+		{
+			virtual void RestoreState(UObject& InObject, IMovieScenePlayer& Player) override
+			{
+				UCameraAnimInst* CameraAnim = CastChecked<UCameraAnimInst>(&InObject);
+				CameraAnim->Stop(true);
+				CameraAnim->RemoveFromRoot();
+			}
+		};
+
+		return FRestoreToken();
+	}
+};
+
+FMovieSceneCameraAnimSectionTemplate::FMovieSceneCameraAnimSectionTemplate(const UMovieSceneCameraAnimSection& Section)
+	: SourceData(Section.AnimData)
+	, SectionStartTime(Section.GetStartTime())
+{
+}
+
+bool FMovieSceneCameraAnimSectionTemplate::EnsureSetup(const FMovieSceneEvaluationOperand& Operand, FPersistentEvaluationData& PersistentData, IMovieScenePlayer& Player) const
+{
+	// Get the camera anim instance from the section data (local to this specific section)
+	FMovieSceneCameraAnimSectionInstanceData& SectionData = PersistentData.GetOrAddSectionData<FMovieSceneCameraAnimSectionInstanceData>();
+	UCameraAnimInst* CameraAnimInstance = SectionData.CameraAnimInst.Get();
+
+	if (!CameraAnimInstance)
+	{
+		if (!SourceData.CameraAnim)
+		{
+			return false;
+		}
+
+		// Start playing the camera anim
+		CameraAnimInstance = NewObject<UCameraAnimInst>(GetTransientPackage());
+		if (ensure(CameraAnimInstance))
+		{
+			// make it root so GC doesn't take it away
+			CameraAnimInstance->AddToRoot();
+			CameraAnimInstance->SetStopAutomatically(false);
+
+			// Store the anim instance with the section and always remove it when we've finished evaluating
+			FMovieSceneAnimTypeID AnimTypeID = TMovieSceneAnimTypeID<FMovieSceneCameraAnimSectionTemplate>();
+			Player.SavePreAnimatedState(*CameraAnimInstance, AnimTypeID, FPreAnimatedCameraAnimTokenProducer(), PersistentData.GetSectionKey());
+
+			// We use the global temp actor from the shared data (shared across all additive camera effects for this operand)
+			ACameraActor* TempCameraActor = FMovieSceneAdditiveCameraData::Get(Operand, PersistentData).GetTempCameraActor(Player);
+
+			CameraAnimInstance->Play(SourceData.CameraAnim, TempCameraActor, SourceData.PlayRate, SourceData.PlayScale, SourceData.BlendInTime, SourceData.BlendOutTime, SourceData.bLooping, SourceData.bRandomStartTime);
+		}
+
+		SectionData.CameraAnimInst = CameraAnimInstance;
+	}
+
+	// If we failed to create the camera anim instance, we're doomed
+	return ensure(CameraAnimInstance);
+}
+
+bool FMovieSceneCameraAnimSectionTemplate::UpdateCamera(ACameraActor& TempCameraActor, FMinimalViewInfo& POV, const FMovieSceneContext& Context, FPersistentEvaluationData& PersistentData) const
+{
+	// Get the camera anim instance from the section data (local to this specific section)
+	FMovieSceneCameraAnimSectionInstanceData& SectionData = PersistentData.GetOrAddSectionData<FMovieSceneCameraAnimSectionInstanceData>();
+	UCameraAnimInst* CameraAnimInstance = SectionData.CameraAnimInst.Get();
+
+	if (!CameraAnimInstance || !CameraAnimInstance->CamAnim)
+	{
+		return false;
+	}
+
+	// prepare temp camera actor by resetting it
+	{
+		TempCameraActor.SetActorLocationAndRotation(FVector::ZeroVector, FRotator::ZeroRotator);
+
+		ACameraActor const* const DefaultCamActor = GetDefault<ACameraActor>();
+		if (DefaultCamActor)
+		{
+			TempCameraActor.GetCameraComponent()->AspectRatio = DefaultCamActor->GetCameraComponent()->AspectRatio;
+			TempCameraActor.GetCameraComponent()->PostProcessSettings = CameraAnimInstance->CamAnim->BasePostProcessSettings;
+			TempCameraActor.GetCameraComponent()->PostProcessBlendWeight = CameraAnimInstance->CamAnim->BasePostProcessBlendWeight;
+		}
+	}
+	
+	// set camera anim to the correct time
+	const float NewCameraAnimTime = Context.GetTime() - SectionStartTime;
+	CameraAnimInstance->SetCurrentTime(NewCameraAnimTime);
+
+	if (CameraAnimInstance->CurrentBlendWeight <= 0.f)
+	{
+		return false;
+	}
+
+	// harvest properties from the actor and apply
+	CameraAnimInstance->ApplyToView(POV);
+
+	return true;
+}
+
+/** Persistent data that exists as long as a given camera anim section is being evaluated */
+struct FMovieSceneCameraShakeSectionInstanceData : IPersistentEvaluationData
+{
+	TWeakObjectPtr<UCameraShake> CameraShakeInst;
+};
+
+/** Pre animated token that restores a camera animation */
+struct FPreAnimatedCameraShakeTokenProducer : IMovieScenePreAnimatedTokenProducer
+{
+	virtual IMovieScenePreAnimatedTokenPtr CacheExistingState(UObject& Object) const override
+	{
+		struct FRestoreToken : IMovieScenePreAnimatedToken
+		{
+			virtual void RestoreState(UObject& InObject, IMovieScenePlayer& Player) override
+			{
+				UCameraShake* CameraShake = CastChecked<UCameraShake>(&InObject);
+				CameraShake->StopShake(true);
+				CameraShake->RemoveFromRoot();
+			}
+		};
+
+		return FRestoreToken();
+	}
+};
+
+FMovieSceneCameraShakeSectionTemplate::FMovieSceneCameraShakeSectionTemplate(const UMovieSceneCameraShakeSection& Section)
+	: SourceData(Section.ShakeData)
+	, SectionStartTime(Section.GetStartTime())
+{
+}
+
+bool FMovieSceneCameraShakeSectionTemplate::EnsureSetup(const FMovieSceneEvaluationOperand& Operand, FPersistentEvaluationData& PersistentData, IMovieScenePlayer& Player) const
+{
+	// Get the camera anim instance from the section data (local to this specific section)
+	FMovieSceneCameraShakeSectionInstanceData& SectionData = PersistentData.GetOrAddSectionData<FMovieSceneCameraShakeSectionInstanceData>();
+	UCameraShake* CameraShakeInstance = SectionData.CameraShakeInst.Get();
+
+	if (!CameraShakeInstance)
+	{
+		if (!*SourceData.ShakeClass)
+		{
+			return false;
+		}
+
+		CameraShakeInstance = NewObject<UCameraShake>(GetTransientPackage(), SourceData.ShakeClass);
+		if (CameraShakeInstance)
+		{
+			// Store the anim instance with the section and always remove it when we've finished evaluating
+			FMovieSceneAnimTypeID AnimTypeID = TMovieSceneAnimTypeID<FMovieSceneCameraShakeSectionTemplate>();
+			Player.SavePreAnimatedState(*CameraShakeInstance, AnimTypeID, FPreAnimatedCameraShakeTokenProducer(), PersistentData.GetSectionKey());
+
+			// We use the global temp actor from the shared data (shared across all additive camera effects for this operand)
+			ACameraActor* TempCameraActor = FMovieSceneAdditiveCameraData::Get(Operand, PersistentData).GetTempCameraActor(Player);
+
+			// make it root so GC doesn't take it away
+			CameraShakeInstance->AddToRoot();
+			CameraShakeInstance->SetTempCameraAnimActor(TempCameraActor);
+			CameraShakeInstance->PlayShake(nullptr, SourceData.PlayScale, SourceData.PlaySpace, SourceData.UserDefinedPlaySpace);
+			if (CameraShakeInstance->AnimInst)
+			{
+				CameraShakeInstance->AnimInst->SetStopAutomatically(false);
+			}
+		}
+
+		SectionData.CameraShakeInst = CameraShakeInstance;
+	}
+
+	// If we failed to create the camera anim instance, we're doomed
+	return ensure(CameraShakeInstance);
+}
+
+bool FMovieSceneCameraShakeSectionTemplate::UpdateCamera(ACameraActor& TempCameraActor, FMinimalViewInfo& POV, const FMovieSceneContext& Context, FPersistentEvaluationData& PersistentData) const
+{
+	// Get the camera anim instance from the section data (local to this specific section)
+	FMovieSceneCameraShakeSectionInstanceData& SectionData = PersistentData.GetOrAddSectionData<FMovieSceneCameraShakeSectionInstanceData>();
+	UCameraShake* CameraShakeInstance = SectionData.CameraShakeInst.Get();
+
+	if (!ensure(CameraShakeInstance))
+	{
+		return false;
+	}
+
+	// prepare temp camera actor by resetting it
+	{
+		TempCameraActor.SetActorLocationAndRotation(FVector::ZeroVector, FRotator::ZeroRotator);
+
+		ACameraActor const* const DefaultCamActor = GetDefault<ACameraActor>();
+		if (DefaultCamActor)
+		{
+			TempCameraActor.GetCameraComponent()->AspectRatio = DefaultCamActor->GetCameraComponent()->AspectRatio;
+
+			UCameraAnim* CamAnim = CameraShakeInstance->AnimInst ? CameraShakeInstance->AnimInst->CamAnim : nullptr;
+
+			TempCameraActor.GetCameraComponent()->PostProcessSettings = CamAnim ? CamAnim->BasePostProcessSettings : FPostProcessSettings();
+			TempCameraActor.GetCameraComponent()->PostProcessBlendWeight = CamAnim ? CameraShakeInstance->AnimInst->CamAnim->BasePostProcessBlendWeight : 0.f;
+		}
+	}
+	
+	// set camera anim to the correct time
+	const float NewCameraAnimTime = Context.GetTime() - SectionStartTime;
+	CameraShakeInstance->SetCurrentTimeAndApplyShake(NewCameraAnimTime, POV);
+
+	return true;
+}
\ No newline at end of file
diff --git a/Engine/Source/Runtime/MovieSceneTracks/Private/Evaluation/MovieSceneCameraAnimTemplate.h b/Engine/Source/Runtime/MovieSceneTracks/Private/Evaluation/MovieSceneCameraAnimTemplate.h
new file mode 100644
--- /dev/null
+++ b/Engine/Source/Runtime/MovieSceneTracks/Private/Evaluation/MovieSceneCameraAnimTemplate.h
@@ -0,0 +1,97 @@
+// Copyright 1998-2016 Epic Games, Inc. All Rights Reserved.
+
+#pragma once
+
+#include "MovieSceneFwd.h"
+
+#include "MovieSceneEvalTemplate.h"
+#include "MovieSceneCameraAnimSection.h"
+#include "MovieSceneCameraShakeSection.h"
+
+#include "MovieSceneCameraAnimTemplate.generated.h"
+
+struct FMovieSceneAdditiveCameraData;
+
+/** Generic section template for any additive camera animation effects */
+USTRUCT()
+struct FMovieSceneAdditiveCameraAnimationTemplate : public FMovieSceneEvalTemplate
+{
+	GENERATED_BODY()
+
+	virtual bool EnsureSetup(const FMovieSceneEvaluationOperand& Operand, FPersistentEvaluationData& PersistentData, IMovieScenePlayer& Player) const { return false; }
+	virtual bool UpdateCamera(ACameraActor& TempCameraActor, FMinimalViewInfo& POV, const FMovieSceneContext& Context, FPersistentEvaluationData& PersistentData) const { return false; }
+
+private:
+
+	virtual UScriptStruct& GetScriptStructImpl() const override { check(false); return *StaticStruct(); }
+	virtual void Evaluate(const FMovieSceneEvaluationOperand& Operand, const FMovieSceneContext& Context, const FPersistentEvaluationData& PersistentData, FMovieSceneExecutionTokens& ExecutionTokens) const override;
+};
+
+/** Section template for a camera anim */
+USTRUCT()
+struct FMovieSceneCameraAnimSectionTemplate : public FMovieSceneAdditiveCameraAnimationTemplate
+{
+	GENERATED_BODY()
+	
+	FMovieSceneCameraAnimSectionTemplate() {}
+	FMovieSceneCameraAnimSectionTemplate(const UMovieSceneCameraAnimSection& Section);
+
+private:
+	virtual UScriptStruct& GetScriptStructImpl() const override { return *StaticStruct(); }
+
+	virtual bool EnsureSetup(const FMovieSceneEvaluationOperand& Operand, FPersistentEvaluationData& PersistentData, IMovieScenePlayer& Player) const override;
+	virtual bool UpdateCamera(ACameraActor& TempCameraActor, FMinimalViewInfo& POV, const FMovieSceneContext& Context, FPersistentEvaluationData& PersistentData) const override;
+
+	/** Source data taken from the section */
+	UPROPERTY()
+	FMovieSceneCameraAnimSectionData SourceData;
+
+	/** Cached section start time */
+	UPROPERTY()
+	float SectionStartTime;
+};
+
+/** Section template for a camera anim */
+USTRUCT()
+struct FMovieSceneCameraShakeSectionTemplate : public FMovieSceneAdditiveCameraAnimationTemplate
+{
+	GENERATED_BODY()
+	
+	FMovieSceneCameraShakeSectionTemplate() {}
+	FMovieSceneCameraShakeSectionTemplate(const UMovieSceneCameraShakeSection& Section);
+
+private:
+	virtual UScriptStruct& GetScriptStructImpl() const override { return *StaticStruct(); }
+
+	virtual bool EnsureSetup(const FMovieSceneEvaluationOperand& Operand, FPersistentEvaluationData& PersistentData, IMovieScenePlayer& Player) const override;
+	virtual bool UpdateCamera(ACameraActor& TempCameraActor, FMinimalViewInfo& POV, const FMovieSceneContext& Context, FPersistentEvaluationData& PersistentData) const override;
+
+
+	/** Source data taken from the section */
+	UPROPERTY()
+	FMovieSceneCameraShakeSectionData SourceData;
+
+	/** Cached section start time */
+	UPROPERTY()
+	float SectionStartTime;
+};
+
+
+USTRUCT()
+struct FMovieSceneAdditiveCameraAnimationTrackTemplate : public FMovieSceneEvalTemplate
+{
+	GENERATED_BODY()
+
+	FMovieSceneAdditiveCameraAnimationTrackTemplate();
+
+	static const FMovieSceneSharedDataId SharedDataId;
+
+	virtual UScriptStruct& GetScriptStructImpl() const override { return *StaticStruct(); }
+	virtual void SetupOverrides() override
+	{
+		EnableOverrides(RequiresInitializeFlag);
+	}
+	
+	virtual void Initialize(const FMovieSceneEvaluationOperand& Operand, const FMovieSceneContext& Context, FPersistentEvaluationData& PersistentData, IMovieScenePlayer& Player) const override;
+	virtual void Evaluate(const FMovieSceneEvaluationOperand& Operand, const FMovieSceneContext& Context, const FPersistentEvaluationData& PersistentData, FMovieSceneExecutionTokens& ExecutionTokens) const override;
+};
\ No newline at end of file
diff --git a/Engine/Source/Runtime/MovieSceneTracks/Private/Evaluation/MovieSceneCameraCutTemplate.cpp b/Engine/Source/Runtime/MovieSceneTracks/Private/Evaluation/MovieSceneCameraCutTemplate.cpp
new file mode 100644
--- /dev/null
+++ b/Engine/Source/Runtime/MovieSceneTracks/Private/Evaluation/MovieSceneCameraCutTemplate.cpp
@@ -0,0 +1,85 @@
+// Copyright 1998-2016 Epic Games, Inc. All Rights Reserved.
+
+#include "MovieSceneTracksPrivatePCH.h"
+#include "MovieSceneCameraCutTemplate.h"
+#include "MovieSceneCameraCutSection.h"
+
+DECLARE_CYCLE_STAT(TEXT("Camera Cut Track Token Execute"), MovieSceneEval_CameraCutTrack_TokenExecute, STATGROUP_MovieSceneEval);
+
+struct FCameraCutTrackData : IPersistentEvaluationData
+{
+	TWeakObjectPtr<> LastLockedCamera;
+};
+
+/** A movie scene pre-animated token that stores a pre-animated camera cut */
+struct FCameraCutPreAnimatedToken : IMovieScenePreAnimatedGlobalToken
+{
+	virtual void RestoreState(IMovieScenePlayer& Player) override
+	{
+		Player.UpdateCameraCut(nullptr, nullptr);
+	}
+};
+
+/** A movie scene execution token that applies camera cuts */
+struct FCameraCutExecutionToken : IMovieSceneExecutionToken
+{
+	FGuid CameraGuid;
+
+	FCameraCutExecutionToken(const FGuid& InCameraGuid)
+		: CameraGuid(InCameraGuid)
+	{}
+
+	static FMovieSceneAnimTypeID GetAnimTypeID()
+	{
+		return TMovieSceneAnimTypeID<FCameraCutExecutionToken>();
+	}
+	
+	/** Execute this token, operating on all objects referenced by 'Operand' */
+	virtual void Execute(const FMovieSceneContext& Context, const FMovieSceneEvaluationOperand& Operand, FPersistentEvaluationData& PersistentData, IMovieScenePlayer& Player) override
+	{
+		MOVIESCENE_DETAILED_SCOPE_CYCLE_COUNTER(MovieSceneEval_CameraCutTrack_TokenExecute)
+
+		FMovieSceneEvaluationOperand CameraOperand(Operand.SequenceID, CameraGuid);
+		
+		TArrayView<TWeakObjectPtr<>> Objects = Player.FindBoundObjects(CameraOperand);
+		if (!Objects.Num())
+		{
+			return;
+		}
+
+		// Only ever deal with one camera
+		UObject* CameraObject = Objects[0].Get();
+
+		struct FProducer : IMovieScenePreAnimatedGlobalTokenProducer
+		{
+			virtual IMovieScenePreAnimatedGlobalTokenPtr CacheExistingState() const override
+			{
+				return FCameraCutPreAnimatedToken();
+			}
+		} Producer;
+
+		Player.SavePreAnimatedState(GetAnimTypeID(), Producer);
+
+		FCameraCutTrackData& CameraCutCache = PersistentData.GetOrAddTrackData<FCameraCutTrackData>();
+		if (CameraCutCache.LastLockedCamera.Get() != CameraObject)
+		{
+			Player.UpdateCameraCut(CameraObject, CameraCutCache.LastLockedCamera.Get(), Context.HasJumped());
+			CameraCutCache.LastLockedCamera = CameraObject;
+		}
+		else if (CameraObject)
+		{
+			Player.UpdateCameraCut(CameraObject, nullptr, Context.HasJumped());
+		}
+	}
+};
+
+
+FMovieSceneCameraCutSectionTemplate::FMovieSceneCameraCutSectionTemplate(const UMovieSceneCameraCutSection& Section)
+	: CameraGuid(Section.GetCameraGuid())
+{
+}
+
+void FMovieSceneCameraCutSectionTemplate::Evaluate(const FMovieSceneEvaluationOperand& Operand, const FMovieSceneContext& Context, const FPersistentEvaluationData& PersistentData, FMovieSceneExecutionTokens& ExecutionTokens) const
+{
+	ExecutionTokens.Add(FCameraCutExecutionToken(CameraGuid));
+}
\ No newline at end of file
diff --git a/Engine/Source/Runtime/MovieSceneTracks/Private/Evaluation/MovieSceneCameraCutTemplate.h b/Engine/Source/Runtime/MovieSceneTracks/Private/Evaluation/MovieSceneCameraCutTemplate.h
new file mode 100644
--- /dev/null
+++ b/Engine/Source/Runtime/MovieSceneTracks/Private/Evaluation/MovieSceneCameraCutTemplate.h
@@ -0,0 +1,29 @@
+// Copyright 1998-2016 Epic Games, Inc. All Rights Reserved.
+
+#pragma once
+
+#include "MovieSceneFwd.h"
+#include "MovieSceneTrackImplementation.h"
+#include "MovieSceneEvalTemplate.h"
+#include "MovieSceneCameraCutTemplate.generated.h"
+
+class UMovieSceneCameraCutSection;
+
+/** Camera cut track evaluation template */
+USTRUCT()
+struct FMovieSceneCameraCutSectionTemplate : public FMovieSceneEvalTemplate
+{
+	GENERATED_BODY()
+	
+	FMovieSceneCameraCutSectionTemplate() {}
+	FMovieSceneCameraCutSectionTemplate(const UMovieSceneCameraCutSection& Section);
+
+	/** GUID of the camera we should cut to in this sequence */
+	UPROPERTY()
+	FGuid CameraGuid;
+
+private:
+
+	virtual UScriptStruct& GetScriptStructImpl() const override { return *StaticStruct(); }
+	virtual void Evaluate(const FMovieSceneEvaluationOperand& Operand, const FMovieSceneContext& Context, const FPersistentEvaluationData& PersistentData, FMovieSceneExecutionTokens& ExecutionTokens) const override;
+};
\ No newline at end of file
diff --git a/Engine/Source/Runtime/MovieSceneTracks/Private/Evaluation/MovieSceneColorTemplate.cpp b/Engine/Source/Runtime/MovieSceneTracks/Private/Evaluation/MovieSceneColorTemplate.cpp
new file mode 100644
--- /dev/null
+++ b/Engine/Source/Runtime/MovieSceneTracks/Private/Evaluation/MovieSceneColorTemplate.cpp
@@ -0,0 +1,235 @@
+// Copyright 1998-2016 Epic Games, Inc. All Rights Reserved.
+
+#include "MovieSceneTracksPrivatePCH.h"
+#include "MovieSceneColorSection.h"
+#include "MovieSceneColorTrack.h"
+#include "MovieSceneColorTemplate.h"
+
+DECLARE_CYCLE_STAT(TEXT("Color Track Token Execute"), MovieSceneEval_ColorTrack_TokenExecute, STATGROUP_MovieSceneEval);
+
+enum class EColorType : uint8
+{
+	/** FSlateColor */
+	Slate, 
+	/** FLinearColor */
+	Linear,
+	/** FColor */
+	Color,
+};
+
+struct FColorSectionData : IPersistentEvaluationData
+{
+	TSharedPtr<FTrackInstancePropertyBindings> PropertyBindings;
+	FLinearColor DefaultColor;
+};
+
+struct FColorToken
+{
+	FColorToken() {}
+	FColorToken(FLinearColor InColorValue, EColorType InType) : ColorValue(InColorValue), Type(InType) {}
+
+	void Apply(UObject& Object, FTrackInstancePropertyBindings& Bindings)
+	{
+		switch(Type)
+		{
+		case EColorType::Slate:		ApplySlateColor(Object, Bindings);		break;
+		case EColorType::Linear: 	ApplyLinearColor(Object, Bindings);		break;
+		case EColorType::Color: 	ApplyColor(Object, Bindings);			break;
+		}
+	}
+
+	void ApplyColor(UObject& Object, FTrackInstancePropertyBindings& Bindings)
+	{
+		const bool bConvertBackToSRgb = true;
+		if (ULightComponent* LightComponent = Cast<ULightComponent>(&Object))
+		{
+			// Light components have to be handled specially here because their set function takes two values, the linear color
+			// and whether or not the linear color needs to be converted back to sRGB.  All other other set function cases should
+			// follow the sequencer convention of having a single parameter of the correct type, which in this case is an FColor
+			// already in sRGB format.
+			LightComponent->SetLightColor( ColorValue, bConvertBackToSRgb );
+		}
+		else
+		{
+			FColor SRgbColorValue = ColorValue.ToFColor( bConvertBackToSRgb );
+			Bindings.CallFunction<FColor>( Object, SRgbColorValue );
+		}
+	}
+
+	void ApplySlateColor(UObject& Object, FTrackInstancePropertyBindings& Bindings)
+	{
+		FSlateColor NewColor(ColorValue);
+		Bindings.CallFunction<FSlateColor>(Object, NewColor);
+	}
+
+	void ApplyLinearColor(UObject& Object, FTrackInstancePropertyBindings& Bindings)
+	{
+		Bindings.CallFunction<FLinearColor>(Object, ColorValue);
+	}
+
+	void DeduceColorType(const UObject& InObject, FTrackInstancePropertyBindings& Bindings)
+	{
+		const UStructProperty* StructProp = Cast<const UStructProperty>(Bindings.GetProperty(InObject));
+		if (StructProp && StructProp->Struct)
+		{
+			FName StructName = StructProp->Struct->GetFName();
+
+			static const FName SlateColor("SlateColor");
+			if( StructName == NAME_Color )
+			{
+				Type = EColorType::Color;
+			}
+			else if( StructName == SlateColor )
+			{
+				Type = EColorType::Slate;
+			}
+			else
+			{
+				Type = EColorType::Linear;
+			}
+		}
+	}
+
+	void AssignFromCurrentValue(const UObject& InObject, FTrackInstancePropertyBindings& Bindings)
+	{
+		switch (Type)
+		{
+		case EColorType::Color:
+			// We assume the color we get back is in sRGB, assigning it to a linear color will implicitly
+			// convert it to a linear color instead of using ReinterpretAsLinear which will just change the
+			// bytes into floats using divide by 255.
+			ColorValue = Bindings.GetCurrentValue<FColor>(InObject);
+			break;
+
+		case EColorType::Slate:
+			ColorValue = Bindings.GetCurrentValue<FSlateColor>(InObject).GetSpecifiedColor();
+			break;
+
+		case EColorType::Linear:
+			ColorValue = Bindings.GetCurrentValue<FLinearColor>(InObject);
+			break;
+		}
+	}
+
+	FLinearColor ColorValue;
+	EColorType Type;
+};
+
+struct FColorTrackPreAnimatedState : IMovieScenePreAnimatedToken
+{
+	FColorTrackPreAnimatedState(FColorToken InToken, const FTrackInstancePropertyBindings& InBindings)
+		: Token(InToken)
+		, Bindings(InBindings)
+	{
+	}
+
+	virtual void RestoreState(UObject& Object, IMovieScenePlayer& Player) override
+	{
+		Token.Apply(Object, Bindings);
+	}
+
+	FColorToken Token;
+	FTrackInstancePropertyBindings Bindings;
+};
+
+struct FColorTokenProducer : IMovieScenePreAnimatedTokenProducer
+{
+	FColorTokenProducer(const FColorToken& InSourceToken, FTrackInstancePropertyBindings& InPropertyBindings)
+		: SourceToken(InSourceToken), PropertyBindings(InPropertyBindings)
+	{}
+
+	virtual IMovieScenePreAnimatedTokenPtr CacheExistingState(UObject& Object) const
+	{
+		FColorToken TokenCopy = SourceToken;
+		TokenCopy.AssignFromCurrentValue(Object, PropertyBindings);
+		return FColorTrackPreAnimatedState(TokenCopy, PropertyBindings);
+	}
+
+	const FColorToken& SourceToken;
+	FTrackInstancePropertyBindings& PropertyBindings;
+};
+
+struct FColorTrackExecutionToken : IMovieSceneExecutionToken
+{
+	FColorTrackExecutionToken(FLinearColor InColor)
+		: Token(InColor, EColorType::Linear)	// Just initialize to linear color for now, we'll deduce the type later
+	{
+	}
+
+	/** Execute this token, operating on all objects referenced by 'Operand' */
+	virtual void Execute(const FMovieSceneContext& Context, const FMovieSceneEvaluationOperand& Operand, FPersistentEvaluationData& SectionData, IMovieScenePlayer& Player) override
+	{
+		MOVIESCENE_DETAILED_SCOPE_CYCLE_COUNTER(MovieSceneEval_ColorTrack_TokenExecute)
+		
+		FTrackInstancePropertyBindings* PropertyBindings = SectionData.GetSectionData<FColorSectionData>().PropertyBindings.Get();
+
+		for (TWeakObjectPtr<> Object : Player.FindBoundObjects(Operand))
+		{
+			if (UObject* ObjectPtr = Object.Get())
+			{
+				// Deduce the color type before we save pre animated state (so we store the right type)
+				Token.DeduceColorType(*ObjectPtr, *PropertyBindings);
+
+				FMovieSceneAnimTypeID AnimType = TMovieSceneAnimTypeID<FColorTrackExecutionToken>();
+				Player.SavePreAnimatedState(*ObjectPtr, AnimType, FColorTokenProducer(Token, *PropertyBindings));
+
+				Token.Apply(*ObjectPtr, *PropertyBindings);
+			}
+		}
+	}
+
+	FColorToken Token;
+};
+
+FMovieSceneColorSectionTemplate::FMovieSceneColorSectionTemplate(const UMovieSceneColorSection& Section, const UMovieSceneColorTrack& Track)
+	: PropertyName(Track.GetPropertyName())
+	, PropertyPath(Track.GetPropertyPath())
+{
+	Curves[0] = Section.GetRedCurve();
+	Curves[1] = Section.GetGreenCurve();
+	Curves[2] = Section.GetBlueCurve();
+	Curves[3] = Section.GetAlphaCurve();
+}
+
+void FMovieSceneColorSectionTemplate::Setup(FPersistentEvaluationData& PersistentData, IMovieScenePlayer& Player) const
+{
+	PersistentData.AddSectionData<FColorSectionData>().PropertyBindings = MakeShareable(new FTrackInstancePropertyBindings(PropertyName, PropertyPath));
+}
+
+void FMovieSceneColorSectionTemplate::Initialize(const FMovieSceneEvaluationOperand& Operand, const FMovieSceneContext& Context, FPersistentEvaluationData& PersistentData, IMovieScenePlayer& Player) const
+{
+	FColorSectionData& ColorSectionData = PersistentData.GetSectionData<FColorSectionData>();
+
+	// Cache off the default color from the first object.
+	// To support multiple relative animations, we'd need to animate each color relatively
+	for (TWeakObjectPtr<> Object : Player.FindBoundObjects(Operand))
+	{
+		if (UObject* ObjectPtr = Object.Get())
+		{
+			FColorToken TempToken;
+			TempToken.DeduceColorType(*ObjectPtr, *ColorSectionData.PropertyBindings);
+			TempToken.AssignFromCurrentValue(*ObjectPtr, *ColorSectionData.PropertyBindings);
+			ColorSectionData.DefaultColor = TempToken.ColorValue;
+			break;
+		}
+	}
+
+	
+}
+
+void FMovieSceneColorSectionTemplate::Evaluate(const FMovieSceneEvaluationOperand& Operand, const FMovieSceneContext& Context, const FPersistentEvaluationData& PersistentData, FMovieSceneExecutionTokens& ExecutionTokens) const
+{
+	float Time = Context.GetTime();
+	FLinearColor DefaultColor = PersistentData.GetSectionData<FColorSectionData>().DefaultColor;
+
+	ExecutionTokens.Add(
+		FColorTrackExecutionToken(
+			FLinearColor(
+				Curves[0].Eval(Time, DefaultColor.R),
+				Curves[1].Eval(Time, DefaultColor.G),
+				Curves[2].Eval(Time, DefaultColor.B),
+				Curves[3].Eval(Time, DefaultColor.A)
+			)
+		)
+	);
+}
\ No newline at end of file
diff --git a/Engine/Source/Runtime/MovieSceneTracks/Private/Evaluation/MovieSceneColorTemplate.h b/Engine/Source/Runtime/MovieSceneTracks/Private/Evaluation/MovieSceneColorTemplate.h
new file mode 100644
--- /dev/null
+++ b/Engine/Source/Runtime/MovieSceneTracks/Private/Evaluation/MovieSceneColorTemplate.h
@@ -0,0 +1,44 @@
+// Copyright 1998-2016 Epic Games, Inc. All Rights Reserved.
+
+#pragma once
+
+#include "MovieSceneFwd.h"
+#include "Curves/RichCurve.h"
+#include "MovieSceneEvalTemplate.h"
+
+#include "MovieSceneColorTemplate.generated.h"
+
+class UMovieSceneColorSection;
+class UMovieSceneColorTrack;
+
+USTRUCT()
+struct FMovieSceneColorSectionTemplate : public FMovieSceneEvalTemplate
+{
+	GENERATED_BODY()
+	
+	FMovieSceneColorSectionTemplate() {}
+	FMovieSceneColorSectionTemplate(const UMovieSceneColorSection& Section, const UMovieSceneColorTrack& Track);
+
+	UPROPERTY()
+	FName PropertyName;
+
+	UPROPERTY()
+	FString PropertyPath;
+
+	/** Curve data as RGBA */
+	UPROPERTY()
+	FRichCurve Curves[4];
+
+private:
+
+	virtual UScriptStruct& GetScriptStructImpl() const override { return *StaticStruct(); }
+	virtual void SetupOverrides() override
+	{
+		EnableOverrides(RequiresSetupFlag | RequiresInitializeFlag);
+	}
+	
+	virtual void Setup(FPersistentEvaluationData& PersistentData, IMovieScenePlayer& Player) const override;
+	virtual void Initialize(const FMovieSceneEvaluationOperand& Operand, const FMovieSceneContext& Context, FPersistentEvaluationData& PersistentData, IMovieScenePlayer& Player) const override;
+
+	virtual void Evaluate(const FMovieSceneEvaluationOperand& Operand, const FMovieSceneContext& Context, const FPersistentEvaluationData& PersistentData, FMovieSceneExecutionTokens& ExecutionTokens) const override;
+};
\ No newline at end of file
diff --git a/Engine/Source/Runtime/MovieSceneTracks/Private/Evaluation/MovieSceneEventTemplate.cpp b/Engine/Source/Runtime/MovieSceneTracks/Private/Evaluation/MovieSceneEventTemplate.cpp
new file mode 100644
--- /dev/null
+++ b/Engine/Source/Runtime/MovieSceneTracks/Private/Evaluation/MovieSceneEventTemplate.cpp
@@ -0,0 +1,191 @@
+// Copyright 1998-2016 Epic Games, Inc. All Rights Reserved.
+
+#include "MovieSceneTracksPrivatePCH.h"
+#include "MovieSceneEventTemplate.h"
+#include "MovieSceneEventTrack.h"
+#include "MovieSceneEventSection.h"
+#include "MovieSceneSequence.h"
+#include "MovieSceneEvaluationTemplateInstance.h"
+
+DECLARE_CYCLE_STAT(TEXT("Event Track Token Execute"), MovieSceneEval_EventTrack_TokenExecute, STATGROUP_MovieSceneEval);
+
+struct FEventData
+{
+	FEventData(const FEventPayload& InPayload, float InGlobalPosition) : Payload(InPayload), GlobalPosition(InGlobalPosition) {}
+
+	FEventPayload Payload;
+	float GlobalPosition;
+};
+
+/** A movie scene execution token that stores a specific transform, and an operand */
+struct FEventTrackExecutionToken
+	: IMovieSceneExecutionToken
+{
+	FEventTrackExecutionToken(TArray<FEventData> InEvents) : Events(MoveTemp(InEvents)) {}
+
+	/** Execute this token, operating on all objects referenced by 'Operand' */
+	virtual void Execute(const FMovieSceneContext& Context, const FMovieSceneEvaluationOperand& Operand, FPersistentEvaluationData& PersistentData, IMovieScenePlayer& Player) override
+	{
+		MOVIESCENE_DETAILED_SCOPE_CYCLE_COUNTER(MovieSceneEval_EventTrack_TokenExecute)
+		
+		TArray<float> PerformanceCaptureEventPositions;
+
+		for (UObject* EventContextObject : Player.GetEventContexts())
+		{
+			if (!EventContextObject)
+			{
+				continue;
+			}
+
+			for (FEventData& Event : Events)
+			{
+#if !UE_BUILD_SHIPPING
+				if (Event.Payload.EventName == NAME_PerformanceCapture)
+				{
+					PerformanceCaptureEventPositions.Add(Event.GlobalPosition);
+				}
+#endif
+				TriggerEvent(Event, *EventContextObject, Player);
+			}
+		}
+
+#if !UE_BUILD_SHIPPING
+		UWorld* World = Player.GetPlaybackContext()->GetWorld();
+		if (World)
+		{
+			FString LevelSequenceName = Player.GetEvaluationTemplate().GetSequence(MovieSceneSequenceID::Root)->GetName();
+			
+			for (float EventPosition : PerformanceCaptureEventPositions)
+			{
+				GEngine->PerformanceCapture(World, World->GetName(), LevelSequenceName, EventPosition);
+			}
+		}
+#endif	// UE_BUILD_SHIPPING
+	}
+
+	void TriggerEvent(FEventData& Event, UObject& EventContextObject, IMovieScenePlayer& Player)
+	{
+		UFunction* EventFunction = EventContextObject.FindFunction(Event.Payload.EventName);
+
+		if (EventFunction == nullptr)
+		{
+			// Don't want to log out a warning for every event context.
+			return;
+		}
+		else
+		{
+			FStructOnScope ParameterStruct(nullptr);
+			Event.Payload.Parameters.GetInstance(ParameterStruct);
+
+			uint8* Parameters = ParameterStruct.GetStructMemory();
+
+			const UStruct* Struct = ParameterStruct.GetStruct();
+			if (EventFunction->ReturnValueOffset != MAX_uint16)
+			{
+				UE_LOG(LogMovieScene, Warning, TEXT("Sequencer Event Track: Cannot trigger events that return values (for event '%s')."), *Event.Payload.EventName.ToString());
+				return;
+			}
+			else
+			{
+				TFieldIterator<UProperty> ParamIt(EventFunction);
+				TFieldIterator<UProperty> ParamInstanceIt(Struct);
+				for (int32 NumParams = 0; ParamIt || ParamInstanceIt; ++NumParams, ++ParamIt, ++ParamInstanceIt)
+				{
+					if (!ParamInstanceIt)
+					{
+						UE_LOG(LogMovieScene, Warning, TEXT("Sequencer Event Track: Parameter count mistatch for event '%s'. Required parameter of type '%s' at index '%d'."), *Event.Payload.EventName.ToString(), *ParamIt->GetName(), NumParams);
+						return;
+					}
+					else if (!ParamIt)
+					{
+						// Mismatch (too many params)
+						UE_LOG(LogMovieScene, Warning, TEXT("Sequencer Event Track: Parameter count mistatch for event '%s'. Parameter struct contains too many parameters ('%s' is superfluous at index '%d'."), *Event.Payload.EventName.ToString(), *ParamInstanceIt->GetName(), NumParams);
+						return;
+					}
+					else if (!ParamInstanceIt->SameType(*ParamIt) || ParamInstanceIt->GetOffset_ForUFunction() != ParamIt->GetOffset_ForUFunction() || ParamInstanceIt->GetSize() != ParamIt->GetSize())
+					{
+						UE_LOG(LogMovieScene, Warning, TEXT("Sequencer Event Track: Parameter type mistatch for event '%s' ('%s' != '%s')."),
+							*Event.Payload.EventName.ToString(),
+							*ParamInstanceIt->GetClass()->GetName(),
+							*ParamIt->GetClass()->GetName()
+						);
+						return;
+					}
+				}
+			}
+
+			// Technically, anything bound to the event could mutate the parameter payload,
+			// but we're going to treat that as misuse, rather than copy the parameters each time
+			EventContextObject.ProcessEvent(EventFunction, Parameters);
+		}
+	}
+
+	TArray<FEventData> Events;
+};
+
+FMovieSceneEventSectionTemplate::FMovieSceneEventSectionTemplate(const UMovieSceneEventSection& Section, const UMovieSceneEventTrack& Track)
+	: EventData(Section.GetEventData())
+	, bFireEventsWhenForwards(Track.bFireEventsWhenForwards)
+	, bFireEventsWhenBackwards(Track.bFireEventsWhenBackwards)
+{
+}
+
+void FMovieSceneEventSectionTemplate::EvaluateSwept(const FMovieSceneEvaluationOperand& Operand, const FMovieSceneContext& Context, const FPersistentEvaluationData& PersistentData, FMovieSceneExecutionTokens& ExecutionTokens) const
+{
+	// Don't allow events to fire when playback is in a stopped state. This can occur when stopping 
+	// playback and returning the current position to the start of playback. It's not desireable to have 
+	// all the events from the last playback position to the start of playback be fired.
+	if (Context.GetStatus() == EMovieScenePlayerStatus::Stopped || Context.IsSilent())
+	{
+		return;
+	}
+
+	const bool bBackwards = Context.GetDirection() == EPlayDirection::Backwards;
+
+	if ((!bBackwards && !bFireEventsWhenForwards) ||
+		(bBackwards && !bFireEventsWhenBackwards))
+	{
+		return;
+	}
+
+	TArray<FEventData> Events;
+
+	TRange<float> SweptRange = Context.GetRange();
+
+	const TArray<float>& KeyTimes = EventData.KeyTimes;
+	const TArray<FEventPayload>& KeyValues = EventData.KeyValues;
+
+	const int32 First = bBackwards ? KeyTimes.Num() - 1 : 0;
+	const int32 Last = bBackwards ? 0 : KeyTimes.Num() - 1;
+	const int32 Inc = bBackwards ? -1 : 1;
+
+	const float Position = Context.GetTime() * Context.GetRootToSequenceTransform().Inverse();
+
+	if (bBackwards)
+	{
+		// Trigger events backwards
+		for (int32 KeyIndex = KeyTimes.Num() - 1; KeyIndex >= 0; --KeyIndex)
+		{
+			float Time = KeyTimes[KeyIndex];
+			if (SweptRange.Contains(Time))
+			{
+				Events.Add(FEventData(KeyValues[KeyIndex], Position));
+			}
+		}
+	}
+	// Trigger events forwards
+	else for (int32 KeyIndex = 0; KeyIndex < KeyTimes.Num(); ++KeyIndex)
+	{
+		float Time = KeyTimes[KeyIndex];
+		if (SweptRange.Contains(Time))
+		{
+			Events.Add(FEventData(KeyValues[KeyIndex], Position));
+		}
+	}
+
+
+	if (Events.Num())
+	{
+		ExecutionTokens.Add(FEventTrackExecutionToken(MoveTemp(Events)));
+	}
+}
\ No newline at end of file
diff --git a/Engine/Source/Runtime/MovieSceneTracks/Private/Evaluation/MovieSceneEventTemplate.h b/Engine/Source/Runtime/MovieSceneTracks/Private/Evaluation/MovieSceneEventTemplate.h
new file mode 100644
--- /dev/null
+++ b/Engine/Source/Runtime/MovieSceneTracks/Private/Evaluation/MovieSceneEventTemplate.h
@@ -0,0 +1,34 @@
+// Copyright 1998-2016 Epic Games, Inc. All Rights Reserved.
+
+#pragma once
+
+#include "MovieSceneFwd.h"
+#include "MovieSceneEventSection.h"
+#include "MovieSceneEvalTemplate.h"
+
+#include "MovieSceneEventTemplate.generated.h"
+
+class UMovieSceneEventTrack;
+
+USTRUCT()
+struct FMovieSceneEventSectionTemplate : public FMovieSceneEvalTemplate
+{
+	GENERATED_BODY()
+	
+	FMovieSceneEventSectionTemplate() {}
+	FMovieSceneEventSectionTemplate(const UMovieSceneEventSection& Section, const UMovieSceneEventTrack& Track);
+
+	UPROPERTY()
+	FMovieSceneEventSectionData EventData;
+
+	UPROPERTY()
+	uint32 bFireEventsWhenForwards : 1;
+
+	UPROPERTY()
+	uint32 bFireEventsWhenBackwards : 1;
+
+private:
+
+	virtual UScriptStruct& GetScriptStructImpl() const override { return *StaticStruct(); }
+	virtual void EvaluateSwept(const FMovieSceneEvaluationOperand& Operand, const FMovieSceneContext& Context, const FPersistentEvaluationData& PersistentData, FMovieSceneExecutionTokens& ExecutionTokens) const override;
+};
\ No newline at end of file
diff --git a/Engine/Source/Runtime/MovieSceneTracks/Private/Evaluation/MovieSceneFadeTemplate.cpp b/Engine/Source/Runtime/MovieSceneTracks/Private/Evaluation/MovieSceneFadeTemplate.cpp
new file mode 100644
--- /dev/null
+++ b/Engine/Source/Runtime/MovieSceneTracks/Private/Evaluation/MovieSceneFadeTemplate.cpp
@@ -0,0 +1,113 @@
+// Copyright 1998-2016 Epic Games, Inc. All Rights Reserved.
+
+#include "MovieSceneTracksPrivatePCH.h"
+#include "MovieSceneFadeTemplate.h"
+#include "MovieSceneFadeTrack.h"
+#include "MovieSceneFadeSection.h"
+#include "MovieSceneSequence.h"
+#include "MovieSceneEvaluationTemplateInstance.h"
+
+DECLARE_CYCLE_STAT(TEXT("Fade Track Token Execute"), MovieSceneEval_FadeTrack_TokenExecute, STATGROUP_MovieSceneEval);
+
+struct FFadeTrackToken
+{
+	FFadeTrackToken(float InFadeValue, const FLinearColor& InFadeColor, bool bInFadeAudio)
+		: FadeValue(InFadeValue)
+		, FadeColor(InFadeColor)
+		, bFadeAudio(bInFadeAudio)
+	{}
+
+	float FadeValue;
+	FLinearColor FadeColor;
+	bool bFadeAudio;
+
+	void Apply(IMovieScenePlayer& Player)
+	{
+		// Set editor preview/fade
+		EMovieSceneViewportParams ViewportParams;
+		ViewportParams.SetWhichViewportParam = (EMovieSceneViewportParams::SetViewportParam)(EMovieSceneViewportParams::SVP_FadeAmount | EMovieSceneViewportParams::SVP_FadeColor);
+		ViewportParams.FadeAmount = FadeValue;
+		ViewportParams.FadeColor = FadeColor;
+
+		TMap<FViewportClient*, EMovieSceneViewportParams> ViewportParamsMap;
+		Player.GetViewportSettings(ViewportParamsMap);
+		for (auto ViewportParamsPair : ViewportParamsMap)
+		{
+			ViewportParamsMap[ViewportParamsPair.Key] = ViewportParams;
+		}
+		Player.SetViewportSettings(ViewportParamsMap);
+
+		// Set runtime fade
+		for (const FWorldContext& Context : GEngine->GetWorldContexts())
+		{
+			if (Context.WorldType == EWorldType::Game || Context.WorldType == EWorldType::PIE)
+			{
+				UWorld* World = Context.World();
+				if (World != nullptr)
+				{
+					APlayerController* PlayerController = World->GetGameInstance()->GetFirstLocalPlayerController();
+					if (PlayerController != nullptr && PlayerController->PlayerCameraManager && !PlayerController->PlayerCameraManager->IsPendingKill())
+					{
+						PlayerController->PlayerCameraManager->SetManualCameraFade(FadeValue, FadeColor, bFadeAudio);
+					}
+				}
+			}
+		}
+	}
+};
+
+struct FFadePreAnimatedGlobalToken : FFadeTrackToken, IMovieScenePreAnimatedGlobalToken
+{
+	FFadePreAnimatedGlobalToken(float InFadeValue, const FLinearColor& InFadeColor, bool bInFadeAudio)
+		: FFadeTrackToken(InFadeValue, InFadeColor, bInFadeAudio)
+	{}
+
+	virtual void RestoreState(IMovieScenePlayer& Player) override
+	{
+		Apply(Player);
+	}
+};
+
+struct FFadePreAnimatedGlobalTokenProducer : IMovieScenePreAnimatedGlobalTokenProducer
+{
+	virtual IMovieScenePreAnimatedGlobalTokenPtr CacheExistingState() const override
+	{
+		return FFadePreAnimatedGlobalToken(0.f, FLinearColor::Black, false);
+	}
+};
+
+/** A movie scene execution token that applies fades */
+struct FFadeExecutionToken : IMovieSceneExecutionToken, FFadeTrackToken
+{
+	FFadeExecutionToken(float InFadeValue, const FLinearColor& InFadeColor, bool bInFadeAudio)
+		: FFadeTrackToken(InFadeValue, InFadeColor, bInFadeAudio)
+	{}
+
+	static FMovieSceneAnimTypeID GetAnimTypeID()
+	{
+		return TMovieSceneAnimTypeID<FFadeExecutionToken>();
+	}
+	
+	/** Execute this token, operating on all objects referenced by 'Operand' */
+	virtual void Execute(const FMovieSceneContext& Context, const FMovieSceneEvaluationOperand& Operand, FPersistentEvaluationData& PersistentData, IMovieScenePlayer& Player) override
+	{
+		MOVIESCENE_DETAILED_SCOPE_CYCLE_COUNTER(MovieSceneEval_FadeTrack_TokenExecute)
+
+		Player.SavePreAnimatedState(GetAnimTypeID(), FFadePreAnimatedGlobalTokenProducer());
+
+		Apply(Player);
+	}
+};
+
+FMovieSceneFadeSectionTemplate::FMovieSceneFadeSectionTemplate(const UMovieSceneFadeSection& Section)
+	: FadeCurve(Section.GetFloatCurve())
+	, FadeColor(Section.FadeColor)
+	, bFadeAudio(Section.bFadeAudio)
+{
+}
+
+void FMovieSceneFadeSectionTemplate::Evaluate(const FMovieSceneEvaluationOperand& Operand, const FMovieSceneContext& Context, const FPersistentEvaluationData& PersistentData, FMovieSceneExecutionTokens& ExecutionTokens) const
+{
+	float FadeValue = FadeCurve.Eval(Context.GetTime());
+	ExecutionTokens.Add(FFadeExecutionToken(FadeValue, FadeColor, bFadeAudio));
+}
\ No newline at end of file
diff --git a/Engine/Source/Runtime/MovieSceneTracks/Private/Evaluation/MovieSceneFadeTemplate.h b/Engine/Source/Runtime/MovieSceneTracks/Private/Evaluation/MovieSceneFadeTemplate.h
new file mode 100644
--- /dev/null
+++ b/Engine/Source/Runtime/MovieSceneTracks/Private/Evaluation/MovieSceneFadeTemplate.h
@@ -0,0 +1,36 @@
+// Copyright 1998-2016 Epic Games, Inc. All Rights Reserved.
+
+#pragma once
+
+#include "MovieSceneFwd.h"
+#include "MovieSceneTrackImplementation.h"
+#include "MovieSceneEvalTemplate.h"
+#include "MovieSceneFadeSection.h"
+#include "Curves/RichCurve.h"
+
+#include "MovieSceneFadeTemplate.generated.h"
+
+class UMovieSceneFadeTrack;
+
+USTRUCT()
+struct FMovieSceneFadeSectionTemplate : public FMovieSceneEvalTemplate
+{
+	GENERATED_BODY()
+	
+	FMovieSceneFadeSectionTemplate() {}
+	FMovieSceneFadeSectionTemplate(const UMovieSceneFadeSection& Section);
+
+private:
+
+	virtual UScriptStruct& GetScriptStructImpl() const override { return *StaticStruct(); }
+	virtual void Evaluate(const FMovieSceneEvaluationOperand& Operand, const FMovieSceneContext& Context, const FPersistentEvaluationData& PersistentData, FMovieSceneExecutionTokens& ExecutionTokens) const override;
+
+	UPROPERTY()
+	FRichCurve FadeCurve;
+	
+	UPROPERTY()
+	FLinearColor FadeColor;
+
+	UPROPERTY()
+	uint32 bFadeAudio:1;
+};
