Copying //UE4/Dev-Sequencer to //UE4/Dev-Main (Source: //UE4/Dev-Sequencer @ 3178529)

diff --git a/Engine/Source/Runtime/MovieSceneTracks/Private/Evaluation/MovieSceneLevelVisibilityTemplate.cpp b/Engine/Source/Runtime/MovieSceneTracks/Private/Evaluation/MovieSceneLevelVisibilityTemplate.cpp
new file mode 100644
--- /dev/null
+++ b/Engine/Source/Runtime/MovieSceneTracks/Private/Evaluation/MovieSceneLevelVisibilityTemplate.cpp
@@ -0,0 +1,323 @@
+// Copyright 1998-2016 Epic Games, Inc. All Rights Reserved.
+
+#include "MovieSceneTracksPrivatePCH.h"
+#include "MovieSceneLevelVisibilityTemplate.h"
+
+#include "MovieSceneLevelVisibilitySection.h"
+
+bool GetLevelVisibility(const ULevelStreaming& Level)
+{
+#if WITH_EDITOR
+	if (GIsEditor && !Level.GetWorld()->IsPlayInEditor())
+	{
+		return Level.bShouldBeVisibleInEditor;
+	}
+	else
+#endif
+	{
+		return Level.bShouldBeVisible;
+	}
+}
+
+void SetLevelVisibility(ULevelStreaming& Level, bool bVisible)
+{
+#if WITH_EDITOR
+	if (GIsEditor && !Level.GetWorld()->IsPlayInEditor())
+	{
+		Level.bShouldBeVisibleInEditor = bVisible;
+		Level.GetWorld()->FlushLevelStreaming();
+
+		// Iterate over the level's actors
+		TArray<AActor*>& Actors = Level.GetLoadedLevel()->Actors;
+		for ( int32 ActorIndex = 0; ActorIndex < Actors.Num(); ++ActorIndex )
+		{
+			AActor* Actor = Actors[ActorIndex];
+			if ( Actor )
+			{
+				if (Actor->bHiddenEdLevel == bVisible )
+				{
+					Actor->bHiddenEdLevel = !bVisible;
+					if ( bVisible )
+					{
+						Actor->ReregisterAllComponents();
+					}
+					else
+					{
+						Actor->UnregisterAllComponents();
+					}
+				}
+			}
+		}
+	}
+	else
+#endif
+	{
+		Level.bShouldBeVisible = bVisible;
+	}
+}
+
+// TODO: This was copied from LevelStreaming.cpp, it should be in a set of shared utilities somewhere.
+FString MakeSafeLevelName(const FName& InLevelName, UWorld& World)
+{
+	// Special case for PIE, the PackageName gets mangled.
+	if (!World.StreamingLevelsPrefix.IsEmpty())
+	{
+		FString PackageName = World.StreamingLevelsPrefix + FPackageName::GetShortName(InLevelName);
+		if (!FPackageName::IsShortPackageName(InLevelName))
+		{
+			PackageName = FPackageName::GetLongPackagePath(InLevelName.ToString()) + TEXT( "/" ) + PackageName;
+		}
+
+		return PackageName;
+	}
+
+	return InLevelName.ToString();
+}
+
+ULevelStreaming* GetStreamingLevel(FName LevelName, UWorld& World)
+{
+	if (LevelName != NAME_None)
+	{
+		FString SafeLevelName = MakeSafeLevelName(LevelName, World);
+		if (FPackageName::IsShortPackageName(SafeLevelName))
+		{
+			// Make sure MyMap1 and Map1 names do not resolve to a same streaming level
+			SafeLevelName = TEXT("/") + SafeLevelName;
+		}
+
+		for (ULevelStreaming* LevelStreaming : World.StreamingLevels)
+		{
+			if (LevelStreaming && LevelStreaming->GetWorldAssetPackageName().EndsWith(SafeLevelName, ESearchCase::IgnoreCase))
+			{
+				return LevelStreaming;
+			}
+		}
+	}
+
+	return nullptr;
+}
+
+struct FLevelStreamingPreAnimatedToken : IMovieScenePreAnimatedToken
+{
+	FLevelStreamingPreAnimatedToken(bool bInIsVisible)
+		: bVisible(bInIsVisible)
+	{
+	}
+
+	virtual void RestoreState(UObject& Object, IMovieScenePlayer& Player) override
+	{
+		ULevelStreaming* LevelStreaming = CastChecked<ULevelStreaming>(&Object);
+		SetLevelVisibility(*LevelStreaming, bVisible);
+	}
+
+	bool bVisible;
+};
+
+struct FLevelStreamingPreAnimatedTokenProducer : IMovieScenePreAnimatedTokenProducer
+{
+	virtual IMovieScenePreAnimatedTokenPtr CacheExistingState(UObject& Object) const override
+	{
+		ULevelStreaming* LevelStreaming = CastChecked<ULevelStreaming>(&Object);
+		return FLevelStreamingPreAnimatedToken(GetLevelVisibility(*LevelStreaming));
+	}
+};
+
+struct FLevelStreamingSharedTrackData : IPersistentEvaluationData
+{
+	bool HasAnythingToDo() const
+	{
+		return VisibilityMap.Num() != 0;
+	}
+
+	void AssignLevelVisibilityOverrides(TArrayView<const FName> LevelNames, ELevelVisibility Visibility)
+	{
+		const int32 Inc = Visibility == ELevelVisibility::Visible ? 1 : -1;
+		for (FName Name : LevelNames)
+		{
+			VisibilityMap.FindOrAdd(Name).VisibilityRequestCount += Inc;
+		}
+	}
+	
+	void UnassignLevelVisibilityOverrides(TArrayView<const FName> LevelNames, ELevelVisibility Visibility)
+	{
+		const int32 Dec = Visibility == ELevelVisibility::Visible ? 1 : -1;
+		for (FName Name : LevelNames)
+		{
+			VisibilityMap.FindOrAdd(Name).VisibilityRequestCount -= Dec;
+		}
+	}
+
+	void TearDown(IMovieScenePlayer& Player)
+	{
+		for (auto& Pair : VisibilityMap)
+		{
+			Pair.Value.VisibilityRequestCount = 0;
+		}
+
+		ApplyLevelVisibility(Player);
+	}
+
+	void ApplyLevelVisibility(IMovieScenePlayer& Player)
+	{
+		UWorld* World = Player.GetPlaybackContext()->GetWorld();
+		if (!World)
+		{
+			return;
+		}
+
+		FLevelStreamingPreAnimatedTokenProducer TokenProducer;
+
+		TArray<FName, TInlineAllocator<8>> LevelsToRestore;
+
+		bool bAnythingChanged = false;
+		for (auto& Pair : VisibilityMap)
+		{
+			FName SafeLevelName(*MakeSafeLevelName(Pair.Key, *World));
+
+			ULevelStreaming* Level = GetLevel(SafeLevelName, *World);
+			if (!Level)
+			{
+				continue;
+			}
+
+			if (Pair.Value.VisibilityRequestCount == 0)
+			{
+				LevelsToRestore.Add(Pair.Key);
+
+				// Restore the state from before our evaluation
+				if (Pair.Value.bPreviousState.IsSet())
+				{
+					SetLevelVisibility(*Level, Pair.Value.bPreviousState.GetValue());
+					bAnythingChanged = true;
+				}
+			}
+			else
+			{
+				const bool bShouldBeVisible = Pair.Value.VisibilityRequestCount > 0;
+				if (GetLevelVisibility(*Level) != bShouldBeVisible)
+				{
+					// Save preanimated state for this track, and globally
+					if (!Pair.Value.bPreviousState.IsSet())
+					{
+						Pair.Value.bPreviousState = GetLevelVisibility(*Level);
+					}
+
+					Player.SavePreAnimatedState(*Level, TMovieSceneAnimTypeID<FLevelStreamingSharedTrackData>(), TokenProducer);
+
+					SetLevelVisibility(*Level, bShouldBeVisible);
+					bAnythingChanged = true;
+				}
+			}
+		}
+
+		for (FName Level : LevelsToRestore)
+		{
+			VisibilityMap.Remove(Level);
+		}
+
+		if (bAnythingChanged)
+		{
+			World->FlushLevelStreaming( EFlushLevelStreamingType::Visibility );
+		}
+	}
+
+private:
+
+	ULevelStreaming* GetLevel(FName SafeLevelName, UWorld& World)
+	{
+		if (TWeakObjectPtr<ULevelStreaming>* FoundStreamingLevel = NameToLevelMap.Find(SafeLevelName))
+		{
+			if (ULevelStreaming* Level = FoundStreamingLevel->Get())
+			{
+				return Level;
+			}
+
+			NameToLevelMap.Remove(SafeLevelName);
+		}
+
+		ULevelStreaming* Level = GetStreamingLevel(SafeLevelName, World);
+		if (Level)
+		{
+			NameToLevelMap.Add(SafeLevelName, Level);
+		}
+
+		return Level;
+	}
+	
+	struct FVisibilityData
+	{
+		FVisibilityData() : VisibilityRequestCount(0) {}
+
+		/** Ref count of things asking for this level to be (in)visible. > 0 signifies visible, < 0 signifies invisible */
+		int32 VisibilityRequestCount;
+		TOptional<bool> bPreviousState;
+	};
+	TMap<FName, FVisibilityData> VisibilityMap;
+
+	TMap<FName, TWeakObjectPtr<ULevelStreaming>> NameToLevelMap;
+};
+
+FMovieSceneLevelVisibilitySectionTemplate::FMovieSceneLevelVisibilitySectionTemplate(const UMovieSceneLevelVisibilitySection& Section)
+	: Visibility(Section.GetVisibility())
+	, LevelNames(Section.GetLevelNames())
+{
+}
+
+void FMovieSceneLevelVisibilitySectionTemplate::Evaluate(const FMovieSceneEvaluationOperand& Operand, const FMovieSceneContext& Context, const FPersistentEvaluationData& PersistentData, FMovieSceneExecutionTokens& ExecutionTokens) const
+{
+	// Evaluate doesn't actually do anything
+}
+
+void FMovieSceneLevelVisibilitySectionTemplate::Setup(FPersistentEvaluationData& PersistentData, IMovieScenePlayer& Player) const
+{
+	FLevelStreamingSharedTrackData& TrackData = PersistentData.GetOrAdd<FLevelStreamingSharedTrackData>(FMovieSceneLevelVisibilitySharedTrack::GetSharedDataKey());
+	TrackData.AssignLevelVisibilityOverrides(LevelNames, Visibility);
+}
+
+void FMovieSceneLevelVisibilitySectionTemplate::TearDown(FPersistentEvaluationData& PersistentData, IMovieScenePlayer& Player) const
+{
+	// Track data may have already disappeared if the FMovieSceneLevelVisibilitySharedTrack is no longer evaluated either
+	FLevelStreamingSharedTrackData* TrackData = PersistentData.Find<FLevelStreamingSharedTrackData>(FMovieSceneLevelVisibilitySharedTrack::GetSharedDataKey());
+	if (TrackData)
+	{
+		TrackData->UnassignLevelVisibilityOverrides(LevelNames, Visibility);
+	}
+}
+
+FSharedPersistentDataKey FMovieSceneLevelVisibilitySharedTrack::GetSharedDataKey()
+{
+	static FMovieSceneSharedDataId DataId(FMovieSceneSharedDataId::Allocate());
+	return FSharedPersistentDataKey(DataId, FMovieSceneEvaluationOperand());
+}
+
+void FMovieSceneLevelVisibilitySharedTrack::TearDown(FPersistentEvaluationData& PersistentData, IMovieScenePlayer& Player) const
+{
+	FLevelStreamingSharedTrackData* TrackData = PersistentData.Find<FLevelStreamingSharedTrackData>(GetSharedDataKey());
+	if (TrackData && TrackData->HasAnythingToDo())
+	{
+		TrackData->TearDown(Player);
+
+		PersistentData.Reset(GetSharedDataKey());
+	}
+}
+
+struct FLevelVisibilityExecutionToken : IMovieSceneExecutionToken
+{
+	virtual void Execute(const FMovieSceneContext& Context, const FMovieSceneEvaluationOperand& Operand, FPersistentEvaluationData& PersistentData, IMovieScenePlayer& Player) override
+	{
+		FLevelStreamingSharedTrackData* TrackData = PersistentData.Find<FLevelStreamingSharedTrackData>(FMovieSceneLevelVisibilitySharedTrack::GetSharedDataKey());
+		if (TrackData)
+		{
+			TrackData->ApplyLevelVisibility(Player);
+		}
+	}
+};
+
+void FMovieSceneLevelVisibilitySharedTrack::Evaluate(const FMovieSceneEvaluationOperand& Operand, const FMovieSceneContext& Context, const FPersistentEvaluationData& PersistentData, FMovieSceneExecutionTokens& ExecutionTokens) const
+{
+	const FLevelStreamingSharedTrackData* TrackData = PersistentData.Find<FLevelStreamingSharedTrackData>(GetSharedDataKey());
+	if (TrackData && TrackData->HasAnythingToDo())
+	{
+		ExecutionTokens.Add(FLevelVisibilityExecutionToken());
+	}
+}
\ No newline at end of file
diff --git a/Engine/Source/Runtime/MovieSceneTracks/Private/Evaluation/MovieSceneLevelVisibilityTemplate.h b/Engine/Source/Runtime/MovieSceneTracks/Private/Evaluation/MovieSceneLevelVisibilityTemplate.h
new file mode 100644
--- /dev/null
+++ b/Engine/Source/Runtime/MovieSceneTracks/Private/Evaluation/MovieSceneLevelVisibilityTemplate.h
@@ -0,0 +1,65 @@
+// Copyright 1998-2016 Epic Games, Inc. All Rights Reserved.
+
+#pragma once
+
+#include "MovieSceneEvalTemplate.h"
+#include "MovieSceneLevelVisibilitySection.h"
+#include "MovieSceneLevelVisibilityTemplate.generated.h"
+
+
+class UMovieSceneLevelVisibilitySection;
+
+
+USTRUCT()
+struct FMovieSceneLevelVisibilitySectionTemplate
+	: public FMovieSceneEvalTemplate
+{
+	GENERATED_BODY()
+	
+	FMovieSceneLevelVisibilitySectionTemplate(){}
+	FMovieSceneLevelVisibilitySectionTemplate(const UMovieSceneLevelVisibilitySection& Section);
+
+private:
+
+	virtual UScriptStruct& GetScriptStructImpl() const override
+	{
+		return *StaticStruct();
+	}
+	virtual void SetupOverrides() override
+	{
+		EnableOverrides(RequiresSetupFlag | RequiresTearDownFlag);
+	}
+
+	virtual void Evaluate(const FMovieSceneEvaluationOperand& Operand, const FMovieSceneContext& Context, const FPersistentEvaluationData& PersistentData, FMovieSceneExecutionTokens& ExecutionTokens) const override;
+	virtual void Setup(FPersistentEvaluationData& PersistentData, IMovieScenePlayer& Player) const override;
+	virtual void TearDown(FPersistentEvaluationData& PersistentData, IMovieScenePlayer& Player) const override;
+
+	UPROPERTY()
+	ELevelVisibility Visibility;
+
+	UPROPERTY()
+	TArray<FName> LevelNames;
+};
+
+
+USTRUCT()
+struct FMovieSceneLevelVisibilitySharedTrack
+	: public FMovieSceneEvalTemplate
+{
+	GENERATED_BODY()
+
+	MOVIESCENETRACKS_API static FSharedPersistentDataKey GetSharedDataKey();
+
+private:
+
+	virtual UScriptStruct& GetScriptStructImpl() const override
+	{
+		return *StaticStruct();
+	}
+	virtual void SetupOverrides() override
+	{
+		EnableOverrides(RequiresTearDownFlag);
+	}
+	virtual void TearDown(FPersistentEvaluationData& PersistentData, IMovieScenePlayer& Player) const override;
+	virtual void Evaluate(const FMovieSceneEvaluationOperand& Operand, const FMovieSceneContext& Context, const FPersistentEvaluationData& PersistentData, FMovieSceneExecutionTokens& ExecutionTokens) const override;
+};
diff --git a/Engine/Source/Runtime/MovieSceneTracks/Private/Evaluation/MovieSceneParameterTemplate.cpp b/Engine/Source/Runtime/MovieSceneTracks/Private/Evaluation/MovieSceneParameterTemplate.cpp
new file mode 100644
--- /dev/null
+++ b/Engine/Source/Runtime/MovieSceneTracks/Private/Evaluation/MovieSceneParameterTemplate.cpp
@@ -0,0 +1,117 @@
+// Copyright 1998-2016 Epic Games, Inc. All Rights Reserved.
+
+#include "MovieSceneTracksPrivatePCH.h"
+#include "MovieSceneParameterSection.h"
+#include "MovieSceneMaterialTrack.h"
+#include "MovieSceneParameterTemplate.h"
+
+FMovieSceneParameterSectionTemplate::FMovieSceneParameterSectionTemplate(const UMovieSceneParameterSection& Section)
+	: Scalars(Section.GetScalarParameterNamesAndCurves())
+	, Vectors(Section.GetVectorParameterNamesAndCurves())
+	, Colors(Section.GetColorParameterNamesAndCurves())
+{
+}
+
+void FMovieSceneParameterSectionTemplate::EvaluateCurves(const FMovieSceneContext& Context, FEvaluatedParameterSectionValues& Values) const
+{
+	float Time = Context.GetTime();
+
+	for ( const FScalarParameterNameAndCurve& Scalar : Scalars )
+	{
+		Values.ScalarValues.Add(
+			FScalarParameterNameAndValue(Scalar.ParameterName, Scalar.ParameterCurve.Eval(Time))
+			);
+	}
+
+	for ( const FVectorParameterNameAndCurves& Vector : Vectors )
+	{
+		Values.VectorValues.Add(
+			FVectorParameterNameAndValue(
+				Vector.ParameterName,
+				FVector(
+					Vector.XCurve.Eval(Time),
+					Vector.YCurve.Eval(Time),
+					Vector.ZCurve.Eval(Time)
+				)
+			)
+		);
+	}
+
+	for ( const FColorParameterNameAndCurves& Color : Colors )
+	{
+		Values.ColorValues.Add(
+			FColorParameterNameAndValue(
+				Color.ParameterName,
+				FLinearColor(
+					Color.RedCurve.Eval(Time ),
+					Color.GreenCurve.Eval(Time),
+					Color.BlueCurve.Eval(Time),
+					Color.AlphaCurve.Eval(Time)
+				)
+			)
+		);
+	}
+}
+
+void FDefaultMaterialAccessor::Apply(UMaterialInstanceDynamic& Material, const FEvaluatedParameterSectionValues& Values)
+{
+	for (const FScalarParameterNameAndValue& ScalarValue : Values.ScalarValues)
+	{
+		Material.SetScalarParameterValue(ScalarValue.ParameterName, ScalarValue.Value);
+	}
+	for (const FVectorParameterNameAndValue& VectorValue : Values.VectorValues)
+	{
+		Material.SetVectorParameterValue(VectorValue.ParameterName, VectorValue.Value);
+	}
+	for (const FColorParameterNameAndValue& ColorValue : Values.ColorValues)
+	{
+		Material.SetVectorParameterValue(ColorValue.ParameterName, ColorValue.Value);
+	}
+}
+
+
+TMovieSceneAnimTypeIDContainer<int32> MaterialIndexAnimTypeIDs;
+
+struct FComponentMaterialAccessor : FDefaultMaterialAccessor
+{
+	FComponentMaterialAccessor(int32 InMaterialIndex)
+		: MaterialIndex(InMaterialIndex)
+	{}
+
+	FComponentMaterialAccessor(const FComponentMaterialAccessor&) = default;
+	FComponentMaterialAccessor& operator=(const FComponentMaterialAccessor&) = default;
+	
+	FMovieSceneAnimTypeID GetAnimTypeID() const
+	{
+		return MaterialIndexAnimTypeIDs.GetAnimTypeID(MaterialIndex);
+	}
+
+	UMaterialInterface* GetMaterialForObject(UObject& Object) const
+	{
+		UPrimitiveComponent* Component = Cast<UPrimitiveComponent>(&Object);
+		return Component ? Component->GetMaterial(MaterialIndex) : nullptr;
+	}
+
+	void SetMaterialForObject(UObject& Object, UMaterialInterface& Material) const
+	{
+		UPrimitiveComponent* Component = CastChecked<UPrimitiveComponent>(&Object);
+		Component->SetMaterial(MaterialIndex, &Material);
+	}
+
+	int32 MaterialIndex;
+};
+
+FMovieSceneComponentMaterialSectionTemplate::FMovieSceneComponentMaterialSectionTemplate(const UMovieSceneParameterSection& Section, const UMovieSceneComponentMaterialTrack& Track)
+	: FMovieSceneParameterSectionTemplate(Section)
+	, MaterialIndex(Track.GetMaterialIndex())
+{
+}
+
+void FMovieSceneComponentMaterialSectionTemplate::Evaluate(const FMovieSceneEvaluationOperand& Operand, const FMovieSceneContext& Context, const FPersistentEvaluationData& PersistentData, FMovieSceneExecutionTokens& ExecutionTokens) const
+{
+	TMaterialTrackExecutionToken<FComponentMaterialAccessor> ExecutionToken(MaterialIndex);
+
+	EvaluateCurves(Context, ExecutionToken.Values);
+
+	ExecutionTokens.Add(MoveTemp(ExecutionToken));
+}
\ No newline at end of file
diff --git a/Engine/Source/Runtime/MovieSceneTracks/Private/Evaluation/MovieSceneParticleTemplate.cpp b/Engine/Source/Runtime/MovieSceneTracks/Private/Evaluation/MovieSceneParticleTemplate.cpp
new file mode 100644
--- /dev/null
+++ b/Engine/Source/Runtime/MovieSceneTracks/Private/Evaluation/MovieSceneParticleTemplate.cpp
@@ -0,0 +1,104 @@
+// Copyright 1998-2016 Epic Games, Inc. All Rights Reserved.
+
+#include "MovieSceneTracksPrivatePCH.h"
+#include "MovieSceneParticleTemplate.h"
+#include "MovieSceneParticleSection.h"
+
+DECLARE_CYCLE_STAT(TEXT("Particle Track Token Execute"), MovieSceneEval_ParticleTrack_TokenExecute, STATGROUP_MovieSceneEval);
+
+struct FParticleKeyState : IPersistentEvaluationData
+{
+	FKeyHandle LastKeyHandle;
+	FKeyHandle InvalidKeyHandle;
+};
+
+/** A movie scene execution token that stores a specific transform, and an operand */
+struct FParticleTrackExecutionToken
+	: IMovieSceneExecutionToken
+{
+	FParticleTrackExecutionToken(EParticleKey::Type InParticleKey, TOptional<FKeyHandle> InKeyHandle = TOptional<FKeyHandle>())
+		: ParticleKey(InParticleKey), KeyHandle(InKeyHandle)
+	{
+	}
+
+	static UParticleSystemComponent* GetComponentFromObject(UObject* Object)
+	{
+		if(AEmitter* Emitter = Cast<AEmitter>(Object))
+		{
+			return Emitter->GetParticleSystemComponent();
+		}
+		else
+		{
+			return Cast<UParticleSystemComponent>(Object);
+		}
+	}
+
+	/** Execute this token, operating on all objects referenced by 'Operand' */
+	virtual void Execute(const FMovieSceneContext& Context, const FMovieSceneEvaluationOperand& Operand, FPersistentEvaluationData& PersistentData, IMovieScenePlayer& Player) override
+	{
+		MOVIESCENE_DETAILED_SCOPE_CYCLE_COUNTER(MovieSceneEval_ParticleTrack_TokenExecute)
+		
+		if (KeyHandle.IsSet())
+		{
+			PersistentData.GetOrAddSectionData<FParticleKeyState>().LastKeyHandle = KeyHandle.GetValue();
+		}
+
+		for (TWeakObjectPtr<> Object : Player.FindBoundObjects(Operand))
+		{
+			UObject* ObjectPtr = Object.Get();
+			UParticleSystemComponent* ParticleSystemComponent = GetComponentFromObject(ObjectPtr);
+
+			if (ParticleSystemComponent)
+			{
+				if ( ParticleKey == EParticleKey::Activate)
+				{
+					if ( ParticleSystemComponent->IsActive() )
+					{
+						ParticleSystemComponent->SetActive(false, true);
+					}
+					ParticleSystemComponent->SetActive(true, true);
+				}
+				else if( ParticleKey == EParticleKey::Deactivate )
+				{
+					ParticleSystemComponent->SetActive(false, true);
+				}
+				else if ( ParticleKey == EParticleKey::Trigger )
+				{
+					ParticleSystemComponent->ActivateSystem(true);
+				}
+			}
+		}
+	}
+
+	EParticleKey::Type ParticleKey;
+	TOptional<FKeyHandle> KeyHandle;
+};
+
+FMovieSceneParticleSectionTemplate::FMovieSceneParticleSectionTemplate(const UMovieSceneParticleSection& Section)
+	: ParticleKeys(Section.GetParticleCurve())
+{
+}
+
+void FMovieSceneParticleSectionTemplate::Evaluate(const FMovieSceneEvaluationOperand& Operand, const FMovieSceneContext& Context, const FPersistentEvaluationData& PersistentData, FMovieSceneExecutionTokens& ExecutionTokens) const
+{
+	const bool bPlaying = Context.GetDirection() == EPlayDirection::Forwards && Context.GetRange().Size<float>() >= 0.f && Context.GetStatus() == EMovieScenePlayerStatus::Playing;
+
+	const FParticleKeyState* SectionData = PersistentData.FindSectionData<FParticleKeyState>();
+
+	if (!bPlaying)
+	{
+		ExecutionTokens.Add(FParticleTrackExecutionToken(EParticleKey::Deactivate, SectionData ? SectionData->InvalidKeyHandle : TOptional<FKeyHandle>()));
+	}
+	else
+	{
+		FKeyHandle PreviousHandle = ParticleKeys.FindKeyBeforeOrAt(Context.GetTime());
+		if (ParticleKeys.IsKeyHandleValid(PreviousHandle) && (!SectionData || SectionData->LastKeyHandle != PreviousHandle))
+		{
+			FParticleTrackExecutionToken NewToken(
+				(EParticleKey::Type)ParticleKeys.GetKey(PreviousHandle).Value,
+				PreviousHandle);
+
+			ExecutionTokens.Add(MoveTemp(NewToken));
+		}
+	}
+}
diff --git a/Engine/Source/Runtime/MovieSceneTracks/Private/Evaluation/MovieScenePropertyTemplates.cpp b/Engine/Source/Runtime/MovieSceneTracks/Private/Evaluation/MovieScenePropertyTemplates.cpp
new file mode 100644
--- /dev/null
+++ b/Engine/Source/Runtime/MovieSceneTracks/Private/Evaluation/MovieScenePropertyTemplates.cpp
@@ -0,0 +1,182 @@
+// Copyright 1998-2016 Epic Games, Inc. All Rights Reserved.
+
+#include "MovieSceneTracksPrivatePCH.h"
+#include "MovieScenePropertyTemplates.h"
+#include "MovieSceneBoolSection.h"
+#include "MovieSceneFloatSection.h"
+#include "MovieSceneByteSection.h"
+#include "MovieSceneIntegerSection.h"
+#include "MovieSceneVectorSection.h"
+#include "MovieSceneStringSection.h"
+#include "MovieScenePropertyTrack.h"
+
+namespace
+{
+	FName SanitizeBoolPropertyName(FName InPropertyName)
+	{
+		FString PropertyVarName = InPropertyName.ToString();
+		PropertyVarName.RemoveFromStart("b", ESearchCase::CaseSensitive);
+		return FName(*PropertyVarName);
+	}
+}
+
+FMovieSceneBoolPropertySectionTemplate::FMovieSceneBoolPropertySectionTemplate(const UMovieSceneBoolSection& Section, const UMovieScenePropertyTrack& Track)
+	: PropertyData(SanitizeBoolPropertyName(Track.GetPropertyName()), Track.GetPropertyPath())
+	, BoolCurve(Section.GetCurve())
+{}
+
+void FMovieSceneBoolPropertySectionTemplate::Evaluate(const FMovieSceneEvaluationOperand& Operand, const FMovieSceneContext& Context, const FPersistentEvaluationData& PersistentData, FMovieSceneExecutionTokens& ExecutionTokens) const
+{
+	using namespace PropertyTemplate;
+
+	for (TCachedValue<bool>& ObjectAndValue : PersistentData.GetSectionData<TCachedSectionData<bool>>().ObjectsAndValues)
+	{
+		ObjectAndValue.Value = !!BoolCurve.Evaluate(Context.GetTime(), ObjectAndValue.Value ? 1 : 0);
+	}
+	
+	ExecutionTokens.Add(TCachedPropertyTrackExecutionToken<bool>());
+}
+
+FMovieSceneFloatPropertySectionTemplate::FMovieSceneFloatPropertySectionTemplate(const UMovieSceneFloatSection& Section, const UMovieScenePropertyTrack& Track)
+	: PropertyData(Track.GetPropertyName(), Track.GetPropertyPath())
+	, FloatCurve(Section.GetFloatCurve())
+{}
+
+void FMovieSceneFloatPropertySectionTemplate::Evaluate(const FMovieSceneEvaluationOperand& Operand, const FMovieSceneContext& Context, const FPersistentEvaluationData& PersistentData, FMovieSceneExecutionTokens& ExecutionTokens) const
+{
+	using namespace PropertyTemplate;
+
+	for (TCachedValue<float>& ObjectAndValue : PersistentData.GetSectionData<TCachedSectionData<float>>().ObjectsAndValues)
+	{
+		ObjectAndValue.Value = FloatCurve.Eval(Context.GetTime(), ObjectAndValue.Value);
+	}
+	
+	ExecutionTokens.Add(TCachedPropertyTrackExecutionToken<float>());
+}
+
+FMovieSceneBytePropertySectionTemplate::FMovieSceneBytePropertySectionTemplate(const UMovieSceneByteSection& Section, const UMovieScenePropertyTrack& Track)
+	: PropertyData(Track.GetPropertyName(), Track.GetPropertyPath())
+	, ByteCurve(Section.GetCurve())
+{}
+
+void FMovieSceneBytePropertySectionTemplate::Evaluate(const FMovieSceneEvaluationOperand& Operand, const FMovieSceneContext& Context, const FPersistentEvaluationData& PersistentData, FMovieSceneExecutionTokens& ExecutionTokens) const
+{
+	using namespace PropertyTemplate;
+
+	for (TCachedValue<uint8>& ObjectAndValue : PersistentData.GetSectionData<TCachedSectionData<uint8>>().ObjectsAndValues)
+	{
+		ObjectAndValue.Value = ByteCurve.Evaluate(Context.GetTime(), ObjectAndValue.Value);
+	}
+	
+	ExecutionTokens.Add(TCachedPropertyTrackExecutionToken<uint8>());
+}
+
+FMovieSceneIntegerPropertySectionTemplate::FMovieSceneIntegerPropertySectionTemplate(const UMovieSceneIntegerSection& Section, const UMovieScenePropertyTrack& Track)
+	: PropertyData(Track.GetPropertyName(), Track.GetPropertyPath())
+	, IntegerCurve(Section.GetCurve())
+{}
+
+void FMovieSceneIntegerPropertySectionTemplate::Evaluate(const FMovieSceneEvaluationOperand& Operand, const FMovieSceneContext& Context, const FPersistentEvaluationData& PersistentData, FMovieSceneExecutionTokens& ExecutionTokens) const
+{
+	using namespace PropertyTemplate;
+
+	for (TCachedValue<int32>& ObjectAndValue : PersistentData.GetSectionData<TCachedSectionData<int32>>().ObjectsAndValues)
+	{
+		ObjectAndValue.Value = IntegerCurve.Evaluate(Context.GetTime(), ObjectAndValue.Value);
+	}
+	
+	ExecutionTokens.Add(TCachedPropertyTrackExecutionToken<int32>());
+}
+
+FMovieSceneStringPropertySectionTemplate::FMovieSceneStringPropertySectionTemplate(const UMovieSceneStringSection& Section, const UMovieScenePropertyTrack& Track)
+	: PropertyData(Track.GetPropertyName(), Track.GetPropertyPath())
+	, StringCurve(Section.GetStringCurve())
+{}
+
+void FMovieSceneStringPropertySectionTemplate::Evaluate(const FMovieSceneEvaluationOperand& Operand, const FMovieSceneContext& Context, const FPersistentEvaluationData& PersistentData, FMovieSceneExecutionTokens& ExecutionTokens) const
+{
+	using namespace PropertyTemplate;
+
+	for (TCachedValue<FString>& ObjectAndValue : PersistentData.GetSectionData<TCachedSectionData<FString>>().ObjectsAndValues)
+	{
+		ObjectAndValue.Value = StringCurve.Eval(Context.GetTime(), ObjectAndValue.Value);
+	}
+	
+	ExecutionTokens.Add(TCachedPropertyTrackExecutionToken<FString>());
+}
+
+FMovieSceneVectorPropertySectionTemplate::FMovieSceneVectorPropertySectionTemplate(const UMovieSceneVectorSection& Section, const UMovieScenePropertyTrack& Track)
+	: PropertyData(Track.GetPropertyName(), Track.GetPropertyPath())
+	, NumChannelsUsed(Section.GetChannelsUsed())
+{
+	for (int32 Index = 0; Index < NumChannelsUsed; ++Index)
+	{
+		ComponentCurves[Index] = Section.GetCurve(Index);
+	}
+}
+
+void FMovieSceneVectorPropertySectionTemplate::Setup(FPersistentEvaluationData& PersistentData, IMovieScenePlayer& Player) const
+{
+	switch (NumChannelsUsed)
+	{
+	case 2:	PropertyData.SetupCachedTrack<FVector2D>(PersistentData);	break;
+	case 3:	PropertyData.SetupCachedTrack<FVector>(PersistentData);	break;
+	case 4:	PropertyData.SetupCachedTrack<FVector4>(PersistentData);	break;
+
+	default:
+		UE_LOG(LogMovieScene, Warning, TEXT("Invalid number of channels(%d) for vector track"), NumChannelsUsed );
+		break;
+	}
+}
+
+void FMovieSceneVectorPropertySectionTemplate::Initialize(const FMovieSceneEvaluationOperand& Operand, const FMovieSceneContext& Context, FPersistentEvaluationData& PersistentData, IMovieScenePlayer& Player) const
+{
+	switch (NumChannelsUsed)
+	{
+	case 2:	PropertyData.SetupCachedFrame<FVector2D>(Operand, PersistentData, Player);	break;
+	case 3:	PropertyData.SetupCachedFrame<FVector>(Operand, PersistentData, Player);		break;
+	case 4:	PropertyData.SetupCachedFrame<FVector4>(Operand, PersistentData, Player);		break;
+
+	default:
+		UE_LOG(LogMovieScene, Warning, TEXT("Invalid number of channels(%d) for vector track"), NumChannelsUsed );
+		break;
+	}
+}
+
+void FMovieSceneVectorPropertySectionTemplate::Evaluate(const FMovieSceneEvaluationOperand& Operand, const FMovieSceneContext& Context, const FPersistentEvaluationData& PersistentData, FMovieSceneExecutionTokens& ExecutionTokens) const
+{
+	using namespace PropertyTemplate;
+
+	float Time = Context.GetTime();
+
+	switch (NumChannelsUsed)
+	{
+	case 2:
+		for (auto& ObjectAndValue : PersistentData.GetSectionData<TCachedSectionData<FVector2D>>().ObjectsAndValues)
+		{
+			ObjectAndValue.Value = FVector2D(ComponentCurves[0].Eval(Time, ObjectAndValue.Value.X), ComponentCurves[1].Eval(Time, ObjectAndValue.Value.Y));
+		}
+		ExecutionTokens.Add(TCachedPropertyTrackExecutionToken<FVector2D>());
+		break;
+
+	case 3:
+		for (auto& ObjectAndValue : PersistentData.GetSectionData<TCachedSectionData<FVector>>().ObjectsAndValues)
+		{
+			ObjectAndValue.Value = FVector(ComponentCurves[0].Eval(Time, ObjectAndValue.Value.X), ComponentCurves[1].Eval(Time, ObjectAndValue.Value.Y), ComponentCurves[2].Eval(Time, ObjectAndValue.Value.Z));
+		}
+		ExecutionTokens.Add(TCachedPropertyTrackExecutionToken<FVector>());
+		break;
+
+	case 4:
+		for (auto& ObjectAndValue : PersistentData.GetSectionData<TCachedSectionData<FVector4>>().ObjectsAndValues)
+		{
+			ObjectAndValue.Value = FVector4(ComponentCurves[0].Eval(Time, ObjectAndValue.Value.X), ComponentCurves[1].Eval(Time, ObjectAndValue.Value.Y), ComponentCurves[2].Eval(Time, ObjectAndValue.Value.Z), ComponentCurves[3].Eval(Time, ObjectAndValue.Value.W));
+		}
+		ExecutionTokens.Add(TCachedPropertyTrackExecutionToken<FVector4>());
+		break;
+
+	default:
+		UE_LOG(LogMovieScene, Warning, TEXT("Invalid number of channels(%d) for vector track"), NumChannelsUsed );
+		break;
+	}
+}
\ No newline at end of file
diff --git a/Engine/Source/Runtime/MovieSceneTracks/Private/Evaluation/MovieScenePropertyTemplates.h b/Engine/Source/Runtime/MovieSceneTracks/Private/Evaluation/MovieScenePropertyTemplates.h
new file mode 100644
--- /dev/null
+++ b/Engine/Source/Runtime/MovieSceneTracks/Private/Evaluation/MovieScenePropertyTemplates.h
@@ -0,0 +1,233 @@
+// Copyright 1998-2016 Epic Games, Inc. All Rights Reserved.
+
+#pragma once
+
+#include "MovieSceneFwd.h"
+#include "MovieScenePropertyTemplate.h"
+#include "Curves/RichCurve.h"
+#include "Curves/IntegralCurve.h"
+#include "Curves/StringCurve.h"
+
+#include "MovieScenePropertyTemplates.generated.h"
+
+class UMovieSceneBoolSection;
+class UMovieSceneFloatSection;
+class UMovieSceneByteSection;
+class UMovieSceneIntegerSection;
+class UMovieSceneVectorSection;
+class UMovieSceneStringSection;
+class UMovieScenePropertyTrack;
+
+
+USTRUCT()
+struct FMovieSceneBoolPropertySectionTemplate : public FMovieSceneEvalTemplate
+{
+	GENERATED_BODY()
+	
+	FMovieSceneBoolPropertySectionTemplate(){}
+	FMovieSceneBoolPropertySectionTemplate(const UMovieSceneBoolSection& Section, const UMovieScenePropertyTrack& Track);
+
+protected:
+
+	virtual UScriptStruct& GetScriptStructImpl() const override
+	{
+		return *StaticStruct();
+	}
+	virtual void SetupOverrides() override
+	{
+		EnableOverrides(RequiresSetupFlag | RequiresInitializeFlag);
+	}
+	virtual void Setup(FPersistentEvaluationData& PersistentData, IMovieScenePlayer& Player) const override
+	{
+		PropertyData.SetupCachedTrack<bool>(PersistentData);
+	}
+	virtual void Initialize(const FMovieSceneEvaluationOperand& Operand, const FMovieSceneContext& Context, FPersistentEvaluationData& PersistentData, IMovieScenePlayer& Player) const override
+	{
+		PropertyData.SetupCachedFrame<bool>(Operand, PersistentData, Player);
+	}
+
+	virtual void Evaluate(const FMovieSceneEvaluationOperand& Operand, const FMovieSceneContext& Context, const FPersistentEvaluationData& PersistentData, FMovieSceneExecutionTokens& ExecutionTokens) const override;
+
+	UPROPERTY()
+	FMovieScenePropertySectionData PropertyData;
+
+	UPROPERTY()
+	FIntegralCurve BoolCurve;
+};
+
+USTRUCT()
+struct FMovieSceneFloatPropertySectionTemplate : public FMovieSceneEvalTemplate
+{
+	GENERATED_BODY()
+	
+	FMovieSceneFloatPropertySectionTemplate(){}
+	FMovieSceneFloatPropertySectionTemplate(const UMovieSceneFloatSection& Section, const UMovieScenePropertyTrack& Track);
+
+protected:
+
+	virtual UScriptStruct& GetScriptStructImpl() const override
+	{
+		return *StaticStruct();
+	}
+	virtual void SetupOverrides() override
+	{
+		EnableOverrides(RequiresSetupFlag | RequiresInitializeFlag);
+	}
+	virtual void Setup(FPersistentEvaluationData& PersistentData, IMovieScenePlayer& Player) const override
+	{
+		PropertyData.SetupCachedTrack<float>(PersistentData);
+	}
+	virtual void Initialize(const FMovieSceneEvaluationOperand& Operand, const FMovieSceneContext& Context, FPersistentEvaluationData& PersistentData, IMovieScenePlayer& Player) const override
+	{
+		PropertyData.SetupCachedFrame<float>(Operand, PersistentData, Player);
+	}
+
+	virtual void Evaluate(const FMovieSceneEvaluationOperand& Operand, const FMovieSceneContext& Context, const FPersistentEvaluationData& PersistentData, FMovieSceneExecutionTokens& ExecutionTokens) const override;
+
+	UPROPERTY()
+	FMovieScenePropertySectionData PropertyData;
+
+	UPROPERTY()
+	FRichCurve FloatCurve;
+};
+
+USTRUCT()
+struct FMovieSceneBytePropertySectionTemplate : public FMovieSceneEvalTemplate
+{
+	GENERATED_BODY()
+	
+	FMovieSceneBytePropertySectionTemplate(){}
+	FMovieSceneBytePropertySectionTemplate(const UMovieSceneByteSection& Section, const UMovieScenePropertyTrack& Track);
+
+protected:
+
+	virtual UScriptStruct& GetScriptStructImpl() const override
+	{
+		return *StaticStruct();
+	}
+	virtual void SetupOverrides() override
+	{
+		EnableOverrides(RequiresSetupFlag | RequiresInitializeFlag);
+	}
+	virtual void Setup(FPersistentEvaluationData& PersistentData, IMovieScenePlayer& Player) const override
+	{
+		PropertyData.SetupCachedTrack<uint8>(PersistentData);
+	}
+	virtual void Initialize(const FMovieSceneEvaluationOperand& Operand, const FMovieSceneContext& Context, FPersistentEvaluationData& PersistentData, IMovieScenePlayer& Player) const override
+	{
+		PropertyData.SetupCachedFrame<uint8>(Operand, PersistentData, Player);
+	}
+
+	virtual void Evaluate(const FMovieSceneEvaluationOperand& Operand, const FMovieSceneContext& Context, const FPersistentEvaluationData& PersistentData, FMovieSceneExecutionTokens& ExecutionTokens) const override;
+
+	UPROPERTY()
+	FMovieScenePropertySectionData PropertyData;
+
+	UPROPERTY()
+	FIntegralCurve ByteCurve;
+};
+
+USTRUCT()
+struct FMovieSceneIntegerPropertySectionTemplate : public FMovieSceneEvalTemplate
+{
+	GENERATED_BODY()
+	
+	FMovieSceneIntegerPropertySectionTemplate(){}
+	FMovieSceneIntegerPropertySectionTemplate(const UMovieSceneIntegerSection& Section, const UMovieScenePropertyTrack& Track);
+
+protected:
+
+	virtual UScriptStruct& GetScriptStructImpl() const override
+	{
+		return *StaticStruct();
+	}
+	virtual void SetupOverrides() override
+	{
+		EnableOverrides(RequiresSetupFlag | RequiresInitializeFlag);
+	}
+	virtual void Setup(FPersistentEvaluationData& PersistentData, IMovieScenePlayer& Player) const override
+	{
+		PropertyData.SetupCachedTrack<int32>(PersistentData);
+	}
+	virtual void Initialize(const FMovieSceneEvaluationOperand& Operand, const FMovieSceneContext& Context, FPersistentEvaluationData& PersistentData, IMovieScenePlayer& Player) const override
+	{
+		PropertyData.SetupCachedFrame<int32>(Operand, PersistentData, Player);
+	}
+
+	virtual void Evaluate(const FMovieSceneEvaluationOperand& Operand, const FMovieSceneContext& Context, const FPersistentEvaluationData& PersistentData, FMovieSceneExecutionTokens& ExecutionTokens) const override;
+
+	UPROPERTY()
+	FMovieScenePropertySectionData PropertyData;
+
+	UPROPERTY()
+	FIntegralCurve IntegerCurve;
+};
+
+USTRUCT()
+struct FMovieSceneStringPropertySectionTemplate : public FMovieSceneEvalTemplate
+{
+	GENERATED_BODY()
+	
+	FMovieSceneStringPropertySectionTemplate(){}
+	FMovieSceneStringPropertySectionTemplate(const UMovieSceneStringSection& Section, const UMovieScenePropertyTrack& Track);
+
+protected:
+
+	virtual UScriptStruct& GetScriptStructImpl() const override
+	{
+		return *StaticStruct();
+	}
+	virtual void SetupOverrides() override
+	{
+		EnableOverrides(RequiresSetupFlag | RequiresInitializeFlag);
+	}
+	virtual void Setup(FPersistentEvaluationData& PersistentData, IMovieScenePlayer& Player) const override
+	{
+		PropertyData.SetupCachedTrack<FString>(PersistentData);
+	}
+	virtual void Initialize(const FMovieSceneEvaluationOperand& Operand, const FMovieSceneContext& Context, FPersistentEvaluationData& PersistentData, IMovieScenePlayer& Player) const override
+	{
+		PropertyData.SetupCachedFrame<FString>(Operand, PersistentData, Player);
+	}
+
+	virtual void Evaluate(const FMovieSceneEvaluationOperand& Operand, const FMovieSceneContext& Context, const FPersistentEvaluationData& PersistentData, FMovieSceneExecutionTokens& ExecutionTokens) const override;
+
+	UPROPERTY()
+	FMovieScenePropertySectionData PropertyData;
+
+	UPROPERTY()
+	FStringCurve StringCurve;
+};
+
+USTRUCT()
+struct FMovieSceneVectorPropertySectionTemplate : public FMovieSceneEvalTemplate
+{
+	GENERATED_BODY()
+	
+	FMovieSceneVectorPropertySectionTemplate(){}
+	FMovieSceneVectorPropertySectionTemplate(const UMovieSceneVectorSection& Section, const UMovieScenePropertyTrack& Track);
+
+protected:
+
+	virtual UScriptStruct& GetScriptStructImpl() const override
+	{
+		return *StaticStruct();
+	}
+	virtual void SetupOverrides() override
+	{
+		EnableOverrides(RequiresSetupFlag | RequiresInitializeFlag);
+	}
+
+	virtual void Setup(FPersistentEvaluationData& PersistentData, IMovieScenePlayer& Player) const override;
+	virtual void Initialize(const FMovieSceneEvaluationOperand& Operand, const FMovieSceneContext& Context, FPersistentEvaluationData& PersistentData, IMovieScenePlayer& Player) const override;
+	virtual void Evaluate(const FMovieSceneEvaluationOperand& Operand, const FMovieSceneContext& Context, const FPersistentEvaluationData& PersistentData, FMovieSceneExecutionTokens& ExecutionTokens) const override;
+
+	UPROPERTY()
+	FMovieScenePropertySectionData PropertyData;
+
+	UPROPERTY()
+	FRichCurve ComponentCurves[4];
+
+	UPROPERTY()
+	int32 NumChannelsUsed;
+};
\ No newline at end of file
diff --git a/Engine/Source/Runtime/MovieSceneTracks/Private/Evaluation/MovieSceneSkeletalAnimationTemplate.cpp b/Engine/Source/Runtime/MovieSceneTracks/Private/Evaluation/MovieSceneSkeletalAnimationTemplate.cpp
new file mode 100644
--- /dev/null
+++ b/Engine/Source/Runtime/MovieSceneTracks/Private/Evaluation/MovieSceneSkeletalAnimationTemplate.cpp
@@ -0,0 +1,284 @@
+// Copyright 1998-2016 Epic Games, Inc. All Rights Reserved.
+
+#include "MovieSceneTracksPrivatePCH.h"
+#include "MovieSceneSkeletalAnimationTemplate.h"
+
+#include "MovieSceneExecutionToken.h"
+#include "MovieSceneSkeletalAnimationSection.h"
+#include "IMovieScenePlayer.h"
+#include "Animation/AnimSequenceBase.h"
+#include "MovieSceneCompilerRules.h"
+#include "Animation/AnimInstance.h"
+#include "Animation/AnimMontage.h"
+#include "Animation/AnimSingleNodeInstance.h"
+
+
+bool ShouldUsePreviewPlayback(IMovieScenePlayer& Player, UObject& RuntimeObject)
+{
+	// we also use PreviewSetAnimPosition in PIE when not playing, as we can preview in PIE
+	bool bIsNotInPIEOrNotPlaying = (RuntimeObject.GetWorld() && !RuntimeObject.GetWorld()->HasBegunPlay()) || Player.GetPlaybackStatus() != EMovieScenePlayerStatus::Playing;
+	return GIsEditor && bIsNotInPIEOrNotPlaying;
+}
+
+bool CanPlayAnimation(USkeletalMeshComponent* SkeletalMeshComponent, UAnimSequenceBase* AnimAssetBase = nullptr)
+{
+	return (SkeletalMeshComponent->SkeletalMesh && SkeletalMeshComponent->SkeletalMesh->Skeleton && 
+		(!AnimAssetBase || SkeletalMeshComponent->SkeletalMesh->Skeleton->IsCompatible(AnimAssetBase->GetSkeleton())));
+}
+
+struct FAnimationTrackTrackState : IPersistentEvaluationData
+{
+	struct FComponentAndMontage
+	{
+		FComponentAndMontage(USkeletalMeshComponent* InComponent) : Component(InComponent) {}
+
+		TWeakObjectPtr<UAnimMontage> Montage;
+		TWeakObjectPtr<USkeletalMeshComponent> Component;
+	};
+
+	TMap<FObjectKey, FComponentAndMontage> CurrentlyPlayingMontages;
+
+	~FAnimationTrackTrackState()
+	{
+		for (auto& Pair : CurrentlyPlayingMontages)
+		{
+			USkeletalMeshComponent* Component = Pair.Value.Component.Get();
+
+			if (Component)
+			{
+				FinishAnimControl(Component);
+			}
+		}
+	}
+
+	void EnsureAnimControl(USkeletalMeshComponent* SkeletalMeshComponent, UAnimSequenceBase* InAnimSequence, IMovieScenePlayer& Player)
+	{
+		SkeletalMeshComponent->MeshComponentUpdateFlag = EMeshComponentUpdateFlag::AlwaysTickPoseAndRefreshBones;
+
+		if (!CurrentlyPlayingMontages.Contains(SkeletalMeshComponent))
+		{
+			BeginAnimControl(SkeletalMeshComponent, InAnimSequence);
+		}
+	}
+
+private:
+
+	void BeginAnimControl(USkeletalMeshComponent* SkeletalMeshComponent, UAnimSequenceBase* InAnimSequence)
+	{
+		if (CanPlayAnimation(SkeletalMeshComponent, InAnimSequence))
+		{
+			UAnimInstance* AnimInstance = SkeletalMeshComponent->GetAnimInstance();
+			if (!AnimInstance)
+			{
+				SkeletalMeshComponent->SetAnimationMode(EAnimationMode::Type::AnimationSingleNode);
+			}
+		}
+
+		CurrentlyPlayingMontages.Add(SkeletalMeshComponent, FComponentAndMontage(SkeletalMeshComponent));
+	}
+
+	void FinishAnimControl(USkeletalMeshComponent* SkeletalMeshComponent)
+	{
+		if (SkeletalMeshComponent->GetAnimationMode() == EAnimationMode::Type::AnimationBlueprint)
+		{
+			UAnimInstance* AnimInstance = SkeletalMeshComponent->GetAnimInstance();
+			if(AnimInstance)
+			{
+				AnimInstance->Montage_Stop(0.f);
+				AnimInstance->UpdateAnimation(0.f, false);
+			}
+
+			// Update space bases to reset it back to ref pose
+			SkeletalMeshComponent->RefreshBoneTransforms();
+			SkeletalMeshComponent->RefreshSlaveComponents();
+			SkeletalMeshComponent->UpdateComponentToWorld();
+		}
+	}
+};
+
+struct FPreAnimatedAnimationTokenProducer : IMovieScenePreAnimatedTokenProducer
+{
+	virtual IMovieScenePreAnimatedTokenPtr CacheExistingState(UObject& Object) const
+	{
+		struct FPreAnimatedToken : IMovieScenePreAnimatedToken
+		{
+			EMeshComponentUpdateFlag::Type MeshComponentUpdateFlag;
+
+			virtual void RestoreState(UObject& InObject, IMovieScenePlayer& Player)
+			{
+				CastChecked<USkeletalMeshComponent>(&InObject)->MeshComponentUpdateFlag = MeshComponentUpdateFlag;
+			}
+		};
+
+		FPreAnimatedToken Token;
+		Token.MeshComponentUpdateFlag = CastChecked<USkeletalMeshComponent>(&Object)->MeshComponentUpdateFlag;
+		return Token;
+	}
+};
+
+struct FAnimationTrackTrackToken : IMovieSceneExecutionToken
+{
+	FAnimationTrackTrackState* TrackState;
+
+	FAnimationTrackTrackToken(const FMovieSceneSkeletalAnimationSectionTemplateParameters& InParams)
+		: Params(InParams)
+	{
+	}
+
+	virtual void Execute(const FMovieSceneContext& Context, const FMovieSceneEvaluationOperand& Operand, FPersistentEvaluationData& PersistentData, IMovieScenePlayer& Player) override
+	{
+		TrackState = &PersistentData.GetOrAddTrackData<FAnimationTrackTrackState>();
+
+		for (TWeakObjectPtr<> Object : Player.FindBoundObjects(Operand))
+		{
+			UObject* ObjectPtr = Object.Get();
+			if (!ObjectPtr)
+			{
+				continue;
+			}
+
+			// first try to control the component directly
+			USkeletalMeshComponent* SkeletalMeshComponent = Cast<USkeletalMeshComponent>(ObjectPtr);
+			if (!SkeletalMeshComponent)
+			{
+				// then check to see if we are controlling an actor & if so use its first USkeletalMeshComponent 
+				if (AActor* Actor = Cast<AActor>(ObjectPtr))
+				{
+					SkeletalMeshComponent = Actor->FindComponentByClass<USkeletalMeshComponent>();
+				}
+			}
+
+			UAnimSequenceBase* AnimSequence = Params.Animation;
+
+			if (!SkeletalMeshComponent || !AnimSequence)
+			{
+				continue;
+			}
+
+			const float EvalTime = Params.MapTimeToAnimation(Context.GetTime());
+
+			Player.SavePreAnimatedState(*SkeletalMeshComponent, TMovieSceneAnimTypeID<FAnimationTrackTrackToken>(), FPreAnimatedAnimationTokenProducer());
+
+			TrackState->EnsureAnimControl(SkeletalMeshComponent, AnimSequence, Player);
+
+			static const int32 ChannelIndex = 0;
+			static const bool bLooping = false;
+
+			if (ShouldUsePreviewPlayback(Player, *SkeletalMeshComponent))
+			{
+				const EMovieScenePlayerStatus::Type PlayerStatus = Player.GetPlaybackStatus();
+
+				// If the playback status is jumping, ie. one such occurrence is setting the time for thumbnail generation, disable anim notifies updates because it could fire audio
+				const bool bFireNotifies = PlayerStatus != EMovieScenePlayerStatus::Jumping && PlayerStatus != EMovieScenePlayerStatus::Stopped;
+
+				// When jumping from one cut to another cut, the delta time should be 0 so that anim notifies before the current position are not evaluated. Note, anim notifies at the current time should still be evaluated.
+				const float DeltaTime = Context.HasJumped() ? 0.f : Context.GetRange().Size<float>();
+
+				const bool bResetDynamics = PlayerStatus == EMovieScenePlayerStatus::Stepping || 
+											PlayerStatus == EMovieScenePlayerStatus::Jumping || 
+											PlayerStatus == EMovieScenePlayerStatus::Scrubbing || 
+											(DeltaTime == 0.0f && PlayerStatus != EMovieScenePlayerStatus::Stopped); 
+			
+				PreviewSetAnimPosition(SkeletalMeshComponent, Params.SlotName, ChannelIndex, AnimSequence, EvalTime, bLooping, bFireNotifies, DeltaTime, Player.GetPlaybackStatus() == EMovieScenePlayerStatus::Playing, bResetDynamics);
+			}
+			else
+			{
+				// Don't fire notifies at runtime since they will be fired through the standard animation tick path.
+				const bool bFireNotifies = false;
+				SetAnimPosition(SkeletalMeshComponent, Params.SlotName, ChannelIndex, AnimSequence, EvalTime, bLooping, bFireNotifies);
+			}
+		}
+	}
+
+	void SetAnimPosition(USkeletalMeshComponent* SkeletalMeshComponent, FName SlotName, int32 ChannelIndex, UAnimSequenceBase* InAnimSequence, float InPosition, bool bLooping, bool bFireNotifies)
+	{
+		if (CanPlayAnimation(SkeletalMeshComponent, InAnimSequence))
+		{
+			UAnimMontage* Montage = FAnimMontageInstance::SetMatineeAnimPositionInner(SlotName, SkeletalMeshComponent, InAnimSequence, InPosition, bLooping);
+
+			// Ensure the sequence is not stopped
+			UAnimInstance* AnimInst = SkeletalMeshComponent->GetAnimInstance();
+			UAnimSingleNodeInstance* SingleNodeInst = SkeletalMeshComponent->GetSingleNodeInstance();
+			if(SingleNodeInst)
+			{
+				SingleNodeInst->SetPlaying(true);
+			}
+			else if (AnimInst && Montage)
+			{
+				AnimInst->Montage_Resume(Montage);
+			}
+		}
+	}
+
+	void PreviewSetAnimPosition(USkeletalMeshComponent* SkeletalMeshComponent, FName SlotName, int32 ChannelIndex, UAnimSequenceBase* InAnimSequence, float InPosition, bool bLooping, bool bFireNotifies, float DeltaTime, bool bPlaying, bool bResetDynamics)
+	{
+		if(!CanPlayAnimation(SkeletalMeshComponent, InAnimSequence))
+		{
+			return;
+		}
+
+		UAnimMontage* Montage = FAnimMontageInstance::PreviewMatineeSetAnimPositionInner(SlotName, SkeletalMeshComponent, InAnimSequence, InPosition, bLooping, bFireNotifies, DeltaTime);
+
+		// if we are not playing, make sure we dont continue (as skeletal meshes can still tick us onwards)
+		UAnimInstance* AnimInst = SkeletalMeshComponent->GetAnimInstance();
+		UAnimSingleNodeInstance * SingleNodeInst = SkeletalMeshComponent->GetSingleNodeInstance();
+		if(SingleNodeInst)
+		{
+			SingleNodeInst->SetPlaying(bPlaying);
+		}
+		else if (AnimInst)
+		{
+			if(Montage)
+			{
+				if(bPlaying)
+				{
+					AnimInst->Montage_Resume(Montage);
+				}
+				else
+				{
+					AnimInst->Montage_Pause(Montage);
+				}
+			}
+
+			if(bResetDynamics)
+			{
+				// make sure we reset any simulations
+				AnimInst->ResetDynamics();
+			}
+		}
+	}
+
+	FMovieSceneSkeletalAnimationSectionTemplateParameters Params;
+};
+
+
+FMovieSceneSkeletalAnimationSectionTemplate::FMovieSceneSkeletalAnimationSectionTemplate(const UMovieSceneSkeletalAnimationSection& InSection)
+	: Params(InSection.Params, InSection.GetStartTime())
+{
+}
+
+void FMovieSceneSkeletalAnimationSectionTemplate::Evaluate(const FMovieSceneEvaluationOperand& Operand, const FMovieSceneContext& Context, const FPersistentEvaluationData& PersistentData, FMovieSceneExecutionTokens& ExecutionTokens) const
+{
+	ExecutionTokens.Add(FAnimationTrackTrackToken(Params));
+}
+
+float FMovieSceneSkeletalAnimationSectionTemplateParameters::MapTimeToAnimation(float ThisPosition) const
+{
+	// @todo: Sequencer: what is this for??
+	//ThisPosition -= 1 / 1000.0f;
+
+	const float SectionPlayRate = PlayRate;
+	const float AnimPlayRate = FMath::IsNearlyZero(SectionPlayRate) ? 1.0f : SectionPlayRate;
+
+	const float SeqLength = GetSequenceLength() - (StartOffset + EndOffset);
+
+	ThisPosition = (ThisPosition - SectionStartTime) * AnimPlayRate;
+	ThisPosition = FMath::Fmod(ThisPosition, SeqLength);
+	ThisPosition += StartOffset;
+	if (bReverse)
+	{
+		ThisPosition = (SeqLength - (ThisPosition - StartOffset)) + StartOffset;
+	}
+
+	return ThisPosition;
+}
\ No newline at end of file
diff --git a/Engine/Source/Runtime/MovieSceneTracks/Private/Evaluation/MovieSceneSkeletalAnimationTemplate.h b/Engine/Source/Runtime/MovieSceneTracks/Private/Evaluation/MovieSceneSkeletalAnimationTemplate.h
new file mode 100644
--- /dev/null
+++ b/Engine/Source/Runtime/MovieSceneTracks/Private/Evaluation/MovieSceneSkeletalAnimationTemplate.h
@@ -0,0 +1,41 @@
+// Copyright 1998-2016 Epic Games, Inc. All Rights Reserved.
+
+#pragma once
+
+#include "MovieSceneEvalTemplate.h"
+#include "MovieSceneSkeletalAnimationSection.h"
+#include "MovieSceneSkeletalAnimationTemplate.generated.h"
+
+class UMovieSceneSkeletalAnimationSection;
+
+USTRUCT()
+struct FMovieSceneSkeletalAnimationSectionTemplateParameters : public FMovieSceneSkeletalAnimationParams
+{
+	GENERATED_BODY()
+
+	FMovieSceneSkeletalAnimationSectionTemplateParameters() {}
+	FMovieSceneSkeletalAnimationSectionTemplateParameters(const FMovieSceneSkeletalAnimationParams& BaseParams, float InSectionStartTime)
+		: FMovieSceneSkeletalAnimationParams(BaseParams)
+		, SectionStartTime(InSectionStartTime)
+	{}
+
+	float MapTimeToAnimation(float InPosition) const;
+
+	UPROPERTY()
+	float SectionStartTime;
+};
+
+USTRUCT()
+struct FMovieSceneSkeletalAnimationSectionTemplate : public FMovieSceneEvalTemplate
+{
+	GENERATED_BODY()
+	
+	FMovieSceneSkeletalAnimationSectionTemplate() {}
+	FMovieSceneSkeletalAnimationSectionTemplate(const UMovieSceneSkeletalAnimationSection& Section);
+
+	virtual UScriptStruct& GetScriptStructImpl() const override { return *StaticStruct(); }
+	virtual void Evaluate(const FMovieSceneEvaluationOperand& Operand, const FMovieSceneContext& Context, const FPersistentEvaluationData& PersistentData, FMovieSceneExecutionTokens& ExecutionTokens) const override;
+
+	UPROPERTY()
+	FMovieSceneSkeletalAnimationSectionTemplateParameters Params;
+};
\ No newline at end of file
diff --git a/Engine/Source/Runtime/MovieSceneTracks/Private/Evaluation/MovieSceneSlomoTemplate.cpp b/Engine/Source/Runtime/MovieSceneTracks/Private/Evaluation/MovieSceneSlomoTemplate.cpp
new file mode 100644
--- /dev/null
+++ b/Engine/Source/Runtime/MovieSceneTracks/Private/Evaluation/MovieSceneSlomoTemplate.cpp
@@ -0,0 +1,113 @@
+// Copyright 1998-2016 Epic Games, Inc. All Rights Reserved.
+
+#include "MovieSceneTracksPrivatePCH.h"
+#include "MovieSceneSlomoTemplate.h"
+#include "MovieSceneSlomoTrack.h"
+#include "MovieSceneSlomoSection.h"
+#include "MovieSceneSequence.h"
+#include "MovieSceneEvaluationTemplateInstance.h"
+
+DECLARE_CYCLE_STAT(TEXT("Slomo Track Token Execute"), MovieSceneEval_SlomoTrack_TokenExecute, STATGROUP_MovieSceneEval);
+
+struct FSlomoTrackToken
+{
+	FSlomoTrackToken(float InSlomoValue)
+		: SlomoValue(InSlomoValue)
+	{}
+
+	float SlomoValue;
+
+	void Apply(IMovieScenePlayer& Player)
+	{
+		UObject* PlaybackContext = Player.GetPlaybackContext();
+		UWorld* World = PlaybackContext ? PlaybackContext->GetWorld() : nullptr;
+
+		if (!GIsEditor || 
+			(World && World->GetWorld()->GetNetMode() == NM_Client) ||
+			!GEngine ||
+			SlomoValue <= 0.f)
+		{
+			return;
+		}
+
+		AWorldSettings* WorldSettings = World ? World->GetWorldSettings() : nullptr;
+			
+		if (WorldSettings)
+		{	
+			WorldSettings->MatineeTimeDilation = SlomoValue;
+			WorldSettings->ForceNetUpdate();
+		}
+	}
+};
+
+struct FSlomoTrackData : IPersistentEvaluationData
+{
+	TOptional<FSlomoTrackToken> PreviousSlomoValue;
+};
+
+struct FSlomoPreAnimatedGlobalToken : FSlomoTrackToken, IMovieScenePreAnimatedGlobalToken
+{
+	FSlomoPreAnimatedGlobalToken(float InSlomoValue)
+		: FSlomoTrackToken(InSlomoValue)
+	{}
+
+	virtual void RestoreState(IMovieScenePlayer& Player) override
+	{
+		Apply(Player);
+	}
+};
+
+struct FSlomoPreAnimatedGlobalTokenProducer : IMovieScenePreAnimatedGlobalTokenProducer
+{
+	FSlomoPreAnimatedGlobalTokenProducer(IMovieScenePlayer& InPlayer) 
+		: Player(InPlayer)
+	{}
+
+	virtual IMovieScenePreAnimatedGlobalTokenPtr CacheExistingState() const override
+	{
+		if (AWorldSettings* WorldSettings = Player.GetPlaybackContext()->GetWorld()->GetWorldSettings())
+		{
+			return FSlomoPreAnimatedGlobalToken(WorldSettings->MatineeTimeDilation);
+		}
+		return IMovieScenePreAnimatedGlobalTokenPtr();
+	}
+
+	IMovieScenePlayer& Player;
+};
+
+/** A movie scene execution token that applies slomo */
+struct FSlomoExecutionToken : IMovieSceneExecutionToken, FSlomoTrackToken
+{
+	FSlomoExecutionToken(float InSlomoValue)
+		: FSlomoTrackToken(InSlomoValue)
+	{}
+
+	static FMovieSceneAnimTypeID GetAnimTypeID()
+	{
+		return TMovieSceneAnimTypeID<FSlomoExecutionToken>();
+	}
+	
+	/** Execute this token, operating on all objects referenced by 'Operand' */
+	virtual void Execute(const FMovieSceneContext& Context, const FMovieSceneEvaluationOperand& Operand, FPersistentEvaluationData& PersistentData, IMovieScenePlayer& Player) override
+	{
+		MOVIESCENE_DETAILED_SCOPE_CYCLE_COUNTER(MovieSceneEval_SlomoTrack_TokenExecute)
+
+		Player.SavePreAnimatedState(GetAnimTypeID(), FSlomoPreAnimatedGlobalTokenProducer(Player));
+		
+		Apply(Player);
+	}
+};
+
+FMovieSceneSlomoSectionTemplate::FMovieSceneSlomoSectionTemplate(const UMovieSceneSlomoSection& Section)
+	: SlomoCurve(Section.GetFloatCurve())
+{
+}
+
+void FMovieSceneSlomoSectionTemplate::Evaluate(const FMovieSceneEvaluationOperand& Operand, const FMovieSceneContext& Context, const FPersistentEvaluationData& PersistentData, FMovieSceneExecutionTokens& ExecutionTokens) const
+{
+	float SlomoValue = SlomoCurve.Eval(Context.GetTime());
+	if (SlomoValue >= 0.f)
+	{
+		ExecutionTokens.Add(FSlomoExecutionToken(SlomoValue));
+	}
+}
\ No newline at end of file
diff --git a/Engine/Source/Runtime/MovieSceneTracks/Private/Evaluation/MovieSceneSlomoTemplate.h b/Engine/Source/Runtime/MovieSceneTracks/Private/Evaluation/MovieSceneSlomoTemplate.h
new file mode 100644
--- /dev/null
+++ b/Engine/Source/Runtime/MovieSceneTracks/Private/Evaluation/MovieSceneSlomoTemplate.h
@@ -0,0 +1,30 @@
+// Copyright 1998-2016 Epic Games, Inc. All Rights Reserved.
+
+#pragma once
+
+#include "MovieSceneFwd.h"
+#include "MovieSceneTrackImplementation.h"
+#include "MovieSceneEvalTemplate.h"
+#include "MovieSceneSlomoSection.h"
+#include "Curves/RichCurve.h"
+
+#include "MovieSceneSlomoTemplate.generated.h"
+
+class UMovieSceneSlomoTrack;
+
+USTRUCT()
+struct FMovieSceneSlomoSectionTemplate : public FMovieSceneEvalTemplate
+{
+	GENERATED_BODY()
+	
+	FMovieSceneSlomoSectionTemplate() {}
+	FMovieSceneSlomoSectionTemplate(const UMovieSceneSlomoSection& Section);
+
+private:
+
+	virtual UScriptStruct& GetScriptStructImpl() const override { return *StaticStruct(); }
+	virtual void Evaluate(const FMovieSceneEvaluationOperand& Operand, const FMovieSceneContext& Context, const FPersistentEvaluationData& PersistentData, FMovieSceneExecutionTokens& ExecutionTokens) const override;
+
+	UPROPERTY()
+	FRichCurve SlomoCurve;
+};
diff --git a/Engine/Source/Runtime/MovieSceneTracks/Private/Evaluation/MovieSceneSpawnTemplate.cpp b/Engine/Source/Runtime/MovieSceneTracks/Private/Evaluation/MovieSceneSpawnTemplate.cpp
new file mode 100644
--- /dev/null
+++ b/Engine/Source/Runtime/MovieSceneTracks/Private/Evaluation/MovieSceneSpawnTemplate.cpp
@@ -0,0 +1,88 @@
+// Copyright 1998-2016 Epic Games, Inc. All Rights Reserved.
+
+#include "MovieSceneTracksPrivatePCH.h"
+#include "MovieSceneSpawnTemplate.h"
+
+#include "MovieSceneSequence.h"
+#include "MovieSceneSpawnSection.h"
+
+DECLARE_CYCLE_STAT(TEXT("Spawn Track Evaluate"), MovieSceneEval_SpawnTrack_Evaluate, STATGROUP_MovieSceneEval);
+DECLARE_CYCLE_STAT(TEXT("Spawn Track Token Execute"), MovieSceneEval_SpawnTrack_TokenExecute, STATGROUP_MovieSceneEval);
+
+/** A movie scene pre-animated token that stores a pre-animated transform */
+struct FSpawnTrackPreAnimatedTokenProducer : IMovieScenePreAnimatedTokenProducer
+{
+	virtual IMovieScenePreAnimatedTokenPtr CacheExistingState(UObject& Object) const
+	{
+		struct FToken : IMovieScenePreAnimatedToken
+		{
+			virtual void RestoreState(UObject& InObject, IMovieScenePlayer& Player) override
+			{
+				Player.GetSpawnRegister().DestroyObjectDirectly(InObject);
+			}
+		};
+		
+		return FToken();
+	}
+};
+
+struct FSpawnObjectToken : IMovieSceneExecutionToken
+{
+	FSpawnObjectToken(bool bInSpawned)
+		: bSpawned(bInSpawned)
+	{}
+
+	virtual void Execute(const FMovieSceneContext& Context, const FMovieSceneEvaluationOperand& Operand, FPersistentEvaluationData& PersistentData, IMovieScenePlayer& Player) override
+	{
+		MOVIESCENE_DETAILED_SCOPE_CYCLE_COUNTER(MovieSceneEval_SpawnTrack_TokenExecute)
+
+		const bool bHasBoundObjects = Player.FindBoundObjects(Operand).Num() != 0;
+		if (bSpawned)
+		{
+			// If it's not spawned, spawn it
+			if (!bHasBoundObjects)
+			{
+				const UMovieSceneSequence* Sequence = Player.State.FindSequence(Operand.SequenceID);
+				if (Sequence)
+				{
+					Player.GetSpawnRegister().SpawnObject(Operand.ObjectBindingID, *Sequence->GetMovieScene(), Operand.SequenceID, Player);
+				}
+			}
+
+			// ensure that pre animated state is saved
+			for (TWeakObjectPtr<> Object : Player.FindBoundObjects(Operand))
+			{
+				if (UObject* ObjectPtr = Object.Get())
+				{
+					Player.SavePreAnimatedState(*ObjectPtr, FMovieSceneSpawnSectionTemplate::GetAnimTypeID(), FSpawnTrackPreAnimatedTokenProducer());
+				}
+			}
+		}
+		else if (!bSpawned && bHasBoundObjects)
+		{
+			Player.GetSpawnRegister().DestroySpawnedObject(Operand.ObjectBindingID, Operand.SequenceID, Player);
+		}
+	}
+
+	bool bSpawned;
+};
+
+FMovieSceneSpawnSectionTemplate::FMovieSceneSpawnSectionTemplate(const UMovieSceneSpawnSection& SpawnSection)
+	: Curve(SpawnSection.GetCurve())
+{
+}
+
+void FMovieSceneSpawnSectionTemplate::Evaluate(const FMovieSceneEvaluationOperand& Operand, const FMovieSceneContext& Context, const FPersistentEvaluationData& PersistentData, FMovieSceneExecutionTokens& ExecutionTokens) const
+{
+	MOVIESCENE_DETAILED_SCOPE_CYCLE_COUNTER(MovieSceneEval_SpawnTrack_Evaluate)
+
+	const float Time = Context.GetTime();
+	ExecutionTokens.Add(
+		FSpawnObjectToken(Curve.Evaluate(Time) != 0)
+		);
+}
+
+FMovieSceneAnimTypeID FMovieSceneSpawnSectionTemplate::GetAnimTypeID()
+{
+	return TMovieSceneAnimTypeID<FMovieSceneSpawnSectionTemplate>();
+}
\ No newline at end of file
diff --git a/Engine/Source/Runtime/MovieSceneTracks/Private/Evaluation/MovieSceneVisibilityTemplate.cpp b/Engine/Source/Runtime/MovieSceneTracks/Private/Evaluation/MovieSceneVisibilityTemplate.cpp
new file mode 100644
--- /dev/null
+++ b/Engine/Source/Runtime/MovieSceneTracks/Private/Evaluation/MovieSceneVisibilityTemplate.cpp
@@ -0,0 +1,110 @@
+// Copyright 1998-2016 Epic Games, Inc. All Rights Reserved.
+
+#include "MovieSceneTracksPrivatePCH.h"
+#include "MovieScenePropertyTemplates.h"
+#include "MovieSceneVisibilityTemplate.h"
+#include "MovieSceneBoolSection.h"
+#include "MovieScenePropertyTrack.h"
+#include "MovieSceneExecutionToken.h"
+
+DECLARE_CYCLE_STAT(TEXT("Visibility Track Evaluate"), MovieSceneEval_VisibilityTrack_Evaluate, STATGROUP_MovieSceneEval);
+DECLARE_CYCLE_STAT(TEXT("Visibility Track Token Execute"), MovieSceneEval_VisibilityTrack_TokenExecute, STATGROUP_MovieSceneEval);
+
+/** A movie scene pre-animated token that stores a pre-animated actor's temporarily hidden in game */
+struct FTemporarilyHiddenInGamePreAnimatedToken : IMovieScenePreAnimatedToken
+{
+	FTemporarilyHiddenInGamePreAnimatedToken(bool bInHidden, bool bInTemporarilyHiddenInGame)
+		: bHidden(bInHidden)
+		, bTemporarilyHiddenInGame(bInTemporarilyHiddenInGame)
+	{}
+
+	virtual void RestoreState(UObject& InObject, IMovieScenePlayer& Player) override
+	{
+		AActor* Actor = CastChecked<AActor>(&InObject);
+		Actor->SetActorHiddenInGame(bHidden);
+
+#if WITH_EDITOR
+		Actor->SetIsTemporarilyHiddenInEditor(bTemporarilyHiddenInGame);
+#endif // WITH_EDITOR
+	}
+
+	bool bHidden;
+	bool bTemporarilyHiddenInGame;
+};
+
+struct FTemporarilyHiddenInGameTokenProducer : IMovieScenePreAnimatedTokenProducer
+{
+	/** Cache the existing state of an object before moving it */
+	virtual IMovieScenePreAnimatedTokenPtr CacheExistingState(UObject& InObject) const override
+	{
+		AActor* Actor = CastChecked<AActor>(&InObject);
+		bool bInTemporarilyHiddenInGame = 
+#if WITH_EDITOR
+			Actor->IsTemporarilyHiddenInEditor();
+#else
+			false;
+#endif
+		return FTemporarilyHiddenInGamePreAnimatedToken(Actor->bHidden, bInTemporarilyHiddenInGame);
+	}
+};
+
+/** A movie scene execution token that stores temporarily hidden in game */
+struct FTemporarilyHiddenInGameExecutionToken
+	: IMovieSceneExecutionToken
+{
+	static FMovieSceneAnimTypeID GetAnimTypeID()
+	{
+		return TMovieSceneAnimTypeID<FTemporarilyHiddenInGameExecutionToken>();
+	}
+	
+	/** Execute this token, operating on all objects referenced by 'Operand' */
+	virtual void Execute(const FMovieSceneContext& Context, const FMovieSceneEvaluationOperand& Operand, FPersistentEvaluationData& PersistentData, IMovieScenePlayer& Player) override
+	{
+		MOVIESCENE_DETAILED_SCOPE_CYCLE_COUNTER(MovieSceneEval_VisibilityTrack_TokenExecute);
+
+		PropertyTemplate::TCachedSectionData<bool>& PropertyTrackData = PersistentData.GetSectionData<PropertyTemplate::TCachedSectionData<bool>>();
+
+		for (PropertyTemplate::TCachedValue<bool>& ObjectAndValue : PropertyTrackData.ObjectsAndValues)
+		{
+			if (UObject* ObjectPtr = ObjectAndValue.WeakObject.Get())
+			{
+				AActor* Actor = Cast<AActor>(ObjectPtr);
+				if (Actor)
+				{
+					Player.SavePreAnimatedState(*Actor, GetAnimTypeID(), FTemporarilyHiddenInGameTokenProducer());
+
+					Actor->SetActorHiddenInGame(ObjectAndValue.Value);
+
+#if WITH_EDITOR
+					if (GIsEditor && Actor->GetWorld() != nullptr && !Actor->GetWorld()->IsPlayInEditor())
+					{
+						Actor->SetIsTemporarilyHiddenInEditor(ObjectAndValue.Value);
+					}
+#endif // WITH_EDITOR
+				}
+			}
+		}
+	}
+};
+
+FMovieSceneVisibilitySectionTemplate::FMovieSceneVisibilitySectionTemplate(const UMovieSceneBoolSection& Section, const UMovieScenePropertyTrack& Track)
+	: FMovieSceneBoolPropertySectionTemplate(Section, Track)
+{}
+
+void FMovieSceneVisibilitySectionTemplate::Evaluate(const FMovieSceneEvaluationOperand& Operand, const FMovieSceneContext& Context, const FPersistentEvaluationData& PersistentData, FMovieSceneExecutionTokens& ExecutionTokens) const
+{
+	MOVIESCENE_DETAILED_SCOPE_CYCLE_COUNTER(MovieSceneEval_VisibilityTrack_Evaluate);
+
+	using namespace PropertyTemplate;
+
+	for (TCachedValue<bool>& ObjectAndValue : PersistentData.GetSectionData<TCachedSectionData<bool>>().ObjectsAndValues)
+	{
+		ObjectAndValue.Value = !!BoolCurve.Evaluate(Context.GetTime(), ObjectAndValue.Value ? 1 : 0);
+
+		// Invert this evaluation since the property is "bHiddenInGame" and we want the visualization to be the inverse of that. Green means visible.
+		ObjectAndValue.Value = !ObjectAndValue.Value;
+	}
+	
+	ExecutionTokens.Add(FTemporarilyHiddenInGameExecutionToken());
+}
+
diff --git a/Engine/Source/Runtime/MovieSceneTracks/Private/Evaluation/MovieSceneVisibilityTemplate.h b/Engine/Source/Runtime/MovieSceneTracks/Private/Evaluation/MovieSceneVisibilityTemplate.h
new file mode 100644
--- /dev/null
+++ b/Engine/Source/Runtime/MovieSceneTracks/Private/Evaluation/MovieSceneVisibilityTemplate.h
@@ -0,0 +1,28 @@
+// Copyright 1998-2016 Epic Games, Inc. All Rights Reserved.
+
+#pragma once
+
+#include "MovieSceneFwd.h"
+#include "MovieScenePropertyTemplates.h"
+
+#include "MovieSceneVisibilityTemplate.generated.h"
+
+class UMovieSceneBoolSection;
+
+
+USTRUCT()
+struct FMovieSceneVisibilitySectionTemplate : public FMovieSceneBoolPropertySectionTemplate
+{
+	GENERATED_BODY()
+	
+	FMovieSceneVisibilitySectionTemplate(){}
+	FMovieSceneVisibilitySectionTemplate(const UMovieSceneBoolSection& Section, const UMovieScenePropertyTrack& Track);
+
+	/** Temporarily hidden in game */
+	UPROPERTY()
+	bool bTemporarilyHiddenInGame;
+
+private:
+	virtual UScriptStruct& GetScriptStructImpl() const override { return *StaticStruct(); }
+	virtual void Evaluate(const FMovieSceneEvaluationOperand& Operand, const FMovieSceneContext& Context, const FPersistentEvaluationData& PersistentData, FMovieSceneExecutionTokens& ExecutionTokens) const override;
+};
diff --git a/Engine/Source/Runtime/MovieSceneTracks/Private/Sections/MovieScene3DAttachSection.cpp b/Engine/Source/Runtime/MovieSceneTracks/Private/Sections/MovieScene3DAttachSection.cpp
--- a/Engine/Source/Runtime/MovieSceneTracks/Private/Sections/MovieScene3DAttachSection.cpp
+++ b/Engine/Source/Runtime/MovieSceneTracks/Private/Sections/MovieScene3DAttachSection.cpp
@@ -16,6 +16,8 @@ UMovieScene3DAttachSection::UMovieScene3DAttachSection( const FObjectInitializer
 	bConstrainRx = true;
 	bConstrainRy = true;
 	bConstrainRz = true;
+
+	EvalOptions.EnableAndSetCompletionMode(EMovieSceneCompletionMode::RestoreState);
 }
 
 
diff --git a/Engine/Source/Runtime/MovieSceneTracks/Private/Sections/MovieScene3DTransformSection.cpp b/Engine/Source/Runtime/MovieSceneTracks/Private/Sections/MovieScene3DTransformSection.cpp
--- a/Engine/Source/Runtime/MovieSceneTracks/Private/Sections/MovieScene3DTransformSection.cpp
+++ b/Engine/Source/Runtime/MovieSceneTracks/Private/Sections/MovieScene3DTransformSection.cpp
@@ -3,7 +3,7 @@
 #include "MovieSceneTracksPrivatePCH.h"
 #include "MovieScene3DTransformSection.h"
 #include "MovieScene3DTransformTrack.h"
-
+#include "Evaluation/MovieScene3DTransformTemplate.h"
 
 /* FMovieScene3DLocationKeyStruct interface
  *****************************************************************************/
@@ -95,7 +95,9 @@ UMovieScene3DTransformSection::UMovieScene3DTransformSection(const FObjectInitia
 #if WITH_EDITORONLY_DATA
 	, Show3DTrajectory(EShow3DTrajectory::EST_OnlyWhenSelected)
 #endif
-{ }
+{
+	EvalOptions.EnableAndSetCompletionMode(EMovieSceneCompletionMode::KeepState);
+}
 
 
 /* UMovieScene3DTransformSection interface
@@ -127,7 +129,8 @@ void UMovieScene3DTransformSection::EvalScale(float Time, FVector& InOutScale) c
 /**
  * Chooses an appropriate curve from an axis and a set of curves
  */
-static FRichCurve* ChooseCurve(EAxis::Type Axis, FRichCurve* Curves)
+template<typename T>
+static T* ChooseCurve(EAxis::Type Axis, T* Curves)
 {
 	switch(Axis)
 	{
@@ -153,18 +156,36 @@ FRichCurve& UMovieScene3DTransformSection::GetTranslationCurve(EAxis::Type Axis)
 }
 
 
+const FRichCurve& UMovieScene3DTransformSection::GetTranslationCurve(EAxis::Type Axis) const
+{
+	return *ChooseCurve(Axis, Translation);
+}
+
+
 FRichCurve& UMovieScene3DTransformSection::GetRotationCurve(EAxis::Type Axis)
 {
 	return *ChooseCurve(Axis, Rotation);
 }
 
 
+const FRichCurve& UMovieScene3DTransformSection::GetRotationCurve(EAxis::Type Axis) const
+{
+	return *ChooseCurve(Axis, Rotation);
+}
+
+
 FRichCurve& UMovieScene3DTransformSection::GetScaleCurve(EAxis::Type Axis)
 {
 	return *ChooseCurve(Axis, Scale);
 }
 
 
+const FRichCurve& UMovieScene3DTransformSection::GetScaleCurve(EAxis::Type Axis) const
+{
+	return *ChooseCurve(Axis, Scale);
+}
+
+
 /* UMovieSceneSection interface
  *****************************************************************************/
 
@@ -563,3 +584,8 @@ void UMovieScene3DTransformSection::ClearDefaults()
 	Scale[1].ClearDefaultValue();
 	Scale[2].ClearDefaultValue();
 }
+
+FMovieSceneEvalTemplatePtr UMovieScene3DTransformSection::GenerateTemplate() const
+{
+	return FMovieScene3DTransformSectionTemplate(*this);
+}
\ No newline at end of file
diff --git a/Engine/Source/Runtime/MovieSceneTracks/Private/Sections/MovieSceneActorReferenceSection.cpp b/Engine/Source/Runtime/MovieSceneTracks/Private/Sections/MovieSceneActorReferenceSection.cpp
--- a/Engine/Source/Runtime/MovieSceneTracks/Private/Sections/MovieSceneActorReferenceSection.cpp
+++ b/Engine/Source/Runtime/MovieSceneTracks/Private/Sections/MovieSceneActorReferenceSection.cpp
@@ -122,8 +122,15 @@ bool UMovieSceneActorReferenceSection::HasKeys( const FGuid& Value ) const
 
 void UMovieSceneActorReferenceSection::SetDefault( const FGuid& Value )
 {
-	int32 DefaultIndex = ActorGuids.Add( Value );
-	ActorGuidIndexCurve.SetDefaultValue( DefaultIndex );
+	int32 CurrentDefault = ActorGuidIndexCurve.GetDefaultValue();
+	if (!ActorGuids.IsValidIndex(CurrentDefault) || ActorGuids[CurrentDefault] != Value)
+	{
+		if (TryModify())
+		{
+			int32 DefaultIndex = ActorGuids.Add( Value );
+			ActorGuidIndexCurve.SetDefaultValue( DefaultIndex );
+		}
+	}
 }
 
 
diff --git a/Engine/Source/Runtime/MovieSceneTracks/Private/Sections/MovieSceneAudioSection.cpp b/Engine/Source/Runtime/MovieSceneTracks/Private/Sections/MovieSceneAudioSection.cpp
--- a/Engine/Source/Runtime/MovieSceneTracks/Private/Sections/MovieSceneAudioSection.cpp
+++ b/Engine/Source/Runtime/MovieSceneTracks/Private/Sections/MovieSceneAudioSection.cpp
@@ -3,7 +3,7 @@
 #include "MovieSceneTracksPrivatePCH.h"
 #include "MovieSceneAudioSection.h"
 #include "SoundDefinitions.h"
-
+#include "Evaluation/MovieSceneAudioTemplate.h"
 
 UMovieSceneAudioSection::UMovieSceneAudioSection( const FObjectInitializer& ObjectInitializer )
 	: Super( ObjectInitializer )
@@ -17,8 +17,14 @@ UMovieSceneAudioSection::UMovieSceneAudioSection( const FObjectInitializer& Obje
 	bShowIntensity = false;
 #endif
 	bSuppressSubtitles = false;
+
+	EvalOptions.EnableAndSetCompletionMode(EMovieSceneCompletionMode::RestoreState);
 }
 
+FMovieSceneEvalTemplatePtr UMovieSceneAudioSection::GenerateTemplate() const
+{
+	return FMovieSceneAudioSectionTemplate(*this);
+}
 
 TRange<float> UMovieSceneAudioSection::GetAudioRange() const
 {
diff --git a/Engine/Source/Runtime/MovieSceneTracks/Private/Sections/MovieSceneBoolSection.cpp b/Engine/Source/Runtime/MovieSceneTracks/Private/Sections/MovieSceneBoolSection.cpp
--- a/Engine/Source/Runtime/MovieSceneTracks/Private/Sections/MovieSceneBoolSection.cpp
+++ b/Engine/Source/Runtime/MovieSceneTracks/Private/Sections/MovieSceneBoolSection.cpp
@@ -90,9 +90,10 @@ void UMovieSceneBoolSection::AddKey( float Time, const bool& Value, EMovieSceneK
 
 void UMovieSceneBoolSection::SetDefault( const bool& Value )
 {
-	if (TryModify())
+	uint8 NewValue = Value ? 1 : 0;
+	if (BoolCurve.GetDefaultValue() != NewValue && TryModify())
 	{
-		BoolCurve.SetDefaultValue(Value ? 1 : 0);
+		BoolCurve.SetDefaultValue(NewValue);
 	}
 }
 
diff --git a/Engine/Source/Runtime/MovieSceneTracks/Private/Sections/MovieSceneByteSection.cpp b/Engine/Source/Runtime/MovieSceneTracks/Private/Sections/MovieSceneByteSection.cpp
--- a/Engine/Source/Runtime/MovieSceneTracks/Private/Sections/MovieSceneByteSection.cpp
+++ b/Engine/Source/Runtime/MovieSceneTracks/Private/Sections/MovieSceneByteSection.cpp
@@ -2,6 +2,7 @@
 
 #include "MovieSceneTracksPrivatePCH.h"
 #include "MovieSceneByteSection.h"
+#include "Evaluation/MovieScenePropertyTemplates.h"
 
 
 UMovieSceneByteSection::UMovieSceneByteSection( const FObjectInitializer& ObjectInitializer )
@@ -9,12 +10,6 @@ UMovieSceneByteSection::UMovieSceneByteSection( const FObjectInitializer& Object
 { }
 
 
-uint8 UMovieSceneByteSection::Eval( float Position, uint8 DefaultValue) const
-{
-	return ByteCurve.Evaluate(Position, DefaultValue);
-}
-
-
 void UMovieSceneByteSection::MoveSection( float DeltaPosition, TSet<FKeyHandle>& KeyHandles )
 {
 	Super::MoveSection( DeltaPosition, KeyHandles );
@@ -79,7 +74,7 @@ void UMovieSceneByteSection::AddKey( float Time, const uint8& Value, EMovieScene
 
 bool UMovieSceneByteSection::NewKeyIsNewData( float Time, const uint8& Value ) const
 {
-	return Eval( Time, Value ) != Value;
+	return ByteCurve.Evaluate( Time, Value ) != Value;
 }
 
 
@@ -91,7 +86,7 @@ bool UMovieSceneByteSection::HasKeys( const uint8& Value ) const
 
 void UMovieSceneByteSection::SetDefault( const uint8& Value )
 {
-	if (TryModify())
+	if (ByteCurve.GetDefaultValue() != Value && TryModify())
 	{
 		ByteCurve.SetDefaultValue(Value);
 	}
diff --git a/Engine/Source/Runtime/MovieSceneTracks/Private/Sections/MovieSceneCameraAnimSection.cpp b/Engine/Source/Runtime/MovieSceneTracks/Private/Sections/MovieSceneCameraAnimSection.cpp
--- a/Engine/Source/Runtime/MovieSceneTracks/Private/Sections/MovieSceneCameraAnimSection.cpp
+++ b/Engine/Source/Runtime/MovieSceneTracks/Private/Sections/MovieSceneCameraAnimSection.cpp
@@ -3,70 +3,62 @@
 #include "MovieSceneTracksPrivatePCH.h"
 #include "MovieSceneCameraAnimSection.h"
 #include "MovieSceneCameraAnimTrack.h"
+#include "Evaluation/MovieSceneCameraAnimTemplate.h"
 
 
 UMovieSceneCameraAnimSection::UMovieSceneCameraAnimSection(const FObjectInitializer& ObjectInitializer)
  	: Super( ObjectInitializer )
 {
-	PlayRate = 1.f;
-	PlayScale = 1.f;
-	BlendInTime = 0.f;
-	BlendOutTime = 0.f;
-	bLooping = false;
-	bRandomStartTime = false;
-//	Duration = false;
-	PlaySpace = ECameraAnimPlaySpace::CameraLocal;
-	UserDefinedPlaySpace = FRotator::ZeroRotator;
+	CameraAnim_DEPRECATED = nullptr;
+	PlayRate_DEPRECATED = 1.f;
+	PlayScale_DEPRECATED = 1.f;
+	BlendInTime_DEPRECATED = 0.f;
+	BlendOutTime_DEPRECATED = 0.f;
+	bLooping_DEPRECATED = false;
+
+	EvalOptions.EnableAndSetCompletionMode(EMovieSceneCompletionMode::RestoreState);
 }
 
-void UMovieSceneCameraAnimSection::MoveSection(float DeltaPosition, TSet<FKeyHandle>& KeyHandles)
+void UMovieSceneCameraAnimSection::PostLoad()
 {
-	Super::MoveSection( DeltaPosition, KeyHandles );
-
-	// Move the curve
-	//AnimWeightCurve.ShiftCurve(DeltaPosition, KeyHandles);
-}
+	if (CameraAnim_DEPRECATED != nullptr)
+	{
+		AnimData.CameraAnim = CameraAnim_DEPRECATED;
+	}
 
+	if (PlayRate_DEPRECATED != 1.f)
+	{
+		AnimData.PlayRate = PlayRate_DEPRECATED;
+	}
 
-void UMovieSceneCameraAnimSection::DilateSection(float DilationFactor, float Origin, TSet<FKeyHandle>& KeyHandles)
-{	
-	Super::DilateSection(DilationFactor, Origin, KeyHandles);
-	
-	//AnimWeightCurve.ScaleCurve(Origin, DilationFactor, KeyHandles);
-}
+	if (PlayScale_DEPRECATED != 1.f)
+	{
+		AnimData.PlayScale = PlayScale_DEPRECATED;
+	}
 
+	if (BlendInTime_DEPRECATED != 0.f)
+	{
+		AnimData.BlendInTime = BlendInTime_DEPRECATED;
+	}
 
-void UMovieSceneCameraAnimSection::GetKeyHandles(TSet<FKeyHandle>& OutKeyHandles, TRange<float> TimeRange) const
-{
-	if (!TimeRange.Overlaps(GetRange()))
+	if (BlendOutTime_DEPRECATED != 0.f)
 	{
-		return;
+		AnimData.BlendOutTime = BlendOutTime_DEPRECATED;
 	}
 
-// 	for (auto It(AnimWeightCurve.GetKeyHandleIterator()); It; ++It)
-// 	{
-// 		float Time = AnimWeightCurve.GetKeyTime(It.Key());
-// 		if (TimeRange.Contains(Time))
-// 		{
-// 			KeyHandles.Add(It.Key());
-// 		}
-// 	}
-}
+	if (bLooping_DEPRECATED)
+	{
+		AnimData.bLooping = bLooping_DEPRECATED;
+	}
 
-TOptional<float> UMovieSceneCameraAnimSection::GetKeyTime( FKeyHandle KeyHandle ) const
-{
-	//if ( AnimWeightCurve.IsKeyHandleValid( KeyHandle ) )
-	//{
-	//	return TOptional<float>( AnimWeightCurve.GetKeyTime( KeyHandle ) );
-	//}
-	return TOptional<float>();
+	Super::PostLoad();
 }
 
-
-void UMovieSceneCameraAnimSection::SetKeyTime( FKeyHandle KeyHandle, float Time )
+FMovieSceneEvalTemplatePtr UMovieSceneCameraAnimSection::GenerateTemplate() const
 {
-	//if ( AnimWeightCurve.IsKeyHandleValid( KeyHandle ) )
-	//{
-	//	AnimWeightCurve.SetKeyTime( KeyHandle, Time );
-	//}
-}
+	if (AnimData.CameraAnim)
+	{
+		return FMovieSceneCameraAnimSectionTemplate(*this);
+	}
+	return FMovieSceneEvalTemplatePtr();
+}
\ No newline at end of file
diff --git a/Engine/Source/Runtime/MovieSceneTracks/Private/Sections/MovieSceneCameraCutSection.cpp b/Engine/Source/Runtime/MovieSceneTracks/Private/Sections/MovieSceneCameraCutSection.cpp
--- a/Engine/Source/Runtime/MovieSceneTracks/Private/Sections/MovieSceneCameraCutSection.cpp
+++ b/Engine/Source/Runtime/MovieSceneTracks/Private/Sections/MovieSceneCameraCutSection.cpp
@@ -3,6 +3,12 @@
 #include "MovieSceneTracksPrivatePCH.h"
 #include "MovieSceneCameraCutSection.h"
 
+#include "Evaluation/MovieSceneCameraCutTemplate.h"
 
 /* UMovieSceneCameraCutSection interface
  *****************************************************************************/
+
+FMovieSceneEvalTemplatePtr UMovieSceneCameraCutSection::GenerateTemplate() const
+{
+	return FMovieSceneCameraCutSectionTemplate(*this);
+}
\ No newline at end of file
diff --git a/Engine/Source/Runtime/MovieSceneTracks/Private/Sections/MovieSceneCameraShakeSection.cpp b/Engine/Source/Runtime/MovieSceneTracks/Private/Sections/MovieSceneCameraShakeSection.cpp
--- a/Engine/Source/Runtime/MovieSceneTracks/Private/Sections/MovieSceneCameraShakeSection.cpp
+++ b/Engine/Source/Runtime/MovieSceneTracks/Private/Sections/MovieSceneCameraShakeSection.cpp
@@ -3,65 +3,50 @@
 #include "MovieSceneTracksPrivatePCH.h"
 #include "MovieSceneCameraShakeSection.h"
 #include "MovieSceneCameraShakeTrack.h"
+#include "Evaluation/MovieSceneCameraAnimTemplate.h"
 
 
 UMovieSceneCameraShakeSection::UMovieSceneCameraShakeSection(const FObjectInitializer& ObjectInitializer)
- 	: Super( ObjectInitializer )
+	: Super( ObjectInitializer )
 {
-	PlayScale = 1.f;
-	PlaySpace = ECameraAnimPlaySpace::CameraLocal;
-	UserDefinedPlaySpace = FRotator::ZeroRotator;
-}
-
-void UMovieSceneCameraShakeSection::MoveSection(float DeltaPosition, TSet<FKeyHandle>& KeyHandles)
-{
-	Super::MoveSection( DeltaPosition, KeyHandles );
-
-	// Move the curve
-//	ShakeWeightCurve.ShiftCurve(DeltaPosition, KeyHandles);
-}
+	ShakeClass_DEPRECATED = nullptr;
+	PlayScale_DEPRECATED = 1.f;
+	PlaySpace_DEPRECATED = ECameraAnimPlaySpace::CameraLocal;
+	UserDefinedPlaySpace_DEPRECATED = FRotator::ZeroRotator;
 
-
-void UMovieSceneCameraShakeSection::DilateSection(float DilationFactor, float Origin, TSet<FKeyHandle>& KeyHandles)
-{	
-	Super::DilateSection(DilationFactor, Origin, KeyHandles);
-	
-//	ShakeWeightCurve.ScaleCurve(Origin, DilationFactor, KeyHandles);
+	EvalOptions.EnableAndSetCompletionMode(EMovieSceneCompletionMode::RestoreState);
 }
 
-
-void UMovieSceneCameraShakeSection::GetKeyHandles(TSet<FKeyHandle>& OutKeyHandles, TRange<float> TimeRange) const
+void UMovieSceneCameraShakeSection::PostLoad()
 {
-	if (!TimeRange.Overlaps(GetRange()))
+	if (ShakeClass_DEPRECATED != nullptr)
 	{
-		return;
+		ShakeData.ShakeClass = ShakeClass_DEPRECATED;
 	}
 
-// 	for (auto It(ShakeWeightCurve.GetKeyHandleIterator()); It; ++It)
-// 	{
-// 		float Time = ShakeWeightCurve.GetKeyTime(It.Key());
-// 		if (TimeRange.Contains(Time))
-// 		{
-// 			KeyHandles.Add(It.Key());
-// 		}
-// 	}
-}
+	if (PlayScale_DEPRECATED != 1.f)
+	{
+		ShakeData.PlayScale = PlayScale_DEPRECATED;
+	}
 
+	if (PlaySpace_DEPRECATED != ECameraAnimPlaySpace::CameraLocal)
+	{
+		ShakeData.PlaySpace = PlaySpace_DEPRECATED;
+	}
 
-TOptional<float> UMovieSceneCameraShakeSection::GetKeyTime( FKeyHandle KeyHandle ) const
-{
-	//if ( ShakeWeightCurve.IsKeyHandleValid( KeyHandle ) )
-	//{
-	//	return TOptional<float>( ShakeWeightCurve.GetKeyTime( KeyHandle ) );
-	//}
-	return TOptional<float>();
-}
+	if (UserDefinedPlaySpace_DEPRECATED != FRotator::ZeroRotator)
+	{
+		ShakeData.UserDefinedPlaySpace = UserDefinedPlaySpace_DEPRECATED;
+	}
 
+	Super::PostLoad();
+}
 
-void UMovieSceneCameraShakeSection::SetKeyTime( FKeyHandle KeyHandle, float Time )
+FMovieSceneEvalTemplatePtr UMovieSceneCameraShakeSection::GenerateTemplate() const
 {
-	//if ( ShakeWeightCurve.IsKeyHandleValid( KeyHandle ) )
-	//{
-	//	ShakeWeightCurve.SetKeyTime( KeyHandle, Time );
-	//}
-}
+	if (*ShakeData.ShakeClass)
+	{
+		return FMovieSceneCameraShakeSectionTemplate(*this);
+	}
+	return FMovieSceneEvalTemplatePtr();
+}
\ No newline at end of file
diff --git a/Engine/Source/Runtime/MovieSceneTracks/Private/Sections/MovieSceneColorSection.cpp b/Engine/Source/Runtime/MovieSceneTracks/Private/Sections/MovieSceneColorSection.cpp
--- a/Engine/Source/Runtime/MovieSceneTracks/Private/Sections/MovieSceneColorSection.cpp
+++ b/Engine/Source/Runtime/MovieSceneTracks/Private/Sections/MovieSceneColorSection.cpp
@@ -4,7 +4,6 @@
 #include "MovieSceneColorSection.h"
 #include "MovieSceneColorTrack.h"
 
-
 /* FMovieSceneColorKeyStruct interface
  *****************************************************************************/
 
@@ -34,20 +33,6 @@ UMovieSceneColorSection::UMovieSceneColorSection(const FObjectInitializer& Objec
 	: Super(ObjectInitializer)
 { }
 
-
-/* UMovieSceneColorSection interface
- *****************************************************************************/
-
-FLinearColor UMovieSceneColorSection::Eval(float Position, const FLinearColor& DefaultColor) const
-{
-	return FLinearColor(
-		RedCurve.Eval(Position, DefaultColor.R),
-		GreenCurve.Eval(Position, DefaultColor.G),
-		BlueCurve.Eval(Position, DefaultColor.B),
-		AlphaCurve.Eval(Position, DefaultColor.A));
-}
-
-
 /* UMovieSceneSection interface
  *****************************************************************************/
 
diff --git a/Engine/Source/Runtime/MovieSceneTracks/Private/Sections/MovieSceneEventSection.cpp b/Engine/Source/Runtime/MovieSceneTracks/Private/Sections/MovieSceneEventSection.cpp
--- a/Engine/Source/Runtime/MovieSceneTracks/Private/Sections/MovieSceneEventSection.cpp
+++ b/Engine/Source/Runtime/MovieSceneTracks/Private/Sections/MovieSceneEventSection.cpp
@@ -3,57 +3,38 @@
 #include "MovieSceneTracksPrivatePCH.h"
 #include "MovieSceneEventSection.h"
 
+#include "Curves/KeyFrameAlgorithms.h"
+
 
 /* UMovieSceneSection structors
  *****************************************************************************/
 
 UMovieSceneEventSection::UMovieSceneEventSection()
+#if WITH_EDITORONLY_DATA
+	: CurveInterface(TCurveInterface<FEventPayload, float>(&EventData.KeyTimes, &EventData.KeyValues, &EventData.KeyHandles))
+#else
+	: CurveInterface(TCurveInterface<FEventPayload, float>(&EventData.KeyTimes, &EventData.KeyValues))
+#endif
 {
 	SetIsInfinite(true);
 }
 
-
-/* UMovieSceneSection overrides
- *****************************************************************************/
-
-void UMovieSceneEventSection::AddKey(float Time, const FName& EventName)
+void UMovieSceneEventSection::PostLoad()
 {
-	if (TryModify())
+	for (FNameCurveKey EventKey : Events_DEPRECATED.GetKeys())
 	{
-		Events.UpdateOrAddKey(Time, EventName);
+		EventData.KeyTimes.Add(EventKey.Time);
+		EventData.KeyValues.Add(FEventPayload(EventKey.Value));
 	}
-}
 
-
-void UMovieSceneEventSection::TriggerEvents(float Position, float LastPosition,  IMovieScenePlayer& Player)
-{
-	const TArray<FNameCurveKey>& Keys = Events.GetKeys();
-
-	if (Position >= LastPosition)
+	if (Events_DEPRECATED.GetKeys().Num())
 	{
-		for (const auto& Key : Keys)
-		{
-			if ((Key.Time >= LastPosition) && (Key.Time <= Position))
-			{
-				TriggerEvent(Key.Value, Position, Player);
-			}
-		}
+		MarkAsChanged();
 	}
-	else
-	{
-		for (int32 KeyIndex = Keys.Num() - 1; KeyIndex >= 0; --KeyIndex)
-		{
-			const auto& Key = Keys[KeyIndex];
 
-			if ((Key.Time >= Position) && (Key.Time <= LastPosition))
-			{
-				TriggerEvent(Key.Value, Position, Player);
-			}		
-		}
-	}
+	Super::PostLoad();
 }
 
-
 /* UMovieSceneSection overrides
  *****************************************************************************/
 
@@ -61,18 +42,17 @@ void UMovieSceneEventSection::DilateSection(float DilationFactor, float Origin,
 {
 	Super::DilateSection(DilationFactor, Origin, KeyHandles);
 
-	Events.ScaleCurve(Origin, DilationFactor, KeyHandles);
+	KeyFrameAlgorithms::Scale(CurveInterface.GetValue(), Origin, DilationFactor, KeyHandles);
 }
 
 
 void UMovieSceneEventSection::GetKeyHandles(TSet<FKeyHandle>& KeyHandles, TRange<float> TimeRange) const
 {
-	for (auto It(Events.GetKeyHandleIterator()); It; ++It)
+	for (auto It(CurveInterface->IterateKeys()); It; ++It)
 	{
-		float Time = Events.GetKeyTime(It.Key());
-		if (TimeRange.Contains(Time))
+		if (TimeRange.Contains(*It))
 		{
-			KeyHandles.Add(It.Key());
+			KeyHandles.Add(It.GetKeyHandle());
 		}
 	}
 }
@@ -82,68 +62,17 @@ void UMovieSceneEventSection::MoveSection(float DeltaPosition, TSet<FKeyHandle>&
 {
 	Super::MoveSection(DeltaPosition, KeyHandles);
 
-	Events.ShiftCurve(DeltaPosition, KeyHandles);
+	KeyFrameAlgorithms::Translate(CurveInterface.GetValue(), DeltaPosition, KeyHandles);
 }
 
 
 TOptional<float> UMovieSceneEventSection::GetKeyTime( FKeyHandle KeyHandle ) const
 {
-	if ( Events.IsKeyHandleValid( KeyHandle ) )
-	{
-		return TOptional<float>( Events.GetKeyTime( KeyHandle ) );
-	}
-	return TOptional<float>();
+	return CurveInterface->GetKeyTime(KeyHandle);
 }
 
 
 void UMovieSceneEventSection::SetKeyTime( FKeyHandle KeyHandle, float Time )
 {
-	if ( Events.IsKeyHandleValid( KeyHandle ) )
-	{
-		Events.SetKeyTime( KeyHandle, Time );
-	}
-}
-
-
-/* UMovieSceneSection overrides
- *****************************************************************************/
-
-void UMovieSceneEventSection::TriggerEvent(const FName& Event, float Position, IMovieScenePlayer& Player)
-{
-	for (UObject* EventContextObject : Player.GetEventContexts())
-	{
-		UFunction* EventFunction = EventContextObject->FindFunction(Event);
-
-		if (EventFunction == nullptr)
-		{
-			// @todo sequencer: gmp: add external log category for MovieScene
-			//UE_LOG(LogMovieScene, Log, TEXT("UMovieSceneEventSection::TriggerEvent: Unable to find function '%s'"), *Event.ToString());
-		}
-		else if (EventFunction->NumParms != 0)
-		{
-			// @todo sequencer: gmp: add external log category for MovieScene
-			//UE_LOG(LogMovieScene, Log, TEXT("UMovieSceneEventSection::TriggerEvent: Function '%s' does not have zero parameters."), *Event.ToString());
-		}
-		else
-		{
-			EventContextObject->ProcessEvent(EventFunction, nullptr);
-		}
-	}
-
-#if !UE_BUILD_SHIPPING
-	if (Event == NAME_PerformanceCapture)
-	{
-		FString PackageName = GetOutermost()->GetName();
-		
-		FString LevelSequenceName;
-		FString FolderName;
-		PackageName.Split(TEXT("/"), &FolderName, &LevelSequenceName, ESearchCase::CaseSensitive, ESearchDir::FromEnd);
-
-		UWorld* PlaybackContext = Cast<UWorld>(Player.GetPlaybackContext());
-
-		FString MapName = PlaybackContext->GetName();
-
-		GEngine->PerformanceCapture(PlaybackContext, MapName, LevelSequenceName, Position);
-	}
-#endif	// UE_BUILD_SHIPPING
-}
+	CurveInterface->SetKeyTime(KeyHandle, Time);
+}
\ No newline at end of file
diff --git a/Engine/Source/Runtime/MovieSceneTracks/Private/Sections/MovieSceneFadeSection.cpp b/Engine/Source/Runtime/MovieSceneTracks/Private/Sections/MovieSceneFadeSection.cpp
--- a/Engine/Source/Runtime/MovieSceneTracks/Private/Sections/MovieSceneFadeSection.cpp
+++ b/Engine/Source/Runtime/MovieSceneTracks/Private/Sections/MovieSceneFadeSection.cpp
@@ -13,4 +13,6 @@ UMovieSceneFadeSection::UMovieSceneFadeSection()
 	, bFadeAudio(false)
 {
 	SetIsInfinite(true);
+
+	EvalOptions.EnableAndSetCompletionMode(EMovieSceneCompletionMode::RestoreState);
 }
diff --git a/Engine/Source/Runtime/MovieSceneTracks/Private/Sections/MovieSceneIntegerSection.cpp b/Engine/Source/Runtime/MovieSceneTracks/Private/Sections/MovieSceneIntegerSection.cpp
new file mode 100644
--- /dev/null
+++ b/Engine/Source/Runtime/MovieSceneTracks/Private/Sections/MovieSceneIntegerSection.cpp
@@ -0,0 +1,99 @@
+// Copyright 1998-2016 Epic Games, Inc. All Rights Reserved.
+
+#include "MovieSceneTracksPrivatePCH.h"
+#include "MovieSceneIntegerSection.h"
+#include "Evaluation/MovieScenePropertyTemplates.h"
+
+
+UMovieSceneIntegerSection::UMovieSceneIntegerSection( const FObjectInitializer& ObjectInitializer )
+	: Super( ObjectInitializer )
+{ }
+
+
+void UMovieSceneIntegerSection::MoveSection( float DeltaPosition, TSet<FKeyHandle>& KeyHandles )
+{
+	Super::MoveSection( DeltaPosition, KeyHandles );
+
+	IntegerCurve.ShiftCurve(DeltaPosition, KeyHandles);
+}
+
+
+void UMovieSceneIntegerSection::DilateSection( float DilationFactor, float Origin, TSet<FKeyHandle>& KeyHandles )
+{
+	Super::DilateSection(DilationFactor, Origin, KeyHandles);
+	
+	IntegerCurve.ScaleCurve(Origin, DilationFactor, KeyHandles);
+}
+
+
+void UMovieSceneIntegerSection::GetKeyHandles(TSet<FKeyHandle>& OutKeyHandles, TRange<float> TimeRange) const
+{
+	if (!TimeRange.Overlaps(GetRange()))
+	{
+		return;
+	}
+
+	for (auto It(IntegerCurve.GetKeyHandleIterator()); It; ++It)
+	{
+		float Time = IntegerCurve.GetKeyTime(It.Key());
+		if (TimeRange.Contains(Time))
+		{
+			OutKeyHandles.Add(It.Key());
+		}
+	}
+}
+
+
+TOptional<float> UMovieSceneIntegerSection::GetKeyTime( FKeyHandle KeyHandle ) const
+{
+	if ( IntegerCurve.IsKeyHandleValid( KeyHandle ) )
+	{
+		return TOptional<float>( IntegerCurve.GetKeyTime( KeyHandle ) );
+	}
+	return TOptional<float>();
+}
+
+
+void UMovieSceneIntegerSection::SetKeyTime( FKeyHandle KeyHandle, float Time )
+{
+	if ( IntegerCurve.IsKeyHandleValid( KeyHandle ) )
+	{
+		IntegerCurve.SetKeyTime( KeyHandle, Time );
+	}
+}
+
+
+void UMovieSceneIntegerSection::AddKey( float Time, const int32& Value, EMovieSceneKeyInterpolation KeyInterpolation )
+{
+	if (TryModify())
+	{
+		IntegerCurve.UpdateOrAddKey(Time, Value);
+	}
+}
+
+
+bool UMovieSceneIntegerSection::NewKeyIsNewData( float Time, const int32& Value ) const
+{
+	return IntegerCurve.Evaluate( Time, Value ) != Value;
+}
+
+
+bool UMovieSceneIntegerSection::HasKeys( const int32& Value ) const
+{
+	return IntegerCurve.GetNumKeys() > 0;
+}
+
+
+void UMovieSceneIntegerSection::SetDefault( const int32& Value )
+{
+	if (IntegerCurve.GetDefaultValue() != Value && TryModify())
+	{
+		IntegerCurve.SetDefaultValue(Value);
+	}
+}
+
+
+void UMovieSceneIntegerSection::ClearDefaults()
+{
+	IntegerCurve.ClearDefaultValue();
+}
diff --git a/Engine/Source/Runtime/MovieSceneTracks/Private/Sections/MovieSceneLevelVisibilitySection.cpp b/Engine/Source/Runtime/MovieSceneTracks/Private/Sections/MovieSceneLevelVisibilitySection.cpp
--- a/Engine/Source/Runtime/MovieSceneTracks/Private/Sections/MovieSceneLevelVisibilitySection.cpp
+++ b/Engine/Source/Runtime/MovieSceneTracks/Private/Sections/MovieSceneLevelVisibilitySection.cpp
@@ -2,6 +2,7 @@
 
 #include "MovieSceneTracksPrivatePCH.h"
 #include "MovieSceneLevelVisibilitySection.h"
+#include "Evaluation/MovieSceneLevelVisibilityTemplate.h"
 
 
 UMovieSceneLevelVisibilitySection::UMovieSceneLevelVisibilitySection( const FObjectInitializer& ObjectInitializer )
@@ -29,6 +30,12 @@ TArray<FName>* UMovieSceneLevelVisibilitySection::GetLevelNames()
 }
 
 
+FMovieSceneEvalTemplatePtr UMovieSceneLevelVisibilitySection::GenerateTemplate() const
+{
+	return FMovieSceneLevelVisibilitySectionTemplate(*this);
+}
+
+
 TOptional<float> UMovieSceneLevelVisibilitySection::GetKeyTime(FKeyHandle KeyHandle) const
 {
 	return TOptional<float>();
diff --git a/Engine/Source/Runtime/MovieSceneTracks/Private/Sections/MovieSceneParameterSection.cpp b/Engine/Source/Runtime/MovieSceneTracks/Private/Sections/MovieSceneParameterSection.cpp
--- a/Engine/Source/Runtime/MovieSceneTracks/Private/Sections/MovieSceneParameterSection.cpp
+++ b/Engine/Source/Runtime/MovieSceneTracks/Private/Sections/MovieSceneParameterSection.cpp
@@ -135,16 +135,31 @@ TArray<FScalarParameterNameAndCurve>* UMovieSceneParameterSection::GetScalarPara
 	return &ScalarParameterNamesAndCurves;
 }
 
+const TArray<FScalarParameterNameAndCurve>& UMovieSceneParameterSection::GetScalarParameterNamesAndCurves() const
+{
+	return ScalarParameterNamesAndCurves;
+}
+
 TArray<FVectorParameterNameAndCurves>* UMovieSceneParameterSection::GetVectorParameterNamesAndCurves()
 {
 	return &VectorParameterNamesAndCurves;
 }
 
+const TArray<FVectorParameterNameAndCurves>& UMovieSceneParameterSection::GetVectorParameterNamesAndCurves() const
+{
+	return VectorParameterNamesAndCurves;
+}
+
 TArray<FColorParameterNameAndCurves>* UMovieSceneParameterSection::GetColorParameterNamesAndCurves()
 {
 	return &ColorParameterNamesAndCurves;
 }
 
+const TArray<FColorParameterNameAndCurves>& UMovieSceneParameterSection::GetColorParameterNamesAndCurves() const
+{
+	return ColorParameterNamesAndCurves;
+}
+
 void UMovieSceneParameterSection::GetParameterNames( TSet<FName>& ParameterNames ) const
 {
 	for ( const FScalarParameterNameAndCurve& ScalarParameterNameAndCurve : ScalarParameterNamesAndCurves )
@@ -161,32 +176,6 @@ void UMovieSceneParameterSection::GetParameterNames( TSet<FName>& ParameterNames
 	}
 }
 
-void UMovieSceneParameterSection::Eval( float Position,
-	TArray<FScalarParameterNameAndValue>& OutScalarValues,
-	TArray<FVectorParameterNameAndValue>& OutVectorValues,
-	TArray<FColorParameterNameAndValue>& OutColorValues ) const
-{
-	for ( const FScalarParameterNameAndCurve& ScalarParameterNameAndCurve : ScalarParameterNamesAndCurves )
-	{
-		OutScalarValues.Add( FScalarParameterNameAndValue( ScalarParameterNameAndCurve.ParameterName, ScalarParameterNameAndCurve.ParameterCurve.Eval( Position ) ) );
-	}
-	for ( const FVectorParameterNameAndCurves& VectorParameterNameAndCurves : VectorParameterNamesAndCurves )
-	{
-		OutVectorValues.Add( FVectorParameterNameAndValue( VectorParameterNameAndCurves.ParameterName, FVector(
-			VectorParameterNameAndCurves.XCurve.Eval( Position ),
-			VectorParameterNameAndCurves.YCurve.Eval( Position ),
-			VectorParameterNameAndCurves.ZCurve.Eval( Position ) ) ) );
-	}
-	for ( const FColorParameterNameAndCurves& ColorParameterNameAndCurves : ColorParameterNamesAndCurves )
-	{
-		OutColorValues.Add( FColorParameterNameAndValue( ColorParameterNameAndCurves.ParameterName, FLinearColor(
-			ColorParameterNameAndCurves.RedCurve.Eval( Position ),
-			ColorParameterNameAndCurves.GreenCurve.Eval( Position ),
-			ColorParameterNameAndCurves.BlueCurve.Eval( Position ),
-			ColorParameterNameAndCurves.AlphaCurve.Eval( Position ) ) ) );
-	}
-}
-
 void UMovieSceneParameterSection::UpdateParameterIndicesFromRemoval( int32 RemovedIndex )
 {
 	for ( FScalarParameterNameAndCurve& ScalarParameterAndCurve : ScalarParameterNamesAndCurves )
diff --git a/Engine/Source/Runtime/MovieSceneTracks/Private/Sections/MovieSceneParticleSection.cpp b/Engine/Source/Runtime/MovieSceneTracks/Private/Sections/MovieSceneParticleSection.cpp
--- a/Engine/Source/Runtime/MovieSceneTracks/Private/Sections/MovieSceneParticleSection.cpp
+++ b/Engine/Source/Runtime/MovieSceneTracks/Private/Sections/MovieSceneParticleSection.cpp
@@ -2,6 +2,7 @@
 
 #include "MovieSceneTracksPrivatePCH.h"
 #include "MovieSceneParticleSection.h"
+#include "Evaluation/MovieSceneParticleTemplate.h"
 
 
 UMovieSceneParticleSection::UMovieSceneParticleSection( const FObjectInitializer& ObjectInitializer )
@@ -17,11 +18,6 @@ void UMovieSceneParticleSection::AddKey( float Time, EParticleKey::Type KeyType
 	ParticleKeys.AddKey( Time, (int32)KeyType );
 }
 
-FIntegralCurve& UMovieSceneParticleSection::GetParticleCurve()
-{
-	return ParticleKeys;
-}
-
 void UMovieSceneParticleSection::MoveSection( float DeltaPosition, TSet<FKeyHandle>& KeyHandles )
 {
 	Super::MoveSection( DeltaPosition, KeyHandles );
@@ -71,3 +67,8 @@ void UMovieSceneParticleSection::SetKeyTime( FKeyHandle KeyHandle, float Time )
 		ParticleKeys.SetKeyTime( KeyHandle, Time );
 	}
 }
+
+FMovieSceneEvalTemplatePtr UMovieSceneParticleSection::GenerateTemplate() const
+{
+	return FMovieSceneParticleSectionTemplate(*this);
+}
