Copying //UE4/Dev-Sequencer to //UE4/Dev-Main (Source: //UE4/Dev-Sequencer @ 3178529)

diff --git a/Engine/Source/Runtime/MovieSceneTracks/Private/Sections/MovieSceneSkeletalAnimationSection.cpp b/Engine/Source/Runtime/MovieSceneTracks/Private/Sections/MovieSceneSkeletalAnimationSection.cpp
--- a/Engine/Source/Runtime/MovieSceneTracks/Private/Sections/MovieSceneSkeletalAnimationSection.cpp
+++ b/Engine/Source/Runtime/MovieSceneTracks/Private/Sections/MovieSceneSkeletalAnimationSection.cpp
@@ -3,34 +3,85 @@
 #include "MovieSceneTracksPrivatePCH.h"
 #include "MovieSceneSkeletalAnimationSection.h"
 #include "Animation/AnimSequence.h"
+#include "Evaluation/MovieSceneSkeletalAnimationTemplate.h"
 
-FName UMovieSceneSkeletalAnimationSection::DefaultSlotName( "DefaultSlot" );
+namespace
+{
+	FName DefaultSlotName( "DefaultSlot" );
+}
 
-UMovieSceneSkeletalAnimationSection::UMovieSceneSkeletalAnimationSection( const FObjectInitializer& ObjectInitializer )
-	: Super( ObjectInitializer )
+FMovieSceneSkeletalAnimationParams::FMovieSceneSkeletalAnimationParams()
 {
-	AnimSequence_DEPRECATED = nullptr;
 	Animation = nullptr;
 	StartOffset = 0.f;
 	EndOffset = 0.f;
 	PlayRate = 1.f;
-#if WITH_EDITOR
-	PreviousPlayRate = PlayRate;
-#endif
 	bReverse = false;
 	SlotName = DefaultSlotName;
 }
 
+UMovieSceneSkeletalAnimationSection::UMovieSceneSkeletalAnimationSection( const FObjectInitializer& ObjectInitializer )
+	: Super( ObjectInitializer )
+{
+	AnimSequence_DEPRECATED = nullptr;
+	Animation_DEPRECATED = nullptr;
+	StartOffset_DEPRECATED = 0.f;
+	EndOffset_DEPRECATED = 0.f;
+	PlayRate_DEPRECATED = 1.f;
+	bReverse_DEPRECATED = false;
+	SlotName_DEPRECATED = DefaultSlotName;
+
+#if WITH_EDITOR
+	PreviousPlayRate = Params.PlayRate;
+#endif
+}
+
 void UMovieSceneSkeletalAnimationSection::PostLoad()
 {
 	if (AnimSequence_DEPRECATED)
 	{
-		Animation = AnimSequence_DEPRECATED;
+		Params.Animation = AnimSequence_DEPRECATED;
+	}
+
+	if (Animation_DEPRECATED != nullptr)
+	{
+		Params.Animation = Animation_DEPRECATED;
+	}
+
+	if (StartOffset_DEPRECATED != 0.f)
+	{
+		Params.StartOffset = StartOffset_DEPRECATED;
+	}
+
+	if (EndOffset_DEPRECATED != 0.f)
+	{
+		Params.EndOffset = EndOffset_DEPRECATED;
+	}
+
+	if (PlayRate_DEPRECATED != 1.f)
+	{
+		Params.PlayRate = PlayRate_DEPRECATED;
+	}
+
+	if (bReverse_DEPRECATED != false)
+	{
+		Params.bReverse = bReverse_DEPRECATED;
+	}
+
+	if (SlotName_DEPRECATED != DefaultSlotName)
+	{
+		Params.SlotName = SlotName_DEPRECATED;
 	}
 
 	Super::PostLoad();
 }
 
+FMovieSceneEvalTemplatePtr UMovieSceneSkeletalAnimationSection::GenerateTemplate() const
+{
+	return FMovieSceneSkeletalAnimationSectionTemplate(*this);
+}
+
+
 void UMovieSceneSkeletalAnimationSection::MoveSection( float DeltaTime, TSet<FKeyHandle>& KeyHandles )
 {
 	Super::MoveSection(DeltaTime, KeyHandles);
@@ -39,25 +90,25 @@ void UMovieSceneSkeletalAnimationSection::MoveSection( float DeltaTime, TSet<FKe
 
 void UMovieSceneSkeletalAnimationSection::DilateSection( float DilationFactor, float Origin, TSet<FKeyHandle>& KeyHandles )
 {
-	PlayRate /= DilationFactor;
+	Params.PlayRate /= DilationFactor;
 
 	Super::DilateSection(DilationFactor, Origin, KeyHandles);
 }
 
 UMovieSceneSection* UMovieSceneSkeletalAnimationSection::SplitSection(float SplitTime)
 {
-	float AnimPlayRate = FMath::IsNearlyZero(GetPlayRate()) ? 1.0f : GetPlayRate();
+	float AnimPlayRate = FMath::IsNearlyZero(Params.PlayRate) ? 1.0f : Params.PlayRate;
 	float AnimPosition = (SplitTime - GetStartTime()) * AnimPlayRate;
-	float SeqLength = GetSequenceLength() - (GetStartOffset() + GetEndOffset());
+	float SeqLength = Params.GetSequenceLength() - (Params.StartOffset + Params.EndOffset);
 
 	float NewOffset = FMath::Fmod(AnimPosition, SeqLength);
-	NewOffset += GetStartOffset();
+	NewOffset += Params.StartOffset;
 
 	UMovieSceneSection* NewSection = Super::SplitSection(SplitTime);
 	if (NewSection != nullptr)
 	{
 		UMovieSceneSkeletalAnimationSection* NewSkeletalSection = Cast<UMovieSceneSkeletalAnimationSection>(NewSection);
-		NewSkeletalSection->SetStartOffset(NewOffset);
+		NewSkeletalSection->Params.StartOffset = NewOffset;
 	}
 	return NewSection;
 }
@@ -74,11 +125,11 @@ void UMovieSceneSkeletalAnimationSection::GetSnapTimes(TArray<float>& OutSnapTim
 	Super::GetSnapTimes(OutSnapTimes, bGetSectionBorders);
 
 	float CurrentTime = GetStartTime();
-	float AnimPlayRate = FMath::IsNearlyZero(GetPlayRate()) ? 1.0f : GetPlayRate();
-	float SeqLength = (GetSequenceLength() - (GetStartOffset() + GetEndOffset())) / AnimPlayRate;
+	float AnimPlayRate = FMath::IsNearlyZero(Params.PlayRate) ? 1.0f : Params.PlayRate;
+	float SeqLength = (Params.GetSequenceLength() - (Params.StartOffset + Params.EndOffset)) / AnimPlayRate;
 
 	// Snap to the repeat times
-	while (CurrentTime <= GetEndTime() && !FMath::IsNearlyZero(GetDuration()) && SeqLength > 0)
+	while (CurrentTime <= GetEndTime() && !FMath::IsNearlyZero(Params.GetDuration()) && SeqLength > 0)
 	{
 		if (CurrentTime >= GetStartTime())
 		{
@@ -93,7 +144,7 @@ void UMovieSceneSkeletalAnimationSection::GetSnapTimes(TArray<float>& OutSnapTim
 void UMovieSceneSkeletalAnimationSection::PreEditChange(UProperty* PropertyAboutToChange)
 {
 	// Store the current play rate so that we can compute the amount to compensate the section end time when the play rate changes
-	PreviousPlayRate = GetPlayRate();
+	PreviousPlayRate = Params.PlayRate;
 
 	Super::PreEditChange(PropertyAboutToChange);
 }
@@ -104,7 +155,7 @@ void UMovieSceneSkeletalAnimationSection::PostEditChangeProperty(FPropertyChange
 	if (PropertyChangedEvent.Property != nullptr &&
 		PropertyChangedEvent.Property->GetFName() == TEXT("PlayRate"))
 	{
-		float NewPlayRate = GetPlayRate();
+		float NewPlayRate = Params.PlayRate;
 
 		if (!FMath::IsNearlyZero(NewPlayRate))
 		{
diff --git a/Engine/Source/Runtime/MovieSceneTracks/Private/Sections/MovieSceneSpawnSection.cpp b/Engine/Source/Runtime/MovieSceneTracks/Private/Sections/MovieSceneSpawnSection.cpp
new file mode 100644
--- /dev/null
+++ b/Engine/Source/Runtime/MovieSceneTracks/Private/Sections/MovieSceneSpawnSection.cpp
@@ -0,0 +1,11 @@
+// Copyright 1998-2016 Epic Games, Inc. All Rights Reserved.
+
+#include "MovieSceneTracksPrivatePCH.h"
+#include "MovieSceneSpawnSection.h"
+
+
+UMovieSceneSpawnSection::UMovieSceneSpawnSection(const FObjectInitializer& Init)
+	: Super(Init)
+{
+	EvalOptions.EnableAndSetCompletionMode(EMovieSceneCompletionMode::RestoreState);
+}
diff --git a/Engine/Source/Runtime/MovieSceneTracks/Private/Sections/MovieSceneStringSection.cpp b/Engine/Source/Runtime/MovieSceneTracks/Private/Sections/MovieSceneStringSection.cpp
--- a/Engine/Source/Runtime/MovieSceneTracks/Private/Sections/MovieSceneStringSection.cpp
+++ b/Engine/Source/Runtime/MovieSceneTracks/Private/Sections/MovieSceneStringSection.cpp
@@ -40,7 +40,10 @@ void UMovieSceneStringSection::AddKey( float Time, const FString& Key, EMovieSce
 
 void UMovieSceneStringSection::SetDefault(const FString& Value)
 {
-	StringCurve.SetDefaultValue(Value);
+	if (StringCurve.DefaultValue.Compare(Value) != 0 && TryModify())
+	{
+		StringCurve.SetDefaultValue(Value);
+	}
 }
 
 void UMovieSceneStringSection::ClearDefaults()
diff --git a/Engine/Source/Runtime/MovieSceneTracks/Private/Sections/MovieSceneSubSection.cpp b/Engine/Source/Runtime/MovieSceneTracks/Private/Sections/MovieSceneSubSection.cpp
--- a/Engine/Source/Runtime/MovieSceneTracks/Private/Sections/MovieSceneSubSection.cpp
+++ b/Engine/Source/Runtime/MovieSceneTracks/Private/Sections/MovieSceneSubSection.cpp
@@ -7,14 +7,55 @@
 
 TWeakObjectPtr<UMovieSceneSubSection> UMovieSceneSubSection::TheRecordingSection;
 
+float DeprecatedMagicNumber = TNumericLimits<float>::Lowest();
+
 /* UMovieSceneSubSection structors
  *****************************************************************************/
 
 UMovieSceneSubSection::UMovieSceneSubSection()
-	: StartOffset(0.0f)
-	, TimeScale(1.0f)
-	, PrerollTime(0.0f)
+	: StartOffset_DEPRECATED(DeprecatedMagicNumber)
+	, TimeScale_DEPRECATED(DeprecatedMagicNumber)
+	, PrerollTime_DEPRECATED(DeprecatedMagicNumber)
+{
+}
+
+FString UMovieSceneSubSection::GetPathNameInMovieScene() const
+{
+	UMovieScene* OuterMovieScene = GetTypedOuter<UMovieScene>();
+	check(OuterMovieScene);
+	return GetPathName(OuterMovieScene);
+}
+
+FMovieSceneSequenceID UMovieSceneSubSection::GetSequenceID() const
+{
+	FString FullPath = GetPathNameInMovieScene();
+	if (SubSequence)
+	{
+		FullPath += TEXT(" / ");
+		FullPath += SubSequence->GetPathName();
+	}
+
+	return FMovieSceneSequenceID(FCrc::Strihash_DEPRECATED(*FullPath));
+}
+
+void UMovieSceneSubSection::PostLoad()
 {
+	if (StartOffset_DEPRECATED != DeprecatedMagicNumber)
+	{
+		Parameters.StartOffset = StartOffset_DEPRECATED;
+	}
+
+	if (TimeScale_DEPRECATED != DeprecatedMagicNumber)
+	{
+		Parameters.TimeScale = TimeScale_DEPRECATED;
+	}
+
+	if (PrerollTime_DEPRECATED != DeprecatedMagicNumber)
+	{
+		Parameters.PrerollTime = PrerollTime_DEPRECATED;
+	}
+
+	Super::PostLoad();
 }
 
 void UMovieSceneSubSection::SetSequence(UMovieSceneSequence* Sequence)
@@ -107,8 +148,8 @@ UMovieSceneSection* UMovieSceneSubSection::SplitSection( float SplitTime )
 	}
 
 	float InitialStartTime = GetStartTime();
-	float InitialStartOffset = StartOffset;
-	float NewStartOffset = ( SplitTime - InitialStartTime ) / TimeScale;
+	float InitialStartOffset = Parameters.StartOffset;
+	float NewStartOffset = ( SplitTime - InitialStartTime ) / Parameters.TimeScale;
 	NewStartOffset += InitialStartOffset;
 
 	// Ensure start offset is not less than 0
@@ -117,7 +158,7 @@ UMovieSceneSection* UMovieSceneSubSection::SplitSection( float SplitTime )
 	UMovieSceneSubSection* NewSection = Cast<UMovieSceneSubSection>( UMovieSceneSection::SplitSection( SplitTime ) );
 	if ( NewSection )
 	{
-		NewSection->StartOffset = NewStartOffset;
+		NewSection->Parameters.StartOffset = NewStartOffset;
 		return NewSection;
 	}
 
@@ -132,17 +173,17 @@ void UMovieSceneSubSection::TrimSection( float TrimTime, bool bTrimLeft )
 	}
 
 	float InitialStartTime = GetStartTime();
-	float InitialStartOffset = StartOffset;
+	float InitialStartOffset = Parameters.StartOffset;
 
 	UMovieSceneSection::TrimSection( TrimTime, bTrimLeft );
 
 	// If trimming off the left, set the offset of the shot
 	if ( bTrimLeft )
 	{
-		float NewStartOffset = ( TrimTime - InitialStartTime ) / TimeScale;
+		float NewStartOffset = ( TrimTime - InitialStartTime ) / Parameters.TimeScale;
 		NewStartOffset += InitialStartOffset;
 
 		// Ensure start offset is not less than 0
-		StartOffset = FMath::Max( NewStartOffset, 0.f );
+		Parameters.StartOffset = FMath::Max( NewStartOffset, 0.f );
 	}
 }
\ No newline at end of file
diff --git a/Engine/Source/Runtime/MovieSceneTracks/Private/Sections/MovieSceneVectorSection.cpp b/Engine/Source/Runtime/MovieSceneTracks/Private/Sections/MovieSceneVectorSection.cpp
--- a/Engine/Source/Runtime/MovieSceneTracks/Private/Sections/MovieSceneVectorSection.cpp
+++ b/Engine/Source/Runtime/MovieSceneTracks/Private/Sections/MovieSceneVectorSection.cpp
@@ -36,20 +36,6 @@ UMovieSceneVectorSection::UMovieSceneVectorSection(const FObjectInitializer& Obj
 	ChannelsUsed = 0;
 }
 
-
-/* UMovieSceneVectorSection interface
- *****************************************************************************/
-
-FVector4 UMovieSceneVectorSection::Eval(float Position, const FVector4& DefaultVector) const
-{
-	return FVector4(
-		Curves[0].Eval(Position, DefaultVector.X),
-		Curves[1].Eval(Position, DefaultVector.Y),
-		Curves[2].Eval(Position, DefaultVector.Z),
-		Curves[3].Eval(Position, DefaultVector.W));
-}
-
-
 /* UMovieSceneSection interface
  *****************************************************************************/
 
@@ -220,7 +206,7 @@ bool UMovieSceneVectorSection::HasKeys(const FVectorKey& Key) const
 void UMovieSceneVectorSection::SetDefault(const FVectorKey& Key)
 {
 	FRichCurve* ChannelCurve = GetCurveForChannel(Key.Channel, Curves, ChannelsUsed);
-	ChannelCurve->SetDefaultValue(Key.Value);
+	SetCurveDefault( *ChannelCurve, Key.Value );
 }
 
 
diff --git a/Engine/Source/Runtime/MovieSceneTracks/Private/Tracks/MovieScene3DAttachTrack.cpp b/Engine/Source/Runtime/MovieSceneTracks/Private/Tracks/MovieScene3DAttachTrack.cpp
--- a/Engine/Source/Runtime/MovieSceneTracks/Private/Tracks/MovieScene3DAttachTrack.cpp
+++ b/Engine/Source/Runtime/MovieSceneTracks/Private/Tracks/MovieScene3DAttachTrack.cpp
@@ -4,7 +4,8 @@
 #include "MovieScene3DAttachSection.h"
 #include "MovieScene3DAttachTrack.h"
 #include "IMovieScenePlayer.h"
-#include "MovieScene3DAttachTrackInstance.h"
+#include "Evaluation/MovieScene3DAttachTemplate.h"
+#include "Evaluation/MovieSceneEvaluationTrack.h"
 
 
 #define LOCTEXT_NAMESPACE "MovieScene3DAttachTrack"
@@ -14,13 +15,11 @@ UMovieScene3DAttachTrack::UMovieScene3DAttachTrack( const FObjectInitializer& Ob
 	: Super( ObjectInitializer )
 { }
 
-
-TSharedPtr<IMovieSceneTrackInstance> UMovieScene3DAttachTrack::CreateInstance()
+FMovieSceneEvalTemplatePtr UMovieScene3DAttachTrack::CreateTemplateForSection(const UMovieSceneSection& InSection) const
 {
-	return MakeShareable( new FMovieScene3DAttachTrackInstance( *this ) ); 
+	return FMovieScene3DAttachSectionTemplate(*CastChecked<UMovieScene3DAttachSection>(&InSection));
 }
 
-
 void UMovieScene3DAttachTrack::AddConstraint(float KeyTime, float ConstraintEndTime, const FName SocketName, const FName ComponentName, const FGuid& ConstraintId)
 {
 	// add the section
diff --git a/Engine/Source/Runtime/MovieSceneTracks/Private/Tracks/MovieScene3DConstraintTrack.cpp b/Engine/Source/Runtime/MovieSceneTracks/Private/Tracks/MovieScene3DConstraintTrack.cpp
--- a/Engine/Source/Runtime/MovieSceneTracks/Private/Tracks/MovieScene3DConstraintTrack.cpp
+++ b/Engine/Source/Runtime/MovieSceneTracks/Private/Tracks/MovieScene3DConstraintTrack.cpp
@@ -4,7 +4,6 @@
 #include "MovieScene3DConstraintSection.h"
 #include "MovieScene3DConstraintTrack.h"
 #include "IMovieScenePlayer.h"
-#include "MovieScene3DConstraintTrackInstance.h"
 
 
 #define LOCTEXT_NAMESPACE "MovieScene3DConstraintTrack"
diff --git a/Engine/Source/Runtime/MovieSceneTracks/Private/Tracks/MovieScene3DPathTrack.cpp b/Engine/Source/Runtime/MovieSceneTracks/Private/Tracks/MovieScene3DPathTrack.cpp
--- a/Engine/Source/Runtime/MovieSceneTracks/Private/Tracks/MovieScene3DPathTrack.cpp
+++ b/Engine/Source/Runtime/MovieSceneTracks/Private/Tracks/MovieScene3DPathTrack.cpp
@@ -4,7 +4,8 @@
 #include "MovieScene3DPathSection.h"
 #include "MovieScene3DPathTrack.h"
 #include "IMovieScenePlayer.h"
-#include "MovieScene3DPathTrackInstance.h"
+#include "Evaluation/MovieScene3DPathTemplate.h"
+#include "Evaluation/MovieSceneEvaluationTrack.h"
 
 
 #define LOCTEXT_NAMESPACE "MovieScene3DPathTrack"
@@ -15,9 +16,9 @@ UMovieScene3DPathTrack::UMovieScene3DPathTrack( const FObjectInitializer& Object
 { }
 
 
-TSharedPtr<IMovieSceneTrackInstance> UMovieScene3DPathTrack::CreateInstance()
+FMovieSceneEvalTemplatePtr UMovieScene3DPathTrack::CreateTemplateForSection(const UMovieSceneSection& InSection) const
 {
-	return MakeShareable( new FMovieScene3DPathTrackInstance( *this ) ); 
+	return FMovieScene3DPathSectionTemplate(*CastChecked<UMovieScene3DPathSection>(&InSection));
 }
 
 
diff --git a/Engine/Source/Runtime/MovieSceneTracks/Private/Tracks/MovieScene3DTransformTrack.cpp b/Engine/Source/Runtime/MovieSceneTracks/Private/Tracks/MovieScene3DTransformTrack.cpp
--- a/Engine/Source/Runtime/MovieSceneTracks/Private/Tracks/MovieScene3DTransformTrack.cpp
+++ b/Engine/Source/Runtime/MovieSceneTracks/Private/Tracks/MovieScene3DTransformTrack.cpp
@@ -4,8 +4,7 @@
 #include "MovieScene3DTransformSection.h"
 #include "MovieScene3DTransformTrack.h"
 #include "IMovieScenePlayer.h"
-#include "MovieScene3DTransformTrackInstance.h"
-
+#include "MovieSceneSegmentCompiler.h"
 
 UMovieScene3DTransformTrack::UMovieScene3DTransformTrack( const FObjectInitializer& ObjectInitializer )
 	: Super( ObjectInitializer )
@@ -16,39 +15,11 @@ UMovieScene3DTransformTrack::UMovieScene3DTransformTrack( const FObjectInitializ
 #if WITH_EDITORONLY_DATA
 	TrackTint = FColor(65, 173, 164, 65);
 #endif
-}
 
+	EvalOptions.bEvaluateNearestSection = EvalOptions.bCanEvaluateNearestSection = true;
+}
 
 UMovieSceneSection* UMovieScene3DTransformTrack::CreateNewSection()
 {
 	return NewObject<UMovieSceneSection>(this, UMovieScene3DTransformSection::StaticClass(), NAME_None, RF_Transactional);
-}
-
-
-TSharedPtr<IMovieSceneTrackInstance> UMovieScene3DTransformTrack::CreateInstance()
-{
-	return MakeShareable( new FMovieScene3DTransformTrackInstance( *this ) );
-}
-
-
-bool UMovieScene3DTransformTrack::Eval( float Position, float LastPosition, FVector& InOutTranslation, FRotator& InOutRotation, FVector& InOutScale ) const
-{
-	const UMovieSceneSection* Section = MovieSceneHelpers::FindNearestSectionAtTime( Sections, Position );
-
-	if( Section )
-	{
-		const UMovieScene3DTransformSection* TransformSection = CastChecked<UMovieScene3DTransformSection>( Section );
-
-		if (!Section->IsInfinite())
-		{
-			Position = FMath::Clamp(Position, Section->GetStartTime(), Section->GetEndTime());
-		}
-
-		// Evaluate translation,rotation, and scale curves.  If no keys were found on one of these, that component of the transform will remain unchained
-		TransformSection->EvalTranslation( Position, InOutTranslation );
-		TransformSection->EvalRotation( Position, InOutRotation );
-		TransformSection->EvalScale( Position, InOutScale );
-	}
-
-	return Section != nullptr;
-}
+}
\ No newline at end of file
diff --git a/Engine/Source/Runtime/MovieSceneTracks/Private/Tracks/MovieSceneActorReferenceTrack.cpp b/Engine/Source/Runtime/MovieSceneTracks/Private/Tracks/MovieSceneActorReferenceTrack.cpp
--- a/Engine/Source/Runtime/MovieSceneTracks/Private/Tracks/MovieSceneActorReferenceTrack.cpp
+++ b/Engine/Source/Runtime/MovieSceneTracks/Private/Tracks/MovieSceneActorReferenceTrack.cpp
@@ -4,7 +4,7 @@
 #include "MovieSceneActorReferenceSection.h"
 #include "MovieSceneActorReferenceTrack.h"
 #include "IMovieScenePlayer.h"
-#include "MovieSceneActorReferenceTrackInstance.h"
+#include "Evaluation/MovieSceneActorReferenceTemplate.h"
 
 
 UMovieSceneActorReferenceTrack::UMovieSceneActorReferenceTrack( const FObjectInitializer& ObjectInitializer )
@@ -17,27 +17,8 @@ UMovieSceneSection* UMovieSceneActorReferenceTrack::CreateNewSection()
 	return NewObject<UMovieSceneSection>(this, UMovieSceneActorReferenceSection::StaticClass(), NAME_None, RF_Transactional);
 }
 
-
-TSharedPtr<IMovieSceneTrackInstance> UMovieSceneActorReferenceTrack::CreateInstance()
-{
-	return MakeShareable( new FMovieSceneActorReferenceTrackInstance( *this ) ); 
-}
-
-
-bool UMovieSceneActorReferenceTrack::Eval( float Position, float LastPosition, FGuid& OutActorReferenceGuid ) const
+FMovieSceneEvalTemplatePtr UMovieSceneActorReferenceTrack::CreateTemplateForSection(const UMovieSceneSection& InSection) const
 {
-	const UMovieSceneSection* Section = MovieSceneHelpers::FindNearestSectionAtTime( Sections, Position );
-
-	if ( Section )
-	{
-		if ( !Section->IsInfinite() )
-		{
-			Position = FMath::Clamp( Position, Section->GetStartTime(), Section->GetEndTime() );
-		}
-
-		OutActorReferenceGuid = CastChecked<UMovieSceneActorReferenceSection>( Section )->Eval( Position );
-	}
-
-	return Section != nullptr;
+	return FMovieSceneActorReferenceSectionTemplate(*CastChecked<UMovieSceneActorReferenceSection>(&InSection), *this);
 }
 
diff --git a/Engine/Source/Runtime/MovieSceneTracks/Private/Tracks/MovieSceneAudioTrack.cpp b/Engine/Source/Runtime/MovieSceneTracks/Private/Tracks/MovieSceneAudioTrack.cpp
--- a/Engine/Source/Runtime/MovieSceneTracks/Private/Tracks/MovieSceneAudioTrack.cpp
+++ b/Engine/Source/Runtime/MovieSceneTracks/Private/Tracks/MovieSceneAudioTrack.cpp
@@ -9,7 +9,6 @@
 #include "Sound/SoundCue.h"
 #include "Runtime/Engine/Classes/Kismet/GameplayStatics.h"
 #include "Runtime/Engine/Public/AudioDecompress.h"
-#include "MovieSceneAudioTrackInstance.h"
 
 
 #define LOCTEXT_NAMESPACE "MovieSceneAudioTrack"
@@ -24,12 +23,6 @@ UMovieSceneAudioTrack::UMovieSceneAudioTrack( const FObjectInitializer& ObjectIn
 }
 
 
-TSharedPtr<IMovieSceneTrackInstance> UMovieSceneAudioTrack::CreateInstance()
-{
-	return MakeShareable( new FMovieSceneAudioTrackInstance( *this ) ); 
-}
-
-
 const TArray<UMovieSceneSection*>& UMovieSceneAudioTrack::GetAllSections() const
 {
 	return AudioSections;
diff --git a/Engine/Source/Runtime/MovieSceneTracks/Private/Tracks/MovieSceneBoolTrack.cpp b/Engine/Source/Runtime/MovieSceneTracks/Private/Tracks/MovieSceneBoolTrack.cpp
--- a/Engine/Source/Runtime/MovieSceneTracks/Private/Tracks/MovieSceneBoolTrack.cpp
+++ b/Engine/Source/Runtime/MovieSceneTracks/Private/Tracks/MovieSceneBoolTrack.cpp
@@ -4,8 +4,7 @@
 #include "MovieSceneBoolSection.h"
 #include "MovieSceneBoolTrack.h"
 #include "IMovieScenePlayer.h"
-#include "MovieSceneBoolTrackInstance.h"
-
+#include "Evaluation/MovieScenePropertyTemplates.h"
 
 UMovieSceneSection* UMovieSceneBoolTrack::CreateNewSection()
 {
@@ -13,12 +12,11 @@ UMovieSceneSection* UMovieSceneBoolTrack::CreateNewSection()
 }
 
 
-TSharedPtr<IMovieSceneTrackInstance> UMovieSceneBoolTrack::CreateInstance()
+FMovieSceneEvalTemplatePtr UMovieSceneBoolTrack::CreateTemplateForSection(const UMovieSceneSection& InSection) const
 {
-	return MakeShareable( new FMovieSceneBoolTrackInstance( *this ) );
+	return FMovieSceneBoolPropertySectionTemplate(*CastChecked<const UMovieSceneBoolSection>(&InSection), *this);
 }
 
-
 bool UMovieSceneBoolTrack::Eval( float Position, float LastPostion, bool& InOutBool ) const
 {	
 	const UMovieSceneSection* Section = MovieSceneHelpers::FindNearestSectionAtTime( Sections, Position );
diff --git a/Engine/Source/Runtime/MovieSceneTracks/Private/Tracks/MovieSceneByteTrack.cpp b/Engine/Source/Runtime/MovieSceneTracks/Private/Tracks/MovieSceneByteTrack.cpp
--- a/Engine/Source/Runtime/MovieSceneTracks/Private/Tracks/MovieSceneByteTrack.cpp
+++ b/Engine/Source/Runtime/MovieSceneTracks/Private/Tracks/MovieSceneByteTrack.cpp
@@ -4,8 +4,7 @@
 #include "MovieSceneByteSection.h"
 #include "MovieSceneByteTrack.h"
 #include "IMovieScenePlayer.h"
-#include "MovieSceneByteTrackInstance.h"
-
+#include "Evaluation/MovieScenePropertyTemplates.h"
 
 UMovieSceneByteTrack::UMovieSceneByteTrack( const FObjectInitializer& ObjectInitializer )
 	: Super( ObjectInitializer )
@@ -17,28 +16,9 @@ UMovieSceneSection* UMovieSceneByteTrack::CreateNewSection()
 	return NewObject<UMovieSceneSection>(this, UMovieSceneByteSection::StaticClass(), NAME_None, RF_Transactional);
 }
 
-
-TSharedPtr<IMovieSceneTrackInstance> UMovieSceneByteTrack::CreateInstance()
+FMovieSceneEvalTemplatePtr UMovieSceneByteTrack::CreateTemplateForSection(const UMovieSceneSection& InSection) const
 {
-	return MakeShareable( new FMovieSceneByteTrackInstance( *this ) );
-}
-
-
-bool UMovieSceneByteTrack::Eval( float Position, float LastPostion, uint8& InOutByte ) const
-{	
-	const UMovieSceneSection* Section = MovieSceneHelpers::FindNearestSectionAtTime( Sections, Position );
-
-	if( Section )
-	{
-		if (!Section->IsInfinite())
-		{
-			Position = FMath::Clamp(Position, Section->GetStartTime(), Section->GetEndTime());
-		}
-
-		InOutByte = CastChecked<UMovieSceneByteSection>( Section )->Eval( Position, InOutByte );
-	}
-
-	return Section != nullptr;
+	return FMovieSceneBytePropertySectionTemplate(*CastChecked<UMovieSceneByteSection>(&InSection), *this);
 }
 
 void UMovieSceneByteTrack::SetEnum(UEnum* InEnum)
diff --git a/Engine/Source/Runtime/MovieSceneTracks/Private/Tracks/MovieSceneCameraAnimTrack.cpp b/Engine/Source/Runtime/MovieSceneTracks/Private/Tracks/MovieSceneCameraAnimTrack.cpp
--- a/Engine/Source/Runtime/MovieSceneTracks/Private/Tracks/MovieSceneCameraAnimTrack.cpp
+++ b/Engine/Source/Runtime/MovieSceneTracks/Private/Tracks/MovieSceneCameraAnimTrack.cpp
@@ -4,8 +4,11 @@
 #include "MovieSceneCameraAnimSection.h"
 #include "MovieSceneCameraAnimTrack.h"
 #include "IMovieScenePlayer.h"
-#include "MovieSceneCameraAnimTrackInstance.h"
 #include "Camera/CameraAnim.h"
+#include "Evaluation/MovieSceneCameraAnimTemplate.h"
+#include "Evaluation/MovieSceneEvaluationTemplate.h"
+#include "Compilation/IMovieSceneTemplateGenerator.h"
+#include "Compilation/MovieSceneCompilerRules.h"
 
 #define LOCTEXT_NAMESPACE "MovieSceneCameraAnimTrack"
 
@@ -16,7 +19,7 @@ void UMovieSceneCameraAnimTrack::AddNewCameraAnim(float KeyTime, UCameraAnim* Ca
 	if (NewSection)
 	{
 		NewSection->InitialPlacement(CameraAnimSections, KeyTime, KeyTime + CameraAnim->AnimLength, SupportsMultipleRows());
-		NewSection->SetCameraAnim(CameraAnim);
+		NewSection->AnimData.CameraAnim = CameraAnim;
 
 		AddSection(*NewSection);
 	}
@@ -25,9 +28,23 @@ void UMovieSceneCameraAnimTrack::AddNewCameraAnim(float KeyTime, UCameraAnim* Ca
 /* UMovieSceneTrack interface
 *****************************************************************************/
 
-TSharedPtr<IMovieSceneTrackInstance> UMovieSceneCameraAnimTrack::CreateInstance()
+
+void UMovieSceneCameraAnimTrack::PostCompile(FMovieSceneEvaluationTrack& OutTrack, const FMovieSceneTrackCompilerArgs& Args) const
+{
+	FMovieSceneSharedDataId UniqueId = FMovieSceneAdditiveCameraAnimationTrackTemplate::SharedDataId;
+
+	// Add a new shared track for the additive camera anim. There will only be one of these, and it will apply all the additive camera animations for this object.
+	FMovieSceneEvaluationTrack SharedTrackTemplate(Args.ObjectBindingId);
+	SharedTrackTemplate.DefineAsSingleTemplate(FMovieSceneAdditiveCameraAnimationTrackTemplate());
+	SharedTrackTemplate.SetEvaluationPriority(0xF);
+	
+	Args.Generator.AddSharedTrack(MoveTemp(SharedTrackTemplate), UniqueId, *this);
+}
+
+
+TInlineValue<FMovieSceneSegmentCompilerRules> UMovieSceneCameraAnimTrack::GetTrackCompilerRules() const
 {
-	return MakeShareable(new FMovieSceneCameraAnimTrackInstance(*this));
+	return FMovieSceneAdditiveCameraRules(this);
 }
 
 
diff --git a/Engine/Source/Runtime/MovieSceneTracks/Private/Tracks/MovieSceneCameraCutTrack.cpp b/Engine/Source/Runtime/MovieSceneTracks/Private/Tracks/MovieSceneCameraCutTrack.cpp
--- a/Engine/Source/Runtime/MovieSceneTracks/Private/Tracks/MovieSceneCameraCutTrack.cpp
+++ b/Engine/Source/Runtime/MovieSceneTracks/Private/Tracks/MovieSceneCameraCutTrack.cpp
@@ -4,13 +4,10 @@
 #include "MovieSceneCameraCutSection.h"
 #include "MovieSceneCameraCutTrack.h"
 #include "MovieSceneCommonHelpers.h"
-#include "IMovieScenePlayer.h"
-#include "MovieSceneCameraCutTrackInstance.h"
-
+#include "Evaluation/MovieSceneCameraCutTemplate.h"
 
 #define LOCTEXT_NAMESPACE "MovieSceneCameraCutTrack"
 
-
 /* UMovieSceneCameraCutTrack interface
  *****************************************************************************/
 UMovieSceneCameraCutTrack::UMovieSceneCameraCutTrack( const FObjectInitializer& ObjectInitializer )
@@ -19,6 +16,9 @@ UMovieSceneCameraCutTrack::UMovieSceneCameraCutTrack( const FObjectInitializer&
 #if WITH_EDITORONLY_DATA
 	TrackTint = FColor(0, 0, 0, 65);
 #endif
+
+	// By default, don't evaluate camera cuts in pre and postroll
+	EvalOptions.bEvaluateInPreroll = EvalOptions.bEvaluateInPostroll = false;
 }
 
 
@@ -77,12 +77,6 @@ void UMovieSceneCameraCutTrack::AddSection(UMovieSceneSection& Section)
 }
 
 
-TSharedPtr<IMovieSceneTrackInstance> UMovieSceneCameraCutTrack::CreateInstance()
-{
-	return MakeShareable(new FMovieSceneCameraCutTrackInstance(*this)); 
-}
-
-
 UMovieSceneSection* UMovieSceneCameraCutTrack::CreateNewSection()
 {
 	return NewObject<UMovieSceneCameraCutSection>(this, NAME_None, RF_Transactional);
@@ -153,5 +147,4 @@ float UMovieSceneCameraCutTrack::FindEndTimeForCameraCut( float StartTime )
 	return EndTime;
 }
 
-
 #undef LOCTEXT_NAMESPACE
\ No newline at end of file
diff --git a/Engine/Source/Runtime/MovieSceneTracks/Private/Tracks/MovieSceneCameraShakeTrack.cpp b/Engine/Source/Runtime/MovieSceneTracks/Private/Tracks/MovieSceneCameraShakeTrack.cpp
--- a/Engine/Source/Runtime/MovieSceneTracks/Private/Tracks/MovieSceneCameraShakeTrack.cpp
+++ b/Engine/Source/Runtime/MovieSceneTracks/Private/Tracks/MovieSceneCameraShakeTrack.cpp
@@ -4,7 +4,10 @@
 #include "MovieSceneCameraShakeSection.h"
 #include "MovieSceneCameraShakeTrack.h"
 #include "IMovieScenePlayer.h"
-#include "MovieSceneCameraShakeTrackInstance.h"
+#include "Evaluation/MovieSceneCameraAnimTemplate.h"
+#include "Evaluation/MovieSceneEvaluationTemplate.h"
+#include "Compilation/IMovieSceneTemplateGenerator.h"
+#include "Compilation/MovieSceneCompilerRules.h"
 
 #define LOCTEXT_NAMESPACE "MovieSceneCameraShakeTrack"
 
@@ -15,12 +18,29 @@ void UMovieSceneCameraShakeTrack::AddNewCameraShake(float KeyTime, TSubclassOf<U
 	{
 		// #fixme get length
 		NewSection->InitialPlacement(CameraShakeSections, KeyTime, KeyTime + 5.f /*AnimSequence->SequenceLength*/, SupportsMultipleRows());
-		NewSection->SetCameraShakeClass(ShakeClass);
+		NewSection->ShakeData.ShakeClass = ShakeClass;
 		
 		AddSection(*NewSection);
 	}
 }
 
+void UMovieSceneCameraShakeTrack::PostCompile(FMovieSceneEvaluationTrack& OutTrack, const FMovieSceneTrackCompilerArgs& Args) const
+{
+	FMovieSceneSharedDataId UniqueId = FMovieSceneAdditiveCameraAnimationTrackTemplate::SharedDataId;
+
+	// Add a new shared track for the additive camera anim. There will only be one of these, and it will apply all the additive camera animations for this object.
+	FMovieSceneEvaluationTrack SharedTrackTemplate(Args.ObjectBindingId);
+	SharedTrackTemplate.DefineAsSingleTemplate(FMovieSceneAdditiveCameraAnimationTrackTemplate());
+	SharedTrackTemplate.SetEvaluationPriority(0xF);
+	
+	Args.Generator.AddSharedTrack(MoveTemp(SharedTrackTemplate), UniqueId, *this);
+}
+
+TInlineValue<FMovieSceneSegmentCompilerRules> UMovieSceneCameraShakeTrack::GetTrackCompilerRules() const
+{
+	return FMovieSceneAdditiveCameraRules(this);
+}
+
 #if WITH_EDITORONLY_DATA
 FText UMovieSceneCameraShakeTrack::GetDisplayName() const
 {
@@ -34,11 +54,6 @@ FText UMovieSceneCameraShakeTrack::GetDisplayName() const
 /* UMovieSceneTrack interface
 *****************************************************************************/
 
-TSharedPtr<IMovieSceneTrackInstance> UMovieSceneCameraShakeTrack::CreateInstance()
-{
-	return MakeShareable(new FMovieSceneCameraShakeTrackInstance(*this));
-}
-
 
 const TArray<UMovieSceneSection*>& UMovieSceneCameraShakeTrack::GetAllSections() const
 {
@@ -95,19 +110,4 @@ TRange<float> UMovieSceneCameraShakeTrack::GetSectionBoundaries() const
 }
 
 
-void UMovieSceneCameraShakeTrack::GetCameraShakeSectionsAtTime(float Time, TArray<UMovieSceneCameraShakeSection*>& OutSections)
-{
-	OutSections.Empty();
-
-	for (auto Section : CameraShakeSections)
-	{
-		UMovieSceneCameraShakeSection* const CSSection = dynamic_cast<UMovieSceneCameraShakeSection*>(Section);
-		if (CSSection && CSSection->IsTimeWithinSection(Time))
-		{
-			OutSections.Add(CSSection);
-		}
-	}
-}
-
-
 #undef LOCTEXT_NAMESPACE
diff --git a/Engine/Source/Runtime/MovieSceneTracks/Private/Tracks/MovieSceneCinematicShotTrack.cpp b/Engine/Source/Runtime/MovieSceneTracks/Private/Tracks/MovieSceneCinematicShotTrack.cpp
--- a/Engine/Source/Runtime/MovieSceneTracks/Private/Tracks/MovieSceneCinematicShotTrack.cpp
+++ b/Engine/Source/Runtime/MovieSceneTracks/Private/Tracks/MovieSceneCinematicShotTrack.cpp
@@ -4,8 +4,8 @@
 #include "MovieSceneSequence.h"
 #include "MovieSceneCinematicShotSection.h"
 #include "MovieSceneCinematicShotTrack.h"
-#include "MovieSceneCinematicShotTrackInstance.h"
 #include "MovieSceneCommonHelpers.h"
+#include "Compilation/MovieSceneCompilerRules.h"
 
 
 #define LOCTEXT_NAMESPACE "MovieSceneCinematicShotTrack"
@@ -57,12 +57,6 @@ void UMovieSceneCinematicShotTrack::AddSection(UMovieSceneSection& Section)
 }
 
 
-TSharedPtr<IMovieSceneTrackInstance> UMovieSceneCinematicShotTrack::CreateInstance()
-{
-	return MakeShareable(new FMovieSceneCinematicShotTrackInstance(*this)); 
-}
-
-
 UMovieSceneSection* UMovieSceneCinematicShotTrack::CreateNewSection()
 {
 	return NewObject<UMovieSceneCinematicShotSection>(this, NAME_None, RF_Transactional);
@@ -82,6 +76,52 @@ bool UMovieSceneCinematicShotTrack::SupportsMultipleRows() const
 	return true;
 }
 
+TInlineValue<FMovieSceneSegmentCompilerRules> UMovieSceneCinematicShotTrack::GetTrackCompilerRules() const
+{
+	// Apply a high pass filter to overlapping sections such that only the highest row in a track wins
+	struct FCinematicShotTrackCompilerRules : FMovieSceneSegmentCompilerRules
+	{
+		virtual void BlendSegment(FMovieSceneSegment& Segment, const TArrayView<const FMovieSceneSectionData>& SourceData) const
+		{
+			MovieSceneSegmentCompiler::BlendSegmentHighPass(Segment, SourceData);
+		}
+	};
+	return FCinematicShotTrackCompilerRules();
+}
+
+TInlineValue<FMovieSceneSegmentCompilerRules> UMovieSceneCinematicShotTrack::GetRowCompilerRules() const
+{
+	class FCinematicRowRules : public FMovieSceneSegmentCompilerRules
+	{
+		virtual void BlendSegment(FMovieSceneSegment& Segment, const TArrayView<const FMovieSceneSectionData>& SourceData) const override
+		{
+			// Sort everything by priority, then latest start time wins
+			if (Segment.Impls.Num() <= 1)
+			{
+				return;
+			}
+
+			Segment.Impls.Sort(
+				[&](const FSectionEvaluationData& A, const FSectionEvaluationData& B)
+				{
+					const FMovieSceneSectionData& DataA = SourceData[A.ImplIndex];
+					const FMovieSceneSectionData& DataB = SourceData[B.ImplIndex];
+
+					if (DataA.Priority == DataB.Priority)
+					{
+						return TRangeBound<float>::MaxLower(DataA.Bounds.GetLowerBound(), DataB.Bounds.GetLowerBound()) == DataA.Bounds.GetLowerBound();
+					}
+					return DataA.Priority > DataB.Priority;
+				}
+			);
+
+			Segment.Impls.RemoveAt(1, Segment.Impls.Num() - 1, true);
+		}
+	};
+
+	return FCinematicRowRules();
+}
+
 #if WITH_EDITOR
 void UMovieSceneCinematicShotTrack::OnSectionMoved(UMovieSceneSection& Section)
 {
diff --git a/Engine/Source/Runtime/MovieSceneTracks/Private/Tracks/MovieSceneColorTrack.cpp b/Engine/Source/Runtime/MovieSceneTracks/Private/Tracks/MovieSceneColorTrack.cpp
--- a/Engine/Source/Runtime/MovieSceneTracks/Private/Tracks/MovieSceneColorTrack.cpp
+++ b/Engine/Source/Runtime/MovieSceneTracks/Private/Tracks/MovieSceneColorTrack.cpp
@@ -3,8 +3,7 @@
 #include "MovieSceneTracksPrivatePCH.h"
 #include "MovieSceneColorSection.h"
 #include "MovieSceneColorTrack.h"
-#include "IMovieScenePlayer.h"
-#include "MovieSceneColorTrackInstance.h"
+#include "Evaluation/MovieSceneColorTemplate.h"
 
 
 UMovieSceneColorTrack::UMovieSceneColorTrack( const FObjectInitializer& ObjectInitializer )
@@ -17,27 +16,7 @@ UMovieSceneSection* UMovieSceneColorTrack::CreateNewSection()
 	return NewObject<UMovieSceneSection>(this, UMovieSceneColorSection::StaticClass(), NAME_None, RF_Transactional);
 }
 
-
-TSharedPtr<IMovieSceneTrackInstance> UMovieSceneColorTrack::CreateInstance()
+FMovieSceneEvalTemplatePtr UMovieSceneColorTrack::CreateTemplateForSection(const UMovieSceneSection& Section) const
 {
-	return MakeShareable( new FMovieSceneColorTrackInstance( *this ) ); 
+	return FMovieSceneColorSectionTemplate(*CastChecked<const UMovieSceneColorSection>(&Section), *this);
 }
-
-
-bool UMovieSceneColorTrack::Eval( float Position, float LastPosition, FLinearColor& InOutColor ) const
-{
-	const UMovieSceneSection* Section = MovieSceneHelpers::FindNearestSectionAtTime( Sections, Position );
-
-	if( Section )
-	{
-		if (!Section->IsInfinite())
-		{
-			Position = FMath::Clamp(Position, Section->GetStartTime(), Section->GetEndTime());
-		}
-
-		InOutColor = CastChecked<UMovieSceneColorSection>( Section )->Eval( Position, InOutColor );
-	}
-
-	return Section != nullptr;
-}
-
diff --git a/Engine/Source/Runtime/MovieSceneTracks/Private/Tracks/MovieSceneEventTrack.cpp b/Engine/Source/Runtime/MovieSceneTracks/Private/Tracks/MovieSceneEventTrack.cpp
--- a/Engine/Source/Runtime/MovieSceneTracks/Private/Tracks/MovieSceneEventTrack.cpp
+++ b/Engine/Source/Runtime/MovieSceneTracks/Private/Tracks/MovieSceneEventTrack.cpp
@@ -3,70 +3,12 @@
 #include "MovieSceneTracksPrivatePCH.h"
 #include "MovieSceneEventSection.h"
 #include "MovieSceneEventTrack.h"
-#include "MovieSceneEventTrackInstance.h"
+#include "Evaluation/MovieSceneEventTemplate.h"
+#include "Evaluation/MovieSceneEvaluationTrack.h"
 
 
 #define LOCTEXT_NAMESPACE "MovieSceneEventTrack"
 
-/* UMovieSceneEventTrack interface
- *****************************************************************************/
-
-bool UMovieSceneEventTrack::AddKeyToSection(float Time, FName EventName)
-{
-	UMovieSceneSection* TargetSection = MovieSceneHelpers::FindNearestSectionAtTime(Sections, Time);
-
-	if (TargetSection == nullptr)
-	{
-		TargetSection = CreateNewSection();
-		TargetSection->SetStartTime(Time);
-		TargetSection->SetEndTime(Time);
-
-		Sections.Add(TargetSection);
-	}
-
-	UMovieSceneEventSection* EventSection = Cast<UMovieSceneEventSection>(TargetSection);
-
-	if (EventSection == nullptr)
-	{
-		return false;
-	}
-
-	EventSection->AddKey(Time, EventName);
-
-	return true;
-}
-
-
-void UMovieSceneEventTrack::TriggerEvents(float Position, float LastPosition, IMovieScenePlayer& Player)
-{
-	if ((Sections.Num() == 0) || (Position == LastPosition))
-	{
-		return;
-	}
-
-	// Don't allow events to fire when playback is in a stopped state. This can occur when stopping 
-	// playback and returning the current position to the start of playback. It's not desireable to have 
-	// all the events from the last playback position to the start of playback be fired.
-	if (Player.GetPlaybackStatus() == EMovieScenePlayerStatus::Stopped)
-	{
-		return;
-	}
-
-	bool Backwards = Position < LastPosition;
-
-	if ((!Backwards && !bFireEventsWhenForwards) ||
-		(Backwards && !bFireEventsWhenBackwards))
-	{
-		return;
-	}
-
-	TArray<UMovieSceneSection*> TraversedSections = MovieSceneHelpers::GetTraversedSections(Sections, Position, LastPosition);
-	for (auto EventSection : TraversedSections)
-	{
-		CastChecked<UMovieSceneEventSection>(EventSection)->TriggerEvents(Position, LastPosition, Player);
-	}
-}
-
 
 /* UMovieSceneTrack interface
  *****************************************************************************/
@@ -77,12 +19,6 @@ void UMovieSceneEventTrack::AddSection(UMovieSceneSection& Section)
 }
 
 
-TSharedPtr<IMovieSceneTrackInstance> UMovieSceneEventTrack::CreateInstance()
-{
-	return MakeShareable(new FMovieSceneEventTrackInstance(*this));
-}
-
-
 UMovieSceneSection* UMovieSceneEventTrack::CreateNewSection()
 {
 	return NewObject<UMovieSceneSection>(this, UMovieSceneEventSection::StaticClass(), NAME_None, RF_Transactional);
@@ -131,6 +67,20 @@ void UMovieSceneEventTrack::RemoveSection(UMovieSceneSection& Section)
 	Sections.Remove(&Section);
 }
 
+FMovieSceneEvalTemplatePtr UMovieSceneEventTrack::CreateTemplateForSection(const UMovieSceneSection& InSection) const
+{
+	return FMovieSceneEventSectionTemplate(*CastChecked<UMovieSceneEventSection>(&InSection), *this);
+}
+
+void UMovieSceneEventTrack::PostCompile(FMovieSceneEvaluationTrack& Track, const FMovieSceneTrackCompilerArgs& Args) const
+{
+	static FName Events("Events");
+	Track.SetEvaluationGroup(Events);
+	// Evaluate events really early on in the frame, after the spawn track (which always evaluates first)
+	Track.SetEvaluationPriority(GetEvaluationPriority());
+	Track.SetEvaluationMethod(EEvaluationMethod::Swept);
+}
+
 #if WITH_EDITORONLY_DATA
 
 FText UMovieSceneEventTrack::GetDefaultDisplayName() const
diff --git a/Engine/Source/Runtime/MovieSceneTracks/Private/Tracks/MovieSceneFadeTrack.cpp b/Engine/Source/Runtime/MovieSceneTracks/Private/Tracks/MovieSceneFadeTrack.cpp
--- a/Engine/Source/Runtime/MovieSceneTracks/Private/Tracks/MovieSceneFadeTrack.cpp
+++ b/Engine/Source/Runtime/MovieSceneTracks/Private/Tracks/MovieSceneFadeTrack.cpp
@@ -3,7 +3,8 @@
 #include "MovieSceneTracksPrivatePCH.h"
 #include "MovieSceneFadeSection.h"
 #include "MovieSceneFadeTrack.h"
-#include "MovieSceneFadeTrackInstance.h"
+#include "Evaluation/MovieSceneFadeTemplate.h"
+#include "Evaluation/MovieSceneEvaluationTrack.h"
 
 
 #define LOCTEXT_NAMESPACE "MovieSceneFadeTrack"
@@ -11,18 +12,22 @@
 
 /* UMovieSceneFadeTrack interface
  *****************************************************************************/
-
-TSharedPtr<IMovieSceneTrackInstance> UMovieSceneFadeTrack::CreateInstance()
+UMovieSceneFadeTrack::UMovieSceneFadeTrack(const FObjectInitializer& Init)
+	: Super(Init)
 {
-	return MakeShareable(new FMovieSceneFadeTrackInstance(*this)); 
+	EvalOptions.bCanEvaluateNearestSection = true;
+	EvalOptions.bEvaluateNearestSection = false;
 }
 
-
 UMovieSceneSection* UMovieSceneFadeTrack::CreateNewSection()
 {
 	return NewObject<UMovieSceneSection>(this, UMovieSceneFadeSection::StaticClass(), NAME_None, RF_Transactional);
 }
 
+FMovieSceneEvalTemplatePtr UMovieSceneFadeTrack::CreateTemplateForSection(const UMovieSceneSection& InSection) const
+{
+	return FMovieSceneFadeSectionTemplate(*CastChecked<UMovieSceneFadeSection>(&InSection));
+}
 
 #if WITH_EDITORONLY_DATA
 FText UMovieSceneFadeTrack::GetDefaultDisplayName() const
diff --git a/Engine/Source/Runtime/MovieSceneTracks/Private/Tracks/MovieSceneFloatTrack.cpp b/Engine/Source/Runtime/MovieSceneTracks/Private/Tracks/MovieSceneFloatTrack.cpp
--- a/Engine/Source/Runtime/MovieSceneTracks/Private/Tracks/MovieSceneFloatTrack.cpp
+++ b/Engine/Source/Runtime/MovieSceneTracks/Private/Tracks/MovieSceneFloatTrack.cpp
@@ -4,7 +4,7 @@
 #include "MovieSceneFloatSection.h"
 #include "MovieSceneFloatTrack.h"
 #include "IMovieScenePlayer.h"
-#include "MovieSceneFloatTrackInstance.h"
+#include "Evaluation/MovieScenePropertyTemplates.h"
 
 
 UMovieSceneFloatTrack::UMovieSceneFloatTrack( const FObjectInitializer& ObjectInitializer )
@@ -17,13 +17,11 @@ UMovieSceneSection* UMovieSceneFloatTrack::CreateNewSection()
 	return NewObject<UMovieSceneSection>(this, UMovieSceneFloatSection::StaticClass(), NAME_None, RF_Transactional);
 }
 
-
-TSharedPtr<IMovieSceneTrackInstance> UMovieSceneFloatTrack::CreateInstance()
+FMovieSceneEvalTemplatePtr UMovieSceneFloatTrack::CreateTemplateForSection(const UMovieSceneSection& InSection) const
 {
-	return MakeShareable( new FMovieSceneFloatTrackInstance( *this ) ); 
+	return FMovieSceneFloatPropertySectionTemplate(*CastChecked<const UMovieSceneFloatSection>(&InSection), *this);
 }
 
-
 bool UMovieSceneFloatTrack::Eval( float Position, float LastPosition, float& InOutFloat ) const
 {
 	const UMovieSceneSection* Section = MovieSceneHelpers::FindNearestSectionAtTime( Sections, Position );
@@ -39,5 +37,4 @@ bool UMovieSceneFloatTrack::Eval( float Position, float LastPosition, float& InO
 	}
 
 	return Section != nullptr;
-}
-
+}
\ No newline at end of file
diff --git a/Engine/Source/Runtime/MovieSceneTracks/Private/Tracks/MovieSceneIntegerTrack.cpp b/Engine/Source/Runtime/MovieSceneTracks/Private/Tracks/MovieSceneIntegerTrack.cpp
new file mode 100644
--- /dev/null
+++ b/Engine/Source/Runtime/MovieSceneTracks/Private/Tracks/MovieSceneIntegerTrack.cpp
@@ -0,0 +1,22 @@
+// Copyright 1998-2016 Epic Games, Inc. All Rights Reserved.
+
+#include "MovieSceneTracksPrivatePCH.h"
+#include "MovieSceneIntegerSection.h"
+#include "MovieSceneIntegerTrack.h"
+#include "IMovieScenePlayer.h"
+#include "Evaluation/MovieScenePropertyTemplates.h"
+
+UMovieSceneIntegerTrack::UMovieSceneIntegerTrack( const FObjectInitializer& ObjectInitializer )
+	: Super( ObjectInitializer )
+{ }
+
+
+UMovieSceneSection* UMovieSceneIntegerTrack::CreateNewSection()
+{
+	return NewObject<UMovieSceneSection>(this, UMovieSceneIntegerSection::StaticClass(), NAME_None, RF_Transactional);
+}
+
+FMovieSceneEvalTemplatePtr UMovieSceneIntegerTrack::CreateTemplateForSection(const UMovieSceneSection& InSection) const
+{
+	return FMovieSceneIntegerPropertySectionTemplate(*CastChecked<UMovieSceneIntegerSection>(&InSection), *this);
+}
diff --git a/Engine/Source/Runtime/MovieSceneTracks/Private/Tracks/MovieSceneLevelVisibilityTrack.cpp b/Engine/Source/Runtime/MovieSceneTracks/Private/Tracks/MovieSceneLevelVisibilityTrack.cpp
--- a/Engine/Source/Runtime/MovieSceneTracks/Private/Tracks/MovieSceneLevelVisibilityTrack.cpp
+++ b/Engine/Source/Runtime/MovieSceneTracks/Private/Tracks/MovieSceneLevelVisibilityTrack.cpp
@@ -4,7 +4,9 @@
 #include "MovieSceneLevelVisibilitySection.h"
 #include "MovieSceneLevelVisibilityTrack.h"
 #include "IMovieScenePlayer.h"
-#include "MovieSceneLevelVisibilityTrackInstance.h"
+#include "Compilation/IMovieSceneTemplateGenerator.h"
+#include "Evaluation/MovieSceneLevelVisibilityTemplate.h"
+#include "Evaluation/MovieSceneEvaluationTrack.h"
 
 
 #define LOCTEXT_NAMESPACE "MovieSceneLevelVisibilityTrack"
@@ -15,9 +17,18 @@ UMovieSceneLevelVisibilityTrack::UMovieSceneLevelVisibilityTrack( const FObjectI
 }
 
 
-TSharedPtr<IMovieSceneTrackInstance> UMovieSceneLevelVisibilityTrack::CreateInstance()
+void UMovieSceneLevelVisibilityTrack::PostCompile(FMovieSceneEvaluationTrack& OutTrack, const FMovieSceneTrackCompilerArgs& Args) const
 {
-	return MakeShareable( new FMovieSceneLevelVisibilityTrackInstance( *this ) ); 
+	// Set priority to highest possible
+	OutTrack.SetEvaluationPriority(GetEvaluationPriority());
+
+	FMovieSceneSharedDataId UniqueId = FMovieSceneLevelVisibilitySharedTrack::GetSharedDataKey().UniqueId;
+
+	FMovieSceneEvaluationTrack ActuatorTemplate(Args.ObjectBindingId);
+	ActuatorTemplate.DefineAsSingleTemplate(FMovieSceneLevelVisibilitySharedTrack());
+	ActuatorTemplate.SetEvaluationPriority(GetEvaluationPriority() - 1);
+	
+	Args.Generator.AddSharedTrack(MoveTemp(ActuatorTemplate), UniqueId, *this);
 }
 
 
diff --git a/Engine/Source/Runtime/MovieSceneTracks/Private/Tracks/MovieSceneMaterialTrack.cpp b/Engine/Source/Runtime/MovieSceneTracks/Private/Tracks/MovieSceneMaterialTrack.cpp
--- a/Engine/Source/Runtime/MovieSceneTracks/Private/Tracks/MovieSceneMaterialTrack.cpp
+++ b/Engine/Source/Runtime/MovieSceneTracks/Private/Tracks/MovieSceneMaterialTrack.cpp
@@ -4,7 +4,7 @@
 #include "MovieSceneParameterSection.h"
 #include "MovieSceneMaterialTrack.h"
 #include "IMovieScenePlayer.h"
-#include "MovieSceneMaterialTrackInstance.h"
+#include "Evaluation/MovieSceneParameterTemplate.h"
 
 
 UMovieSceneMaterialTrack::UMovieSceneMaterialTrack(const FObjectInitializer& ObjectInitializer)
@@ -105,20 +105,6 @@ void UMovieSceneMaterialTrack::AddColorParameterKey(FName ParameterName, float T
 }
 
 
-void UMovieSceneMaterialTrack::Eval(float Position, TArray<FScalarParameterNameAndValue>& OutScalarValues, TArray<FColorParameterNameAndValue>& OutColorValues) const
-{
-	TArray<FVectorParameterNameAndValue> OutVectorValues;
-	for (UMovieSceneSection* Section : Sections)
-	{
-		if ( Section->IsActive() )
-		{
-			UMovieSceneParameterSection* ParameterSection = Cast<UMovieSceneParameterSection>( Section );
-			ParameterSection->Eval( Position, OutScalarValues, OutVectorValues, OutColorValues );
-		}
-	}
-}
-
-
 UMovieSceneComponentMaterialTrack::UMovieSceneComponentMaterialTrack(const FObjectInitializer& ObjectInitializer)
 	: Super(ObjectInitializer)
 {
@@ -126,9 +112,9 @@ UMovieSceneComponentMaterialTrack::UMovieSceneComponentMaterialTrack(const FObje
 }
 
 
-TSharedPtr<IMovieSceneTrackInstance> UMovieSceneComponentMaterialTrack::CreateInstance()
+FMovieSceneEvalTemplatePtr UMovieSceneComponentMaterialTrack::CreateTemplateForSection(const UMovieSceneSection& InSection) const
 {
-	return MakeShareable(new FMovieSceneComponentMaterialTrackInstance(*this));
+	return FMovieSceneComponentMaterialSectionTemplate(*CastChecked<UMovieSceneParameterSection>(&InSection), *this);
 }
 
 
diff --git a/Engine/Source/Runtime/MovieSceneTracks/Private/Tracks/MovieSceneParticleParameterTrack.cpp b/Engine/Source/Runtime/MovieSceneTracks/Private/Tracks/MovieSceneParticleParameterTrack.cpp
--- a/Engine/Source/Runtime/MovieSceneTracks/Private/Tracks/MovieSceneParticleParameterTrack.cpp
+++ b/Engine/Source/Runtime/MovieSceneTracks/Private/Tracks/MovieSceneParticleParameterTrack.cpp
@@ -4,7 +4,7 @@
 #include "MovieSceneParameterSection.h"
 #include "MovieSceneParticleParameterTrack.h"
 #include "IMovieScenePlayer.h"
-#include "MovieSceneParticleParameterTrackInstance.h"
+#include "Evaluation/MovieSceneParticleParameterTemplate.h"
 
 #define LOCTEXT_NAMESPACE "ParticleParameterTrack"
 
@@ -16,9 +16,9 @@ UMovieSceneParticleParameterTrack::UMovieSceneParticleParameterTrack( const FObj
 #endif
 }
 
-TSharedPtr<IMovieSceneTrackInstance> UMovieSceneParticleParameterTrack::CreateInstance()
+FMovieSceneEvalTemplatePtr UMovieSceneParticleParameterTrack::CreateTemplateForSection(const UMovieSceneSection& InSection) const
 {
-	return MakeShareable( new FMovieSceneParticleParameterTrackInstance( *this ) );
+	return FMovieSceneParticleParameterSectionTemplate(*CastChecked<UMovieSceneParameterSection>(&InSection), *this);
 }
 
 UMovieSceneSection* UMovieSceneParticleParameterTrack::CreateNewSection()
@@ -115,17 +115,4 @@ void UMovieSceneParticleParameterTrack::AddColorParameterKey( FName ParameterNam
 	}
 	NearestSection->AddColorParameterKey( ParameterName, Time, Value );
 }
-
-void UMovieSceneParticleParameterTrack::Eval( float Position, TArray<FScalarParameterNameAndValue>& OutScalarValues, TArray<FVectorParameterNameAndValue>& OutVectorValues, TArray<FColorParameterNameAndValue>& OutColorValues ) const
-{
-	for ( UMovieSceneSection* Section : Sections )
-	{
-		UMovieSceneParameterSection* ParameterSection = Cast<UMovieSceneParameterSection>( Section );
-		if ( ParameterSection != nullptr )
-		{
-			ParameterSection->Eval( Position, OutScalarValues, OutVectorValues, OutColorValues );
-		}
-	}
-}
-
 #undef LOCTEXT_NAMESPACE
diff --git a/Engine/Source/Runtime/MovieSceneTracks/Private/Tracks/MovieSceneParticleTrack.cpp b/Engine/Source/Runtime/MovieSceneTracks/Private/Tracks/MovieSceneParticleTrack.cpp
--- a/Engine/Source/Runtime/MovieSceneTracks/Private/Tracks/MovieSceneParticleTrack.cpp
+++ b/Engine/Source/Runtime/MovieSceneTracks/Private/Tracks/MovieSceneParticleTrack.cpp
@@ -4,7 +4,6 @@
 #include "MovieSceneParticleSection.h"
 #include "MovieSceneParticleTrack.h"
 #include "IMovieScenePlayer.h"
-#include "MovieSceneParticleTrackInstance.h"
 
 
 #define LOCTEXT_NAMESPACE "MovieSceneParticleTrack"
@@ -19,12 +18,6 @@ UMovieSceneParticleTrack::UMovieSceneParticleTrack( const FObjectInitializer& Ob
 }
 
 
-TSharedPtr<IMovieSceneTrackInstance> UMovieSceneParticleTrack::CreateInstance()
-{
-	return MakeShareable( new FMovieSceneParticleTrackInstance( *this ) ); 
-}
-
-
 const TArray<UMovieSceneSection*>& UMovieSceneParticleTrack::GetAllSections() const
 {
 	return ParticleSections;
diff --git a/Engine/Source/Runtime/MovieSceneTracks/Private/Tracks/MovieScenePropertyTrack.cpp b/Engine/Source/Runtime/MovieSceneTracks/Private/Tracks/MovieScenePropertyTrack.cpp
--- a/Engine/Source/Runtime/MovieSceneTracks/Private/Tracks/MovieScenePropertyTrack.cpp
+++ b/Engine/Source/Runtime/MovieSceneTracks/Private/Tracks/MovieScenePropertyTrack.cpp
@@ -7,6 +7,7 @@
 UMovieScenePropertyTrack::UMovieScenePropertyTrack(const FObjectInitializer& ObjectInitializer)
 	: Super(ObjectInitializer)
 {
+	EvalOptions.bEvaluateNearestSection = EvalOptions.bCanEvaluateNearestSection = true;
 }
 
 
diff --git a/Engine/Source/Runtime/MovieSceneTracks/Private/Tracks/MovieSceneSkeletalAnimationTrack.cpp b/Engine/Source/Runtime/MovieSceneTracks/Private/Tracks/MovieSceneSkeletalAnimationTrack.cpp
--- a/Engine/Source/Runtime/MovieSceneTracks/Private/Tracks/MovieSceneSkeletalAnimationTrack.cpp
+++ b/Engine/Source/Runtime/MovieSceneTracks/Private/Tracks/MovieSceneSkeletalAnimationTrack.cpp
@@ -3,9 +3,7 @@
 #include "MovieSceneTracksPrivatePCH.h"
 #include "MovieSceneSkeletalAnimationSection.h"
 #include "MovieSceneSkeletalAnimationTrack.h"
-#include "IMovieScenePlayer.h"
-#include "MovieSceneSkeletalAnimationTrackInstance.h"
-#include "Animation/AnimSequenceBase.h"
+#include "Compilation/MovieSceneCompilerRules.h"
 
 
 #define LOCTEXT_NAMESPACE "MovieSceneSkeletalAnimationTrack"
@@ -20,6 +18,8 @@ UMovieSceneSkeletalAnimationTrack::UMovieSceneSkeletalAnimationTrack(const FObje
 #if WITH_EDITORONLY_DATA
 	TrackTint = FColor(124, 15, 124, 65);
 #endif
+
+	EvalOptions.bEvaluateNearestSection = EvalOptions.bCanEvaluateNearestSection = true;
 }
 
 
@@ -31,7 +31,7 @@ void UMovieSceneSkeletalAnimationTrack::AddNewAnimation(float KeyTime, UAnimSequ
 	UMovieSceneSkeletalAnimationSection* NewSection = Cast<UMovieSceneSkeletalAnimationSection>(CreateNewSection());
 	{
 		NewSection->InitialPlacement(AnimationSections, KeyTime, KeyTime + AnimSequence->SequenceLength, SupportsMultipleRows());
-		NewSection->SetAnimSequence(AnimSequence);
+		NewSection->Params.Animation = AnimSequence;
 	}
 
 	AddSection(*NewSection);
@@ -56,11 +56,6 @@ TArray<UMovieSceneSection*> UMovieSceneSkeletalAnimationTrack::GetAnimSectionsAt
 /* UMovieSceneTrack interface
  *****************************************************************************/
 
-TSharedPtr<IMovieSceneTrackInstance> UMovieSceneSkeletalAnimationTrack::CreateInstance()
-{
-	return MakeShareable(new FMovieSceneSkeletalAnimationTrackInstance(*this)); 
-}
-
 
 const TArray<UMovieSceneSection*>& UMovieSceneSkeletalAnimationTrack::GetAllSections() const
 {
@@ -116,7 +111,6 @@ TRange<float> UMovieSceneSkeletalAnimationTrack::GetSectionBoundaries() const
 	return TRange<float>::Hull(Bounds);
 }
 
-
 #if WITH_EDITORONLY_DATA
 
 FText UMovieSceneSkeletalAnimationTrack::GetDefaultDisplayName() const
diff --git a/Engine/Source/Runtime/MovieSceneTracks/Private/Tracks/MovieSceneSlomoTrack.cpp b/Engine/Source/Runtime/MovieSceneTracks/Private/Tracks/MovieSceneSlomoTrack.cpp
--- a/Engine/Source/Runtime/MovieSceneTracks/Private/Tracks/MovieSceneSlomoTrack.cpp
+++ b/Engine/Source/Runtime/MovieSceneTracks/Private/Tracks/MovieSceneSlomoTrack.cpp
@@ -3,7 +3,8 @@
 #include "MovieSceneTracksPrivatePCH.h"
 #include "MovieSceneSlomoSection.h"
 #include "MovieSceneSlomoTrack.h"
-#include "MovieSceneSlomoTrackInstance.h"
+#include "Evaluation/MovieSceneSlomoTemplate.h"
+#include "Evaluation/MovieSceneEvaluationTrack.h"
 
 
 #define LOCTEXT_NAMESPACE "MovieSceneSlomoTrack"
@@ -11,18 +12,23 @@
 
 /* UMovieSceneEventTrack interface
  *****************************************************************************/
-
-TSharedPtr<IMovieSceneTrackInstance> UMovieSceneSlomoTrack::CreateInstance()
+UMovieSceneSlomoTrack::UMovieSceneSlomoTrack(const FObjectInitializer& Init)
+	: Super(Init)
 {
-	return MakeShareable(new FMovieSceneSlomoTrackInstance(*this)); 
+	EvalOptions.bCanEvaluateNearestSection = true;
+	EvalOptions.bEvaluateNearestSection = false;
 }
 
-
 UMovieSceneSection* UMovieSceneSlomoTrack::CreateNewSection()
 {
 	return NewObject<UMovieSceneSection>(this, UMovieSceneSlomoSection::StaticClass(), NAME_None, RF_Transactional);
 }
 
+FMovieSceneEvalTemplatePtr UMovieSceneSlomoTrack::CreateTemplateForSection(const UMovieSceneSection& InSection) const
+{
+	return FMovieSceneSlomoSectionTemplate(*CastChecked<UMovieSceneSlomoSection>(&InSection));
+}
+
 
 #if WITH_EDITORONLY_DATA
 
diff --git a/Engine/Source/Runtime/MovieSceneTracks/Private/Tracks/MovieSceneSpawnTrack.cpp b/Engine/Source/Runtime/MovieSceneTracks/Private/Tracks/MovieSceneSpawnTrack.cpp
--- a/Engine/Source/Runtime/MovieSceneTracks/Private/Tracks/MovieSceneSpawnTrack.cpp
+++ b/Engine/Source/Runtime/MovieSceneTracks/Private/Tracks/MovieSceneSpawnTrack.cpp
@@ -2,9 +2,11 @@
 
 #include "MovieSceneTracksPrivatePCH.h"
 #include "MovieSceneSpawnTrack.h"
-#include "MovieSceneSpawnTrackInstance.h"
-#include "MovieSceneBoolSection.h"
-
+#include "MovieSceneSpawnSection.h"
+#include "Compilation/MovieSceneCompilerRules.h"
+#include "Evaluation/MovieSceneSpawnTemplate.h"
+#include "Evaluation/MovieSceneEvaluationTrack.h"
+#include "Compilation/IMovieSceneTemplateGenerator.h"
 
 #define LOCTEXT_NAMESPACE "MovieSceneSpawnTrack"
 
@@ -20,45 +22,39 @@ UMovieSceneSpawnTrack::UMovieSceneSpawnTrack(const FObjectInitializer& Obj)
 #endif
 }
 
-
-/* UMovieSceneSpawnTrack interface
- *****************************************************************************/
-
-bool UMovieSceneSpawnTrack::Eval(float Position, float LastPostion, bool& bOutSpawned) const
+void UMovieSceneSpawnTrack::PostLoad()
 {
-	const UMovieSceneBoolSection* Section = Cast<UMovieSceneBoolSection>(MovieSceneHelpers::FindNearestSectionAtTime(Sections, Position));
-	if (!Section)
-	{
-		return false;
-	}
+	TArray<uint8> Bytes;
 
-	if (!Section->IsInfinite())
+	for (int32 Index = 0; Index < Sections.Num(); ++Index)
 	{
-		Position = FMath::Clamp(Position, Section->GetStartTime(), Section->GetEndTime());
+		UMovieSceneBoolSection* BoolSection = ExactCast<UMovieSceneBoolSection>(Sections[Index]);
+		if (BoolSection)
+		{
+			Bytes.Reset();
+
+			FObjectWriter(BoolSection, Bytes);
+			UMovieSceneSpawnSection* NewSection = NewObject<UMovieSceneSpawnSection>(this, NAME_None, RF_Transactional);
+			FObjectReader(NewSection, Bytes);
+
+			Sections[Index] = NewSection;
+		}
 	}
 
-	bOutSpawned = Section->Eval(Position, bOutSpawned);
-	return true;
+	Super::PostLoad();
 }
 
-
 /* UMovieSceneTrack interface
  *****************************************************************************/
 
 UMovieSceneSection* UMovieSceneSpawnTrack::CreateNewSection()
 {
-	UMovieSceneBoolSection* Section = NewObject<UMovieSceneBoolSection>(this, NAME_None, RF_Transactional);
+	UMovieSceneSpawnSection* Section = NewObject<UMovieSceneSpawnSection>(this, NAME_None, RF_Transactional);
 	Section->GetCurve().SetDefaultValue(true);
 	return Section;
 }
 
 
-TSharedPtr<IMovieSceneTrackInstance> UMovieSceneSpawnTrack::CreateInstance()
-{
-	return MakeShareable(new FMovieSceneSpawnTrackInstance(*this));
-}
-
-
 bool UMovieSceneSpawnTrack::HasSection(const UMovieSceneSection& Section) const
 {
 	return Sections.ContainsByPredicate([&](const UMovieSceneSection* In){ return In == &Section; });
@@ -94,6 +90,35 @@ const TArray<UMovieSceneSection*>& UMovieSceneSpawnTrack::GetAllSections() const
 	return Sections;
 }
 
+FMovieSceneEvalTemplatePtr UMovieSceneSpawnTrack::CreateTemplateForSection(const UMovieSceneSection& InSection) const
+{
+	const UMovieSceneSpawnSection* BoolSection = CastChecked<const UMovieSceneSpawnSection>(&InSection);
+	return FMovieSceneSpawnSectionTemplate(*BoolSection);
+}
+
+void UMovieSceneSpawnTrack::GenerateTemplate(const FMovieSceneTrackCompilerArgs& Args) const
+{
+	UMovieScene* ParentMovieScene = GetTypedOuter<UMovieScene>();
+	if (ParentMovieScene && ParentMovieScene->FindPossessable(Args.ObjectBindingId))
+	{
+		return;
+	}
+
+	Super::GenerateTemplate(Args);
+}
+
+void UMovieSceneSpawnTrack::PostCompile(FMovieSceneEvaluationTrack& OutTrack, const FMovieSceneTrackCompilerArgs& Args) const
+{
+	static FName SpawnGroupName("SpawnObjects");
+
+	Args.Generator.FlushGroupImmediately(SpawnGroupName);
+
+	// All objects must be spawned/destroyed before the sequence continues
+	OutTrack.SetEvaluationGroup(SpawnGroupName);
+	// Set priority to highest possible
+	OutTrack.SetEvaluationPriority(GetEvaluationPriority());
+}
+
 #if WITH_EDITORONLY_DATA
 
 ECookOptimizationFlags UMovieSceneSpawnTrack::GetCookOptimizationFlags() const
@@ -101,7 +126,7 @@ ECookOptimizationFlags UMovieSceneSpawnTrack::GetCookOptimizationFlags() const
 	// Since the spawn track denotes the lifetime of a spawnable, if the object is never spawned, we can remove the entire object
 	for (UMovieSceneSection* Section : Sections)
 	{
-		UMovieSceneBoolSection* BoolSection = CastChecked<UMovieSceneBoolSection>(Section);
+		UMovieSceneSpawnSection* BoolSection = CastChecked<UMovieSceneSpawnSection>(Section);
 		if (!BoolSection->IsActive())
 		{
 			continue;
@@ -134,4 +159,4 @@ FText UMovieSceneSpawnTrack::GetDisplayName() const
 #endif
 
 
-#undef LOCTEXT_NAMESPACE
+#undef LOCTEXT_NAMESPACE
\ No newline at end of file
diff --git a/Engine/Source/Runtime/MovieSceneTracks/Private/Tracks/MovieSceneStringTrack.cpp b/Engine/Source/Runtime/MovieSceneTracks/Private/Tracks/MovieSceneStringTrack.cpp
--- a/Engine/Source/Runtime/MovieSceneTracks/Private/Tracks/MovieSceneStringTrack.cpp
+++ b/Engine/Source/Runtime/MovieSceneTracks/Private/Tracks/MovieSceneStringTrack.cpp
@@ -3,7 +3,7 @@
 #include "MovieSceneTracksPrivatePCH.h"
 #include "MovieSceneStringSection.h"
 #include "MovieSceneStringTrack.h"
-#include "MovieSceneStringTrackInstance.h"
+#include "Evaluation/MovieScenePropertyTemplates.h"
 
 
 #define LOCTEXT_NAMESPACE "MovieSceneStringTrack"
@@ -65,18 +65,17 @@ void UMovieSceneStringTrack::AddSection(UMovieSceneSection& Section)
 }
 
 
-TSharedPtr<IMovieSceneTrackInstance> UMovieSceneStringTrack::CreateInstance()
+UMovieSceneSection* UMovieSceneStringTrack::CreateNewSection()
 {
-	return MakeShareable(new FMovieSceneStringTrackInstance(*this));
+	return NewObject<UMovieSceneSection>(this, UMovieSceneStringSection::StaticClass(), NAME_None, RF_Transactional);
 }
 
 
-UMovieSceneSection* UMovieSceneStringTrack::CreateNewSection()
+FMovieSceneEvalTemplatePtr UMovieSceneStringTrack::CreateTemplateForSection(const UMovieSceneSection& InSection) const
 {
-	return NewObject<UMovieSceneSection>(this, UMovieSceneStringSection::StaticClass(), NAME_None, RF_Transactional);
+	return FMovieSceneStringPropertySectionTemplate(*CastChecked<UMovieSceneStringSection>(&InSection), *this);
 }
 
-
 const TArray<UMovieSceneSection*>& UMovieSceneStringTrack::GetAllSections() const
 {
 	return Sections;
diff --git a/Engine/Source/Runtime/MovieSceneTracks/Private/Tracks/MovieSceneSubTrack.cpp b/Engine/Source/Runtime/MovieSceneTracks/Private/Tracks/MovieSceneSubTrack.cpp
--- a/Engine/Source/Runtime/MovieSceneTracks/Private/Tracks/MovieSceneSubTrack.cpp
+++ b/Engine/Source/Runtime/MovieSceneTracks/Private/Tracks/MovieSceneSubTrack.cpp
@@ -5,7 +5,8 @@
 #include "MovieSceneSequence.h"
 #include "MovieSceneSubSection.h"
 #include "MovieSceneSubTrack.h"
-#include "MovieSceneSubTrackInstance.h"
+#include "MovieSceneSegmentCompiler.h"
+#include "IMovieSceneTemplateGenerator.h"
 
 
 #define LOCTEXT_NAMESPACE "MovieSceneSubTrack"
@@ -174,12 +175,6 @@ void UMovieSceneSubTrack::AddSection(UMovieSceneSection& Section)
 }
 
 
-TSharedPtr<IMovieSceneTrackInstance> UMovieSceneSubTrack::CreateInstance()
-{
-	return MakeShareable(new FMovieSceneSubTrackInstance(*this)); 
-}
-
-
 UMovieSceneSection* UMovieSceneSubTrack::CreateNewSection()
 {
 	return NewObject<UMovieSceneSubSection>(this, NAME_None, RF_Transactional);
@@ -241,4 +236,340 @@ FText UMovieSceneSubTrack::GetDefaultDisplayName() const
 }
 #endif
 
+TInlineValue<FMovieSceneSegmentCompilerRules> UMovieSceneSubTrack::GetTrackCompilerRules() const
+{
+	return TInlineValue<FMovieSceneSegmentCompilerRules>();
+}
+
+struct FSubTrackRemapper
+{
+	FSubTrackRemapper(const FMovieSceneTrackCompilerArgs& InArgs, const TArray<UMovieSceneSection*>& InSections)
+		: Args(InArgs)
+		, Sections(InSections)
+	{
+	}
+
+	/** Remap the specified segment, generated from the main sub track, into the outer template's generator */
+	void RemapSegment(const FMovieSceneSegment& Segment)
+	{
+		TSet<FMovieSceneEvaluationFieldSegmentPtr> AllSegmentPtrs;
+
+		// Remap, then add all segment ptrs to the master sequence
+		for (const FSectionEvaluationData& EvalData : Segment.Impls)
+		{
+			const int32 SectionIndex = EvalData.ImplIndex;
+
+			// Remap the sequence in the generator, if we haven't already
+			if (!SectionDataCache.Contains(SectionIndex) && !RemapSubSection(SectionIndex))
+			{
+				// We failed to even remap the track. Leave this ptr out of the segment
+				continue;
+			}
+
+			FSectionData* Cache = SectionDataCache.Find(SectionIndex);
+			check(Cache);
+
+			const UMovieSceneSubSection* SubSection = CastChecked<const UMovieSceneSubSection>(Sections[SectionIndex]);
+
+			// Remap the tracks contained within this sequence only within the range of the section
+			FMovieSceneSequenceTransform InnerSequenceTransform = Args.Generator.GetSequenceTransform(Cache->RemappedID);
+
+			// Add preroll section
+			if (SubSection->Parameters.PrerollTime > 0.f && Segment.Range.GetLowerBound().IsClosed())
+			{
+				TRange<float> PrerollRange = TRange<float>(Segment.Range.GetLowerBoundValue() - SubSection->Parameters.PrerollTime, TRangeBound<float>::FlipInclusion(Segment.Range.GetLowerBound()));
+				if (!PrerollRange.IsEmpty())
+				{
+					// Iterate all ptrs that exist in the preroll range
+					TRange<int32> OverlappingSegments = Cache->Template->EvaluationField.OverlapRange(PrerollRange * InnerSequenceTransform);
+					for (int32 Index = OverlappingSegments.GetLowerBoundValue(); Index < OverlappingSegments.GetUpperBoundValue(); ++Index)
+					{
+						AllSegmentPtrs.Reset();
+
+						TRange<float> InnerSegmentRange = Cache->Template->EvaluationField.Ranges[Index] * InnerSequenceTransform.Inverse();
+						TRange<float> OverlappingRange = TRange<float>::Intersection(InnerSegmentRange, PrerollRange);
+						if (!OverlappingRange.IsEmpty())
+						{
+							RemapPrerollGroup(*Cache, *Cache->Template, Index, AllSegmentPtrs);
+							Args.Generator.AddExternalSegments(OverlappingRange, AllSegmentPtrs.Array());
+						}
+					}
+				}
+			}
+
+			// Add main evaluation section
+			{
+				TRange<int32> OverlappingSegments = Cache->Template->EvaluationField.OverlapRange(Segment.Range * InnerSequenceTransform);
+				for (int32 Index = OverlappingSegments.GetLowerBoundValue(); Index < OverlappingSegments.GetUpperBoundValue(); ++Index)
+				{
+					AllSegmentPtrs.Reset();
+
+					TRange<float> InnerSegmentRange = Cache->Template->EvaluationField.Ranges[Index] * InnerSequenceTransform.Inverse();
+					TRange<float> OverlappingRange = TRange<float>::Intersection(InnerSegmentRange, Segment.Range);
+					if (!OverlappingRange.IsEmpty())
+					{
+						RemapEvaluationGroup(*Cache, Cache->Template->EvaluationField.Groups[Index], AllSegmentPtrs);
+						Args.Generator.AddExternalSegments(OverlappingRange, AllSegmentPtrs.Array());
+					}
+				}
+			}
+
+			// Add postroll section
+			if (SubSection->Parameters.PostrollTime > 0.f && Segment.Range.GetUpperBound().IsClosed())
+			{
+				TRange<float> PostrollRange = TRange<float>(TRangeBound<float>::FlipInclusion(Segment.Range.GetUpperBound()), Segment.Range.GetUpperBoundValue() + SubSection->Parameters.PostrollTime);
+				if (!PostrollRange.IsEmpty())
+				{
+					// Iterate all ptrs that exist in the preroll range
+					TRange<int32> OverlappingSegments = Cache->Template->EvaluationField.OverlapRange(PostrollRange * InnerSequenceTransform);
+					for (int32 Index = OverlappingSegments.GetLowerBoundValue(); Index < OverlappingSegments.GetUpperBoundValue(); ++Index)
+					{
+						AllSegmentPtrs.Reset();
+
+						TRange<float> InnerSegmentRange = Cache->Template->EvaluationField.Ranges[Index] * InnerSequenceTransform.Inverse();
+						TRange<float> OverlappingRange = TRange<float>::Intersection(InnerSegmentRange, PostrollRange);
+						if (!OverlappingRange.IsEmpty())
+						{
+							RemapPostrollGroup(*Cache, *Cache->Template, Index, AllSegmentPtrs);
+							Args.Generator.AddExternalSegments(OverlappingRange, AllSegmentPtrs.Array());
+						}
+					}
+				}
+			}
+		}
+	}
+
+	bool RemapSubSection(int32 SectionIndex)
+	{
+		const UMovieSceneSubSection* SubSection = CastChecked<const UMovieSceneSubSection>(Sections[SectionIndex]);
+		UMovieSceneSequence* SubSequence = SubSection->GetSequence();
+		if (!SubSequence)
+		{
+			return false;
+		}
+
+		FMovieSceneEvaluationTemplate& Template = Args.SubSequenceStore.GetCompiledTemplate(*SubSequence);
+		
+		FMovieSceneSequenceTransform RootToSequenceTransform =
+			FMovieSceneSequenceTransform(SubSequence->GetMovieScene()->GetPlaybackRange().GetLowerBoundValue() + SubSection->Parameters.StartOffset) *		// Inner play offset
+			FMovieSceneSequenceTransform(0.f, SubSection->Parameters.TimeScale) *		// Inner play rate
+			FMovieSceneSequenceTransform(-SubSection->GetStartTime());					// Outer section start time
+
+#if WITH_EDITORONLY_DATA
+		TRange<float> InnerSectionRange(
+			SubSection->GetStartTime() * RootToSequenceTransform,
+			SubSection->GetEndTime() * RootToSequenceTransform
+			);
+		FMovieSceneSubSequenceData SubData(*SubSequence, SubSection->GetSequenceID(), *SubSection->GetPathNameInMovieScene(), InnerSectionRange);
+#else
+		FMovieSceneSubSequenceData SubData(*SubSequence, SubSection->GetSequenceID());
+#endif
+
+		// Construct the sub sequence data for this sub section
+		SubData.RootToSequenceTransform = RootToSequenceTransform;
+
+		// Remap it into the template generator
+		FMovieSceneSequenceID RemappedID = Args.Generator.GenerateSequenceID(SubData, MovieSceneSequenceID::Root);
+
+		FSectionData Cache(RemappedID, &Template);
+
+		// For editor previews, we remap all sub sequences regardless of whether they are actually used in the sub sequence or not
+		if (Args.Params.bForEditorPreview)
+		{
+			for (auto& Pair : Cache.Template->Hierarchy.AllSubSequenceData())
+			{
+				FMovieSceneSequenceID ChildID = Pair.Key;
+				RemapSubSequence(Cache, ChildID);
+			}
+		}
+
+		SectionDataCache.Add(SectionIndex, MoveTemp(Cache));
+		return true;
+	}
+
+private:
+
+	struct FSectionData
+	{
+		FSectionData(FMovieSceneSequenceIDRef InRemappedID, const FMovieSceneEvaluationTemplate* InTemplate) : RemappedID(InRemappedID), Template(InTemplate) {}
+
+#if PLATFORM_COMPILER_HAS_DEFAULTED_FUNCTIONS
+		FSectionData(FSectionData&&) = default;
+		FSectionData& operator=(FSectionData&&) = default;
+#else
+		FSectionData(FSectionData&& RHS)
+		{
+			*this = MoveTemp(RHS);
+		}
+		FSectionData& operator=(FSectionData&& RHS)
+		{
+			RemappedID = MoveTemp(RHS.RemappedID);
+			Template = MoveTemp(RHS.Template);
+			ChildrenRemappedIDs = MoveTemp(RHS.ChildrenRemappedIDs);
+			return *this;
+		}
+#endif
+		/** The ID of this section's sequence in the master sequecne */
+		FMovieSceneSequenceID RemappedID;
+		/** The compiled template (stored in the sequence itself) */
+		const FMovieSceneEvaluationTemplate* Template;
+		/** A map of remapped sequence IDs originating from this sequence's external tracks*/
+		TMap<FMovieSceneSequenceID, FMovieSceneSequenceID> ChildrenRemappedIDs;
+	};
+
+	void RemapPrerollGroup(FSectionData& Cache, const FMovieSceneEvaluationTemplate& Template, int32 EvaluationFieldIndex, TSet<FMovieSceneEvaluationFieldSegmentPtr>& AllSegmentPtrs)
+	{
+		const FMovieSceneEvaluationGroup& Group = Template.EvaluationField.Groups[EvaluationFieldIndex];
+
+		for (const FMovieSceneEvaluationGroupLUTIndex& LUTIndex : Group.LUTIndices)
+		{
+			const int32 First = LUTIndex.LUTOffset;
+			const int32 Last = First + LUTIndex.NumInitPtrs + LUTIndex.NumEvalPtrs;
+
+			for (int32 Index = First; Index < Last; ++Index)
+			{
+				FMovieSceneEvaluationFieldSegmentPtr SegmentPtr = Group.SegmentPtrLUT[Index];
+
+				const FMovieSceneEvaluationTrack* Track = FindTrack(SegmentPtr, Cache);
+				if (!ensure(Track) || !Track->ShouldEvaluateInPreroll())
+				{
+					continue;
+				}
+
+				SegmentPtr.SequenceID = RemapSubSequence(Cache, SegmentPtr.SequenceID);
+				AllSegmentPtrs.Add(SegmentPtr);
+			}
+		}
+	}
+
+	void RemapEvaluationGroup(FSectionData& Cache, const FMovieSceneEvaluationGroup& Group, TSet<FMovieSceneEvaluationFieldSegmentPtr>& AllSegmentPtrs)
+	{
+		for (const FMovieSceneEvaluationGroupLUTIndex& LUTIndex : Group.LUTIndices)
+		{
+			const int32 First = LUTIndex.LUTOffset;
+			const int32 Last = First + LUTIndex.NumInitPtrs + LUTIndex.NumEvalPtrs;
+
+			for (int32 Index = First; Index < Last; ++Index)
+			{
+				FMovieSceneEvaluationFieldSegmentPtr SegmentPtr = Group.SegmentPtrLUT[Index];
+				SegmentPtr.SequenceID = RemapSubSequence(Cache, SegmentPtr.SequenceID);
+				AllSegmentPtrs.Add(SegmentPtr);
+			}
+		}
+	}
+
+	void RemapPostrollGroup(FSectionData& Cache, const FMovieSceneEvaluationTemplate& Template, int32 EvaluationFieldIndex, TSet<FMovieSceneEvaluationFieldSegmentPtr>& AllSegmentPtrs)
+	{
+		const FMovieSceneEvaluationGroup& Group = Template.EvaluationField.Groups[EvaluationFieldIndex];
+
+		for (const FMovieSceneEvaluationGroupLUTIndex& LUTIndex : Group.LUTIndices)
+		{
+			const int32 First = LUTIndex.LUTOffset;
+			const int32 Last = First + LUTIndex.NumInitPtrs + LUTIndex.NumEvalPtrs;
+
+			for (int32 Index = First; Index < Last; ++Index)
+			{
+				FMovieSceneEvaluationFieldSegmentPtr SegmentPtr = Group.SegmentPtrLUT[Index];
+
+				const FMovieSceneEvaluationTrack* Track = FindTrack(SegmentPtr, Cache);
+				if (!ensure(Track) || !Track->ShouldEvaluateInPostroll())
+				{
+					continue;
+				}
+
+				SegmentPtr.SequenceID = RemapSubSequence(Cache, SegmentPtr.SequenceID);
+				AllSegmentPtrs.Add(SegmentPtr);
+			}
+		}
+	}
+
+	FMovieSceneSequenceID RemapSubSequence(FSectionData& Cache, FMovieSceneSequenceID SourceID)
+	{
+		// Root tracks use the cache's remapped ID
+		if (SourceID == MovieSceneSequenceID::Root)
+		{
+			return Cache.RemappedID;
+		}
+
+		// Nested tracks tracks - may need to remap the outer sequence
+		if (Cache.ChildrenRemappedIDs.Contains(SourceID))
+		{
+			return Cache.ChildrenRemappedIDs.FindRef(SourceID);
+		}
+
+		// Remap this into the root
+		const FMovieSceneSubSequenceData* SubData = Cache.Template->Hierarchy.FindSubData(SourceID);
+		const FMovieSceneSequenceHierarchyNode* HierarchyNode = Cache.Template->Hierarchy.FindNode(SourceID);
+		check(SubData && HierarchyNode);
+
+		FMovieSceneSubSequenceData SubDataCopy = *SubData;
+
+		// Accumulate the parent's sequence transform
+		SubDataCopy.RootToSequenceTransform = SubDataCopy.RootToSequenceTransform * Args.Generator.GetSequenceTransform(Cache.RemappedID);
+
+		FMovieSceneSequenceID RemappedParentID = HierarchyNode->ParentID == MovieSceneSequenceID::Root ? Cache.RemappedID : RemapSubSequence(Cache, HierarchyNode->ParentID);
+
+		// Remap this sequence within the main generator, under this section as a parent
+		FMovieSceneSequenceID RemappedID = Args.Generator.GenerateSequenceID(SubDataCopy, RemappedParentID);
+		Cache.ChildrenRemappedIDs.Add(SourceID, RemappedID);
+
+		return RemappedID;
+	}
+
+	const FMovieSceneEvaluationTrack* FindTrack(FMovieSceneEvaluationFieldTrackPtr TrackPtr, const FSectionData& SectionData) const
+	{
+		if (TrackPtr.SequenceID == MovieSceneSequenceID::Root)
+		{
+			return SectionData.Template->FindTrack(TrackPtr.TrackIdentifier);
+		}
+		return nullptr;
+	}
+
+private:
+
+	/** Compilation parameters */
+	const FMovieSceneTrackCompilerArgs Args;
+
+	/** Array of sections to compile */
+	const TArray<UMovieSceneSection*>& Sections;
+
+	/** Cached information pertaining to each section's sequence */
+	TMap<int32, FSectionData> SectionDataCache;
+};
+
+void UMovieSceneSubTrack::GenerateTemplate(const FMovieSceneTrackCompilerArgs& Args) const
+{
+	// We compile sub tracks by generating a track evaluation field, and remapping child tracks from the inner sequence to the outer's timeline
+	TInlineValue<FMovieSceneSegmentCompilerRules> TrackRules = GetTrackCompilerRules();
+	TInlineValue<FMovieSceneSegmentCompilerRules> RowRules = GetRowCompilerRules();
+
+	// Generate track rows
+	FMovieSceneTrackCompiler::FRows TrackRows(Sections, RowRules.GetPtr());
+
+	// Compile track rows into an evaluation field based on the supplied compiler rules
+	FMovieSceneTrackCompiler Compiler;
+	FMovieSceneTrackEvaluationField EvaluationField = Compiler.Compile(TrackRows.Rows, TrackRules.GetPtr());
+
+	// Remap each segment in the final eval field into the outer template generator.
+	// This ensures we only reference segments that are actually going to be evaluated in the sub section
+	FSubTrackRemapper Remapper(Args, Sections);
+
+	// This is only required in editor, where we need to navigate sections that aren't necessarily evaluated in the master sequence (outside of an inner playback range, or don't contain any tracks)
+	const bool bAggressivelyRemapSections = Args.Params.bForEditorPreview;
+	if (bAggressivelyRemapSections)
+	{
+		// Ensure all sub data is remapped to the master
+		for (int32 SectionIndex = 0; SectionIndex < Sections.Num(); ++SectionIndex)
+		{
+			Remapper.RemapSubSection(SectionIndex);
+		}
+	}
+
+	for (const FMovieSceneSegment& Segment : EvaluationField.Segments)
+	{
+		Remapper.RemapSegment(Segment);
+	}
+}
+
+
 #undef LOCTEXT_NAMESPACE
diff --git a/Engine/Source/Runtime/MovieSceneTracks/Private/Tracks/MovieSceneVectorTrack.cpp b/Engine/Source/Runtime/MovieSceneTracks/Private/Tracks/MovieSceneVectorTrack.cpp
--- a/Engine/Source/Runtime/MovieSceneTracks/Private/Tracks/MovieSceneVectorTrack.cpp
+++ b/Engine/Source/Runtime/MovieSceneTracks/Private/Tracks/MovieSceneVectorTrack.cpp
@@ -4,7 +4,7 @@
 #include "MovieSceneVectorSection.h"
 #include "MovieSceneVectorTrack.h"
 #include "IMovieScenePlayer.h"
-#include "MovieSceneVectorTrackInstance.h"
+#include "Evaluation/MovieScenePropertyTemplates.h"
 
 
 UMovieSceneVectorTrack::UMovieSceneVectorTrack( const FObjectInitializer& ObjectInitializer )
@@ -22,26 +22,7 @@ UMovieSceneSection* UMovieSceneVectorTrack::CreateNewSection()
 }
 
 
-TSharedPtr<IMovieSceneTrackInstance> UMovieSceneVectorTrack::CreateInstance()
+FMovieSceneEvalTemplatePtr UMovieSceneVectorTrack::CreateTemplateForSection(const UMovieSceneSection& InSection) const
 {
-	return MakeShareable( new FMovieSceneVectorTrackInstance( *this ) );
-}
-
-
-bool UMovieSceneVectorTrack::Eval( float Position, float LastPosition, FVector4& InOutVector ) const
-{
-	const UMovieSceneSection* Section = MovieSceneHelpers::FindNearestSectionAtTime( Sections, Position );
-
-	if( Section )
-	{
-		if (!Section->IsInfinite())
-		{
-			Position = FMath::Clamp(Position, Section->GetStartTime(), Section->GetEndTime());
-		}
-
-		InOutVector = CastChecked<UMovieSceneVectorSection>( Section )->Eval( Position, InOutVector );
-	}
-
-	return Section != nullptr;
-}
-
+	return FMovieSceneVectorPropertySectionTemplate(*CastChecked<UMovieSceneVectorSection>(&InSection), *this);
+}
\ No newline at end of file
diff --git a/Engine/Source/Runtime/MovieSceneTracks/Private/Tracks/MovieSceneVisibilityTrack.cpp b/Engine/Source/Runtime/MovieSceneTracks/Private/Tracks/MovieSceneVisibilityTrack.cpp
--- a/Engine/Source/Runtime/MovieSceneTracks/Private/Tracks/MovieSceneVisibilityTrack.cpp
+++ b/Engine/Source/Runtime/MovieSceneTracks/Private/Tracks/MovieSceneVisibilityTrack.cpp
@@ -2,9 +2,9 @@
 
 #include "MovieSceneTracksPrivatePCH.h"
 #include "MovieSceneVisibilityTrack.h"
+#include "MovieSceneBoolSection.h"
 #include "IMovieScenePlayer.h"
-#include "MovieSceneVisibilitySection.h"
-#include "MovieSceneVisibilityTrackInstance.h"
+#include "Evaluation/MovieSceneVisibilityTemplate.h"
 
 
 #define LOCTEXT_NAMESPACE "MovieSceneVisibilityTrack"
@@ -15,15 +15,9 @@ UMovieSceneVisibilityTrack::UMovieSceneVisibilityTrack(const FObjectInitializer&
 {}
 
 
-UMovieSceneSection* UMovieSceneVisibilityTrack::CreateNewSection()
+FMovieSceneEvalTemplatePtr UMovieSceneVisibilityTrack::CreateTemplateForSection(const UMovieSceneSection& InSection) const
 {
-	return NewObject<UMovieSceneSection>(this, UMovieSceneVisibilitySection::StaticClass(), NAME_None, RF_Transactional);
-}
-
-
-TSharedPtr<IMovieSceneTrackInstance> UMovieSceneVisibilityTrack::CreateInstance()
-{
-	return MakeShareable(new FMovieSceneVisibilityTrackInstance(*this));
+	return FMovieSceneVisibilitySectionTemplate(*CastChecked<const UMovieSceneBoolSection>(&InSection), *this);
 }
 
 
diff --git a/Engine/Source/Runtime/MovieSceneTracks/Public/Evaluation/MovieSceneParameterTemplate.h b/Engine/Source/Runtime/MovieSceneTracks/Public/Evaluation/MovieSceneParameterTemplate.h
new file mode 100644
--- /dev/null
+++ b/Engine/Source/Runtime/MovieSceneTracks/Public/Evaluation/MovieSceneParameterTemplate.h
@@ -0,0 +1,207 @@
+// Copyright 1998-2016 Epic Games, Inc. All Rights Reserved.
+
+#pragma once
+
+#include "MovieSceneFwd.h"
+
+#include "MovieSceneEvalTemplate.h"
+#include "MovieSceneParameterSection.h"
+#include "MovieSceneParameterTemplate.generated.h"
+
+DECLARE_CYCLE_STAT(TEXT("Parameter Track Token Execute"), MovieSceneEval_ParameterTrack_TokenExecute, STATGROUP_MovieSceneEval);
+
+class UMovieSceneComponentMaterialTrack;
+
+/** Evaluation structure that holds evaluated values */
+struct FEvaluatedParameterSectionValues
+{
+	FEvaluatedParameterSectionValues() = default;
+
+#if PLATFORM_COMPILER_HAS_DEFAULTED_FUNCTIONS
+	FEvaluatedParameterSectionValues(FEvaluatedParameterSectionValues&&) = default;
+	FEvaluatedParameterSectionValues& operator=(FEvaluatedParameterSectionValues&&) = default;
+#else
+	FEvaluatedParameterSectionValues(FEvaluatedParameterSectionValues&& RHS)
+	{
+		*this = MoveTemp(RHS);
+	}
+	FEvaluatedParameterSectionValues& operator=(FEvaluatedParameterSectionValues&& RHS)
+	{
+		ScalarValues = MoveTemp(RHS.ScalarValues);
+		VectorValues = MoveTemp(RHS.VectorValues);
+		ColorValues = MoveTemp(RHS.ColorValues);
+		return *this;
+	}
+#endif
+
+	/** Array of evaluated scalar values */
+	TArray<FScalarParameterNameAndValue, TInlineAllocator<2>> ScalarValues;
+	/** Array of evaluated vector values */
+	TArray<FVectorParameterNameAndValue, TInlineAllocator<2>> VectorValues;
+	/** Array of evaluated color values */
+	TArray<FColorParameterNameAndValue, TInlineAllocator<2>> ColorValues;
+};
+
+/** Template that performs evaluation of parameter sections */
+USTRUCT()
+struct FMovieSceneParameterSectionTemplate : public FMovieSceneEvalTemplate
+{
+	GENERATED_BODY()
+
+	FMovieSceneParameterSectionTemplate() {}
+
+protected:
+
+	/** Protected constructor to initialize from a parameter section */
+	MOVIESCENETRACKS_API FMovieSceneParameterSectionTemplate(const UMovieSceneParameterSection& Section);
+
+	/** Evaluate our curves, outputting evaluated values into the specified container */
+	MOVIESCENETRACKS_API void EvaluateCurves(const FMovieSceneContext& Context, FEvaluatedParameterSectionValues& OutValues) const;
+
+private:
+
+	/** The scalar parameter names and their associated curves. */
+	UPROPERTY()
+	TArray<FScalarParameterNameAndCurve> Scalars;
+
+	/** The vector parameter names and their associated curves. */
+	UPROPERTY()
+	TArray<FVectorParameterNameAndCurves> Vectors;
+
+	/** The color parameter names and their associated curves. */
+	UPROPERTY()
+	TArray<FColorParameterNameAndCurves> Colors;
+};
+
+/** Default accessor type for use with TMaterialTrackExecutionToken */
+struct FDefaultMaterialAccessor
+{
+	// Implement in derived classes:
+
+	/** Get the anim type ID for the evaluation token */
+	// FMovieSceneAnimTypeID 	GetAnimTypeID();
+
+	/** Get the material from the specified object */
+	// UMaterialInterface* 		GetMaterialForObject(UObject& Object);
+
+	/** Set the material for the specified object */
+	// void 					SetMaterialForObject(UObject& Object, UMaterialInterface& Material);
+
+	/** Apply the specified values onto the specified material */
+	MOVIESCENETRACKS_API void Apply(UMaterialInstanceDynamic& Material, const FEvaluatedParameterSectionValues& Values);
+};
+
+/**
+ * Material track execution token
+ * Templated on accessor type to allow for copyable accessors into pre animated state
+ */
+template<typename AccessorType>
+struct TMaterialTrackExecutionToken : IMovieSceneExecutionToken
+{
+	TMaterialTrackExecutionToken(AccessorType InAccessor)
+		: Accessor(MoveTemp(InAccessor))
+	{
+	}
+
+#if PLATFORM_COMPILER_HAS_DEFAULTED_FUNCTIONS
+	TMaterialTrackExecutionToken(TMaterialTrackExecutionToken&&) = default;
+	TMaterialTrackExecutionToken& operator=(TMaterialTrackExecutionToken&&) = default;
+#else
+	TMaterialTrackExecutionToken(TMaterialTrackExecutionToken&& RHS)
+		: Accessor(MoveTemp(RHS.Accessor))
+		, Values(MoveTemp(RHS.Values))
+	{
+	}
+	TMaterialTrackExecutionToken& operator=(TMaterialTrackExecutionToken&& RHS)
+	{
+		Accessor = MoveTemp(RHS.Accessor);
+		Values = MoveTemp(RHS.Values);
+		return *this;
+	}
+#endif
+
+	virtual void Execute(const FMovieSceneContext& Context, const FMovieSceneEvaluationOperand& Operand, FPersistentEvaluationData& PersistentData, IMovieScenePlayer& Player)
+	{
+		MOVIESCENE_DETAILED_SCOPE_CYCLE_COUNTER(MovieSceneEval_ParameterTrack_TokenExecute)
+		
+		for (TWeakObjectPtr<>& WeakObject : Player.FindBoundObjects(Operand))
+		{
+			UObject* Object = WeakObject.Get();
+			UMaterialInterface* Material = Object ? Accessor.GetMaterialForObject(*Object) : nullptr;
+
+			if (!Material)
+			{
+				continue;
+			}
+
+			UMaterialInstanceDynamic* DynamicMaterialInstance = Cast<UMaterialInstanceDynamic>(Material);
+			if (!DynamicMaterialInstance)
+			{
+				// Save the old instance
+				Player.SavePreAnimatedState(*Object, Accessor.GetAnimTypeID(), FPreAnimatedTokenProducer(Accessor));
+
+				FString DynamicName = Material->GetName() + "_Animated";
+				FName UniqueDynamicName = MakeUniqueObjectName( Object, UMaterialInstanceDynamic::StaticClass() , *DynamicName );
+				DynamicMaterialInstance = UMaterialInstanceDynamic::Create( Material, Object, UniqueDynamicName );
+
+				Accessor.SetMaterialForObject(*Object, *DynamicMaterialInstance);
+			}
+
+			Accessor.Apply(*DynamicMaterialInstance, Values);
+		}
+	}
+
+	AccessorType Accessor;
+	FEvaluatedParameterSectionValues Values;
+
+private:
+
+	struct FPreAnimatedTokenProducer : IMovieScenePreAnimatedTokenProducer
+	{
+		struct FPreAnimatedToken : IMovieScenePreAnimatedToken
+		{
+			FPreAnimatedToken(UObject& Object, const AccessorType& InAccessor)
+				: Accessor(InAccessor)
+				, Material(Accessor.GetMaterialForObject(Object))
+			{}
+
+			virtual void RestoreState(UObject& Object, IMovieScenePlayer& Player) override
+			{
+				if (UMaterialInterface* PinnedMaterial = Material.Get())
+				{
+					Accessor.SetMaterialForObject(Object, *PinnedMaterial);
+				}
+			}
+
+			AccessorType Accessor;
+			TWeakObjectPtr<UMaterialInterface> Material;
+		};
+
+		FPreAnimatedTokenProducer(const AccessorType& InAccessor) : Accessor(InAccessor) {}
+
+		virtual IMovieScenePreAnimatedTokenPtr CacheExistingState(UObject& Object) const override
+		{
+			return FPreAnimatedToken(Object, Accessor);
+		}
+
+		const AccessorType& Accessor;
+	};
+};
+
+/** Evaluation template for primitive component materials */
+USTRUCT()
+struct FMovieSceneComponentMaterialSectionTemplate : public FMovieSceneParameterSectionTemplate
+{
+	GENERATED_BODY()
+
+	FMovieSceneComponentMaterialSectionTemplate() {}
+	FMovieSceneComponentMaterialSectionTemplate(const UMovieSceneParameterSection& Section, const UMovieSceneComponentMaterialTrack& Track);
+
+private:
+
+	virtual UScriptStruct& GetScriptStructImpl() const override { return *StaticStruct(); }
+	virtual void Evaluate(const FMovieSceneEvaluationOperand& Operand, const FMovieSceneContext& Context, const FPersistentEvaluationData& PersistentData, FMovieSceneExecutionTokens& ExecutionTokens) const override;
+
+	UPROPERTY()
+	int32 MaterialIndex;
+};
\ No newline at end of file
diff --git a/Engine/Source/Runtime/MovieSceneTracks/Public/Evaluation/MovieSceneSpawnTemplate.h b/Engine/Source/Runtime/MovieSceneTracks/Public/Evaluation/MovieSceneSpawnTemplate.h
new file mode 100644
--- /dev/null
+++ b/Engine/Source/Runtime/MovieSceneTracks/Public/Evaluation/MovieSceneSpawnTemplate.h
@@ -0,0 +1,30 @@
+// Copyright 1998-2016 Epic Games, Inc. All Rights Reserved.
+
+#pragma once
+
+#include "MovieSceneFwd.h"
+#include "Curves/IntegralCurve.h"
+#include "MovieSceneEvalTemplate.h"
+#include "MovieSceneSpawnTemplate.generated.h"
+
+class UMovieSceneSpawnSection;
+
+/** Spawn track eval template that evaluates a curve */
+USTRUCT()
+struct FMovieSceneSpawnSectionTemplate : public FMovieSceneEvalTemplate
+{
+	GENERATED_BODY()
+	
+	FMovieSceneSpawnSectionTemplate() {}
+	FMovieSceneSpawnSectionTemplate(const UMovieSceneSpawnSection& SpawnSection);
+
+	MOVIESCENETRACKS_API static FMovieSceneAnimTypeID GetAnimTypeID();
+	
+private:
+
+	virtual UScriptStruct& GetScriptStructImpl() const override { return *StaticStruct(); }
+	virtual void Evaluate(const FMovieSceneEvaluationOperand& Operand, const FMovieSceneContext& Context, const FPersistentEvaluationData& PersistentData, FMovieSceneExecutionTokens& ExecutionTokens) const override;
+
+	UPROPERTY()
+	FIntegralCurve Curve;
+};
\ No newline at end of file
diff --git a/Engine/Source/Runtime/MovieSceneTracks/Public/Sections/MovieScene3DPathSection.h b/Engine/Source/Runtime/MovieSceneTracks/Public/Sections/MovieScene3DPathSection.h
--- a/Engine/Source/Runtime/MovieSceneTracks/Public/Sections/MovieScene3DPathSection.h
+++ b/Engine/Source/Runtime/MovieSceneTracks/Public/Sections/MovieScene3DPathSection.h
@@ -54,11 +54,15 @@ public:
 	 *
 	 * @return The timing curve.
 	 */
-	MOVIESCENETRACKS_API FRichCurve& GetTimingCurve()
-	{
-		return TimingCurve;
-	}
+	MOVIESCENETRACKS_API FRichCurve& GetTimingCurve() { return TimingCurve; }
+	MOVIESCENETRACKS_API const FRichCurve& GetTimingCurve() const { return TimingCurve; }
 	
+	MOVIESCENETRACKS_API MovieScene3DPathSection_Axis GetFrontAxisEnum() const { return FrontAxisEnum; }
+	MOVIESCENETRACKS_API MovieScene3DPathSection_Axis GetUpAxisEnum() const { return UpAxisEnum; }
+	MOVIESCENETRACKS_API bool GetFollow() const { return bFollow; }
+	MOVIESCENETRACKS_API bool GetReverse() const { return bReverse; }
+	MOVIESCENETRACKS_API bool GetForceUpright() const { return bForceUpright; }
+
 public:
 
 	//~ MovieSceneSection interface
diff --git a/Engine/Source/Runtime/MovieSceneTracks/Public/Sections/MovieScene3DTransformSection.h b/Engine/Source/Runtime/MovieSceneTracks/Public/Sections/MovieScene3DTransformSection.h
--- a/Engine/Source/Runtime/MovieSceneTracks/Public/Sections/MovieScene3DTransformSection.h
+++ b/Engine/Source/Runtime/MovieSceneTracks/Public/Sections/MovieScene3DTransformSection.h
@@ -210,6 +210,7 @@ public:
 	 * @return The curve on the axis
 	 */
 	MOVIESCENETRACKS_API FRichCurve& GetTranslationCurve( EAxis::Type Axis );
+	MOVIESCENETRACKS_API const FRichCurve& GetTranslationCurve( EAxis::Type Axis ) const;
 
 	/** 
 	 * Returns the rotation curve for a specific axis
@@ -218,6 +219,7 @@ public:
 	 * @return The curve on the axis
 	 */
 	MOVIESCENETRACKS_API FRichCurve& GetRotationCurve( EAxis::Type Axis );
+	MOVIESCENETRACKS_API const FRichCurve& GetRotationCurve( EAxis::Type Axis ) const;
 
 	/** 
 	 * Returns the scale curve for a specific axis
@@ -226,6 +228,7 @@ public:
 	 * @return The curve on the axis
 	 */
 	MOVIESCENETRACKS_API FRichCurve& GetScaleCurve( EAxis::Type Axis );
+	MOVIESCENETRACKS_API const FRichCurve& GetScaleCurve( EAxis::Type Axis ) const;
 
 	/**
 	 * Return the trajectory visibility
@@ -254,7 +257,8 @@ public:
 	virtual void AddKey( float Time, const FTransformKey& KeyData, EMovieSceneKeyInterpolation KeyInterpolation ) override;
 	virtual void SetDefault( const FTransformKey& KeyData ) override;
 	virtual void ClearDefaults() override;
-
+	virtual FMovieSceneEvalTemplatePtr GenerateTemplate() const override;
+	
 private:
 
 	/** Translation curves */
diff --git a/Engine/Source/Runtime/MovieSceneTracks/Public/Sections/MovieSceneActorReferenceSection.h b/Engine/Source/Runtime/MovieSceneTracks/Public/Sections/MovieSceneActorReferenceSection.h
--- a/Engine/Source/Runtime/MovieSceneTracks/Public/Sections/MovieSceneActorReferenceSection.h
+++ b/Engine/Source/Runtime/MovieSceneTracks/Public/Sections/MovieSceneActorReferenceSection.h
@@ -43,6 +43,10 @@ public:
 		return ActorGuidIndexCurve;
 	}
 
+	const FIntegralCurve& GetActorReferenceCurve() const { return ActorGuidIndexCurve; }
+
+	const TArray<FGuid>& GetActorGuids() const { return ActorGuids; }
+
 public:
 
 	MOVIESCENETRACKS_API FKeyHandle AddKey(float Time, const FGuid& Value);
diff --git a/Engine/Source/Runtime/MovieSceneTracks/Public/Sections/MovieSceneAudioSection.h b/Engine/Source/Runtime/MovieSceneTracks/Public/Sections/MovieSceneAudioSection.h
--- a/Engine/Source/Runtime/MovieSceneTracks/Public/Sections/MovieSceneAudioSection.h
+++ b/Engine/Source/Runtime/MovieSceneTracks/Public/Sections/MovieSceneAudioSection.h
@@ -21,13 +21,13 @@ public:
 	void SetSound(class USoundBase* InSound) {Sound = InSound;}
 
 	/** Gets the sound for this section */
-	class USoundBase* GetSound() {return Sound;}
+	class USoundBase* GetSound() const {return Sound;}
 	
 	/** Sets the time that the sound is supposed to be played at */
 	void SetAudioStartTime(float InAudioStartTime) {AudioStartTime = InAudioStartTime;}
 	
 	/** Gets the (absolute) time that the sound is supposed to be played at */
-	float GetAudioStartTime() {return AudioStartTime;}
+	float GetAudioStartTime() const {return AudioStartTime;}
 	
 	/**
 	 * @return The range of times that the sound plays, truncated by the section limits
@@ -99,6 +99,7 @@ public:
 	virtual void GetSnapTimes(TArray<float>& OutSnapTimes, bool bGetSectionBorders) const override;
 	virtual TOptional<float> GetKeyTime( FKeyHandle KeyHandle ) const override { return TOptional<float>(); }
 	virtual void SetKeyTime( FKeyHandle KeyHandle, float Time ) override { }
+	virtual FMovieSceneEvalTemplatePtr GenerateTemplate() const override;
 
 private:
 
diff --git a/Engine/Source/Runtime/MovieSceneTracks/Public/Sections/MovieSceneBoolSection.h b/Engine/Source/Runtime/MovieSceneTracks/Public/Sections/MovieSceneBoolSection.h
--- a/Engine/Source/Runtime/MovieSceneTracks/Public/Sections/MovieSceneBoolSection.h
+++ b/Engine/Source/Runtime/MovieSceneTracks/Public/Sections/MovieSceneBoolSection.h
@@ -34,11 +34,9 @@ public:
 	virtual bool Eval(float Position, bool DefaultValue) const;
 
 	/** Gets all the keys of this boolean section. */
-	FIntegralCurve& GetCurve()
-	{
-		return BoolCurve;
-	}
-
+	FIntegralCurve& GetCurve() { return BoolCurve; }
+	const FIntegralCurve& GetCurve() const { return BoolCurve; }
+	
 public:
 
 	//~ IKeyframeSection interface
diff --git a/Engine/Source/Runtime/MovieSceneTracks/Public/Sections/MovieSceneByteSection.h b/Engine/Source/Runtime/MovieSceneTracks/Public/Sections/MovieSceneByteSection.h
--- a/Engine/Source/Runtime/MovieSceneTracks/Public/Sections/MovieSceneByteSection.h
+++ b/Engine/Source/Runtime/MovieSceneTracks/Public/Sections/MovieSceneByteSection.h
@@ -18,18 +18,10 @@ class UMovieSceneByteSection
 {
 	GENERATED_UCLASS_BODY()
 public:
-	/**
-	 * Update this section.
-	 *
-	 * @param Position The position in time within the movie scene
-	 */
-	virtual uint8 Eval(float Position, uint8 DefaultValue) const;
 
 	/** Gets all the keys of this byte section */
-	FIntegralCurve& GetCurve()
-	{
-		return ByteCurve;
-	}
+	FIntegralCurve& GetCurve() { return ByteCurve; }
+	const FIntegralCurve& GetCurve() const { return ByteCurve; }
 
 public:
 
diff --git a/Engine/Source/Runtime/MovieSceneTracks/Public/Sections/MovieSceneCameraAnimSection.h b/Engine/Source/Runtime/MovieSceneTracks/Public/Sections/MovieSceneCameraAnimSection.h
--- a/Engine/Source/Runtime/MovieSceneTracks/Public/Sections/MovieSceneCameraAnimSection.h
+++ b/Engine/Source/Runtime/MovieSceneTracks/Public/Sections/MovieSceneCameraAnimSection.h
@@ -5,42 +5,26 @@
 #include "MovieSceneSection.h"
 #include "MovieSceneCameraAnimSection.generated.h"
 
-/**
- *
- */
-UCLASS(MinimalAPI)
-class UMovieSceneCameraAnimSection : public UMovieSceneSection
+USTRUCT()
+struct FMovieSceneCameraAnimSectionData
 {
 	GENERATED_BODY()
 
-public:
-	// ctor
-	UMovieSceneCameraAnimSection(const FObjectInitializer& ObjectInitializer);
-
-	/** Sets the CameraAnim for this section */
-	void SetCameraAnim(UCameraAnim* InCameraAnim) { CameraAnim = InCameraAnim; }
-	
-	/** Gets the CameraAnim for this section */
-	class UCameraAnim* GetCameraAnim() const { return CameraAnim; }
-	
-	/** MovieSceneSection interface */
-	virtual void MoveSection( float DeltaPosition, TSet<FKeyHandle>& KeyHandles ) override;
-	virtual void DilateSection( float DilationFactor, float Origin, TSet<FKeyHandle>& KeyHandles ) override;
-	virtual void GetKeyHandles(TSet<FKeyHandle>& OutKeyHandles, TRange<float> TimeRange) const override;
-	virtual TOptional<float> GetKeyTime( FKeyHandle KeyHandle ) const override;
-	virtual void SetKeyTime( FKeyHandle KeyHandle, float Time ) override;
-
-	/** #todo */
-//	MOVIESCENETRACKS_API FRichCurve& GetAnimWeightCurve() { return AnimWeightCurve; }
-
-private:
-// 	UPROPERTY(EditAnywhere, Category="Camera Anim")
-// 	FRichCurve AnimWeightCurve;
+	FMovieSceneCameraAnimSectionData()
+		: CameraAnim(nullptr)
+		, PlayRate(1.f)
+		, PlayScale(1.f)
+		, BlendInTime(0.f)
+		, BlendOutTime(0.f)
+		, bLooping(false)
+		, bRandomStartTime(false)
+	{
+	}
 
+	/** The camera anim to play */
 	UPROPERTY(EditAnywhere, Category = "Camera Anim")
 	UCameraAnim* CameraAnim;
 
-public:
 	/** How fast to play back the animation. */
 	UPROPERTY(EditAnywhere, Category = "Camera Anim")
 	float PlayRate;
@@ -58,15 +42,47 @@ public:
 	UPROPERTY(EditAnywhere, Category = "Camera Anim")
 	bool bLooping;
 	
+// @todo: Why is this not EditAnywhere?
 // 	UPROPERTY(EditAnywhere, Category = "Camera Anim")
- 	bool bRandomStartTime;
+	bool bRandomStartTime;
+};
+
+/**
+ *
+ */
+UCLASS(MinimalAPI)
+class UMovieSceneCameraAnimSection : public UMovieSceneSection
+{
+	GENERATED_BODY()
+
+public:
+
+	UMovieSceneCameraAnimSection(const FObjectInitializer& ObjectInitializer);
+
+	virtual FMovieSceneEvalTemplatePtr GenerateTemplate() const override;
+	virtual void PostLoad() override;
+
+	UPROPERTY(EditAnywhere, Category="Camera Anim", meta=(ShowOnlyInnerProperties))
+	FMovieSceneCameraAnimSectionData AnimData;
+
+private:
+
+	/** Deprecated members */
+	UPROPERTY()
+	UCameraAnim* CameraAnim_DEPRECATED;
+
+	UPROPERTY()
+	float PlayRate_DEPRECATED;
 	
-// 	UPROPERTY(EditAnywhere, Category = "Camera Anim")
-// 	float Duration;
+	UPROPERTY()
+	float PlayScale_DEPRECATED;
 	
-	UPROPERTY(EditAnywhere, Category = "Camera Anim")
-	TEnumAsByte<ECameraAnimPlaySpace::Type> PlaySpace;
+	UPROPERTY()
+	float BlendInTime_DEPRECATED;
 	
-	UPROPERTY(EditAnywhere, Category = "Camera Anim")
-	FRotator UserDefinedPlaySpace;
+	UPROPERTY()
+	float BlendOutTime_DEPRECATED;
+	
+	UPROPERTY()
+	bool bLooping_DEPRECATED;
 };
\ No newline at end of file
diff --git a/Engine/Source/Runtime/MovieSceneTracks/Public/Sections/MovieSceneCameraCutSection.h b/Engine/Source/Runtime/MovieSceneTracks/Public/Sections/MovieSceneCameraCutSection.h
--- a/Engine/Source/Runtime/MovieSceneTracks/Public/Sections/MovieSceneCameraCutSection.h
+++ b/Engine/Source/Runtime/MovieSceneTracks/Public/Sections/MovieSceneCameraCutSection.h
@@ -16,6 +16,11 @@ class UMovieSceneCameraCutSection
 	GENERATED_BODY()
 
 public:
+	UMovieSceneCameraCutSection(const FObjectInitializer& Init)
+		: Super(Init)
+	{
+		EvalOptions.EnableAndSetCompletionMode(EMovieSceneCompletionMode::RestoreState);
+	}
 
 	/** Gets the camera guid for this CameraCut section */
 	FGuid GetCameraGuid() const
@@ -29,6 +34,8 @@ public:
 		CameraGuid = InGuid;
 	}
 
+	virtual FMovieSceneEvalTemplatePtr GenerateTemplate() const override;
+
 private:
 
 	/** The camera possessable or spawnable that this movie CameraCut uses */
@@ -53,7 +60,7 @@ public:
 	// UMovieSceneSection interface
 	virtual TOptional<float> GetKeyTime( FKeyHandle KeyHandle ) const override { return TOptional<float>(); }
 	virtual void SetKeyTime( FKeyHandle KeyHandle, float Time ) override { }
-
+	
 private:
 
 	/** The reference frame offset for single thumbnail rendering */
