Copying //UE4/Dev-Sequencer to //UE4/Dev-Main (Source: //UE4/Dev-Sequencer @ 3178529)

diff --git a/Engine/Source/Runtime/MovieSceneTracks/Public/Sections/MovieSceneCameraShakeSection.h b/Engine/Source/Runtime/MovieSceneTracks/Public/Sections/MovieSceneCameraShakeSection.h
--- a/Engine/Source/Runtime/MovieSceneTracks/Public/Sections/MovieSceneCameraShakeSection.h
+++ b/Engine/Source/Runtime/MovieSceneTracks/Public/Sections/MovieSceneCameraShakeSection.h
@@ -7,6 +7,32 @@
 
 class UCameraShake;
 
+USTRUCT()
+struct FMovieSceneCameraShakeSectionData
+{
+	GENERATED_BODY()
+
+	FMovieSceneCameraShakeSectionData()
+		: ShakeClass(nullptr)
+		, PlayScale(1.f)
+	{
+	}
+
+	/** Class of the camera shake to play */
+	UPROPERTY(EditAnywhere, Category = "Camera Shake")
+	TSubclassOf<UCameraShake> ShakeClass;
+
+	/** Scalar that affects shake intensity */
+	UPROPERTY(EditAnywhere, Category = "Camera Shake")
+	float PlayScale;
+
+	UPROPERTY(EditAnywhere, Category = "Camera Shake")
+	TEnumAsByte<ECameraAnimPlaySpace::Type> PlaySpace;
+
+	UPROPERTY(EditAnywhere, Category = "Camera Shake")
+	FRotator UserDefinedPlaySpace;
+};
+
 /**
  *
  */
@@ -17,45 +43,24 @@ class UMovieSceneCameraShakeSection : public UMovieSceneSection
 
 public:
 
-	// ctor
 	UMovieSceneCameraShakeSection(const FObjectInitializer& ObjectInitializer);
 
-	/** Sets the CameraShake for this section */
-	void SetCameraShakeClass(TSubclassOf<UCameraShake> InShakeClass) { ShakeClass = InShakeClass; }
-
-	/** Gets the CameraAnim for this section */
-	UClass* GetCameraShakeClass() const { return ShakeClass; }
-
-	/** MovieSceneSection interface */
-	virtual void MoveSection( float DeltaPosition, TSet<FKeyHandle>& KeyHandles ) override;
-	virtual void DilateSection( float DilationFactor, float Origin, TSet<FKeyHandle>& KeyHandles ) override;
-	virtual void GetKeyHandles(TSet<FKeyHandle>& OutKeyHandles, TRange<float> TimeRange) const override;
-	virtual TOptional<float> GetKeyTime( FKeyHandle KeyHandle ) const override;
-	virtual void SetKeyTime( FKeyHandle KeyHandle, float Time ) override;
-
-	/** 
-	 * @return Returns the weighting curve
-	 */
-//	MOVIESCENETRACKS_API FRichCurve& GetShakeWeightCurve() { return ShakeWeightCurve; }
-	
-private:
+	virtual FMovieSceneEvalTemplatePtr GenerateTemplate() const override;
+	virtual void PostLoad() override;
 	
-	/** */
-// 	UPROPERTY(EditAnywhere, Category = "Camera Shake")
-// 	FRichCurve ShakeWeightCurve;
+	UPROPERTY(EditAnywhere, Category="Cmaera Shake", meta=(ShowOnlyInnerProperties))
+	FMovieSceneCameraShakeSectionData ShakeData;
 
 public:
-	/** Class of the camera shake to play */
-	UPROPERTY(EditAnywhere, Category = "Camera Shake")
-	TSubclassOf<UCameraShake> ShakeClass;
+	UPROPERTY()
+	TSubclassOf<UCameraShake> ShakeClass_DEPRECATED;
 	
-	/** Scalar that affects shake intensity */
-	UPROPERTY(EditAnywhere, Category = "Camera Shake")
-	float PlayScale;
-	
-	UPROPERTY(EditAnywhere, Category = "Camera Shake")
-	TEnumAsByte<ECameraAnimPlaySpace::Type> PlaySpace;
+	UPROPERTY()
+	float PlayScale_DEPRECATED;
 
-	UPROPERTY(EditAnywhere, Category = "Camera Shake")
-	FRotator UserDefinedPlaySpace;
+	UPROPERTY()
+	TEnumAsByte<ECameraAnimPlaySpace::Type> PlaySpace_DEPRECATED;
+
+	UPROPERTY()
+	FRotator UserDefinedPlaySpace_DEPRECATED;
 };
\ No newline at end of file
diff --git a/Engine/Source/Runtime/MovieSceneTracks/Public/Sections/MovieSceneColorSection.h b/Engine/Source/Runtime/MovieSceneTracks/Public/Sections/MovieSceneColorSection.h
--- a/Engine/Source/Runtime/MovieSceneTracks/Public/Sections/MovieSceneColorSection.h
+++ b/Engine/Source/Runtime/MovieSceneTracks/Public/Sections/MovieSceneColorSection.h
@@ -64,14 +64,6 @@ class UMovieSceneColorSection
 
 public:
 
-	/**
-	 * Updates this section
-	 *
-	 * @param Position The position in time within the movie scene.
-	 * @param DefaultColor The default value to return.
-	 */
-	virtual FLinearColor Eval(float Position, const FLinearColor& DefaultColor) const;
-
 	/**
 	 * Gets the red color curve
 	 *
diff --git a/Engine/Source/Runtime/MovieSceneTracks/Public/Sections/MovieSceneEventSection.h b/Engine/Source/Runtime/MovieSceneTracks/Public/Sections/MovieSceneEventSection.h
--- a/Engine/Source/Runtime/MovieSceneTracks/Public/Sections/MovieSceneEventSection.h
+++ b/Engine/Source/Runtime/MovieSceneTracks/Public/Sections/MovieSceneEventSection.h
@@ -3,11 +3,182 @@
 #pragma once
 
 #include "Curves/NameCurve.h"
+#include "Curves/CurveInterface.h"
+#include "MovieSceneSection.h"
 #include "MovieSceneEventSection.generated.h"
 
 
-class ALevelScriptActor;
+USTRUCT()
+struct FMovieSceneEventParameters
+{
+	GENERATED_BODY()
+
+	FMovieSceneEventParameters() {}
+
+	/** Construction from a struct type */
+	FMovieSceneEventParameters(UStruct& InStruct)
+		: StructType(&InStruct)
+	{
+	}
+
+	FMovieSceneEventParameters(const FMovieSceneEventParameters& RHS) = default;
+	FMovieSceneEventParameters& operator=(const FMovieSceneEventParameters& RHS) = default;
+
+#if PLATFORM_COMPILER_HAS_DEFAULTED_FUNCTIONS
+	FMovieSceneEventParameters(FMovieSceneEventParameters&&) = default;
+	FMovieSceneEventParameters& operator=(FMovieSceneEventParameters&&) = default;
+#else
+	FMovieSceneEventParameters(FMovieSceneEventParameters&& RHS)
+	{
+		*this = MoveTemp(RHS);
+	}
+	FMovieSceneEventParameters& operator=(FMovieSceneEventParameters&& RHS)
+	{
+		StructType = MoveTemp(RHS.StructType);
+		StructBytes = MoveTemp(RHS.StructBytes);
+		return *this;
+	}
+#endif
+
+	void OverwriteWith(const TArray<uint8>& Bytes)
+	{
+		StructBytes = Bytes;
+	}
+
+	void GetInstance(FStructOnScope& OutStruct) const
+	{
+		UStruct* StructPtr = StructType.Get();
+		OutStruct.Initialize(StructPtr);
+		uint8* Memory = OutStruct.GetStructMemory();
+		if (StructPtr && StructPtr->GetStructureSize() > 0 && StructBytes.Num())
+		{
+			FMemoryReader Reader(StructBytes);
+			StructPtr->SerializeTaggedProperties(Reader, Memory, StructPtr, nullptr);
+		}
+	}
+
+	UStruct* GetStructType() const
+	{
+		return StructType.Get();
+	}
+
+	void Reassign(UStruct* NewStruct)
+	{
+		StructType = NewStruct;
+
+		if (!NewStruct)
+		{
+			StructBytes.Reset();
+		}
+	}
+
+	bool Serialize(FArchive& Ar)
+	{
+		UStruct* StructTypePtr = StructType.Get();
+		Ar << StructTypePtr;
+		StructType = StructTypePtr;
+		
+		Ar << StructBytes;
+
+		return true;
+	}
+
+	friend FArchive& operator<<(FArchive& Ar, FMovieSceneEventParameters& Payload)
+	{
+		Payload.Serialize(Ar);
+		return Ar;
+	}
+
+private:
+
+	TWeakObjectPtr<UStruct> StructType;
+	TArray<uint8> StructBytes;
+};
+
+template<>
+struct TStructOpsTypeTraits<FMovieSceneEventParameters> : public TStructOpsTypeTraitsBase
+{
+	enum 
+	{
+		WithCopy = true,
+		WithSerializer = true
+	};
+};
+
+USTRUCT()
+struct FEventPayload
+{
+	GENERATED_BODY()
 
+	FEventPayload() {}
+	FEventPayload(FName InEventName) : EventName(InEventName) {}
+
+	FEventPayload(const FEventPayload&) = default;
+	FEventPayload& operator=(const FEventPayload&) = default;
+
+#if PLATFORM_COMPILER_HAS_DEFAULTED_FUNCTIONS
+	FEventPayload(FEventPayload&&) = default;
+	FEventPayload& operator=(FEventPayload&&) = default;
+#else
+	FEventPayload(FEventPayload&& RHS)
+		: EventName(MoveTemp(RHS.EventName))
+		, Parameters(MoveTemp(RHS.Parameters))
+	{
+	}
+	FEventPayload& operator=(FEventPayload&& RHS)
+	{
+		EventName = MoveTemp(RHS.EventName);
+		Parameters = MoveTemp(RHS.Parameters);
+		return *this;
+	}
+#endif
+
+	/** The name of the event to trigger */
+	UPROPERTY(EditAnywhere, Category=Event)
+	FName EventName;
+
+	/** The event parameters */
+	UPROPERTY(EditAnywhere, Category=Event, meta=(ShowOnlyInnerProperties))
+	FMovieSceneEventParameters Parameters;
+};
+
+/** A curve of events */
+USTRUCT()
+struct FMovieSceneEventSectionData
+{
+	GENERATED_BODY()
+	
+	FMovieSceneEventSectionData() = default;
+
+	FMovieSceneEventSectionData(const FMovieSceneEventSectionData& RHS)
+		: KeyTimes(RHS.KeyTimes)
+		, KeyValues(RHS.KeyValues)
+	{
+	}
+
+	FMovieSceneEventSectionData& operator=(const FMovieSceneEventSectionData& RHS)
+	{
+		KeyTimes = RHS.KeyTimes;
+		KeyValues = RHS.KeyValues;
+#if WITH_EDITORONLY_DATA
+		KeyHandles.Reset();
+#endif
+		return *this;
+	}
+
+	/** Sorted array of key times */
+	UPROPERTY()
+	TArray<float> KeyTimes;
+
+	/** Array of values that correspond to each key time */
+	UPROPERTY()
+	TArray<FEventPayload> KeyValues;
+
+#if WITH_EDITORONLY_DATA
+	/** Transient key handles */
+	FKeyHandleLookupTable KeyHandles;
+#endif
+};
 
 /**
  * Implements a section in movie scene event tracks.
@@ -22,34 +193,18 @@ class UMovieSceneEventSection
 	UMovieSceneEventSection();
 
 public:
+	
+	// ~UObject interface
+	virtual void PostLoad() override;
 
 	/**
-	 * Add a key with the specified event.
+	 * Get the section's event data.
 	 *
-	 * @param Time The location in time where the key should be added.
-	 * @param EventName The name of the event to fire at the specified time.
-	 * @param KeyParams The keying parameters.
+	 * @return Event data.
 	 */
-	void AddKey(float Time, const FName& EventName);
-
-	/**
-	 * Get the section's event curve.
-	 *
-	 * @return Event curve.
-	 */
-	FNameCurve& GetEventCurve()
-	{
-		return Events;
-	}
-
-	/**
-	 * Trigger the events that fall into the given time range.
-	 *
-	 * @param Position The current position in time.
-	 * @param LastPosition The time at the last update.
-	 * @param Player The movie scene player that has the event contexts where the events should be invoked from.
-	 */
-	void TriggerEvents(float Position, float LastPosition, IMovieScenePlayer& Player);
+	const FMovieSceneEventSectionData& GetEventData() const { return EventData; }
+	
+	TCurveInterface<FEventPayload, float> GetCurveInterface() { return CurveInterface.GetValue(); }
 
 public:
 
@@ -61,20 +216,13 @@ public:
 	virtual TOptional<float> GetKeyTime(FKeyHandle KeyHandle) const override;
 	virtual void SetKeyTime(FKeyHandle KeyHandle, float Time) override;
 
-protected:
-
-	/**
-	 * Trigger event for the specified name.
-	 *
-	 * @param Event The Name to trigger.
-	 * @param Position The current position in time.
-	 * @param Player The movie scene player that has the event contexts where the events should be invoked from.
-	 */
-	void TriggerEvent(const FName& Event, float Position, IMovieScenePlayer& Player);
-
 private:
 
-	/** The section's keys. */
 	UPROPERTY()
-	FNameCurve Events;
+	FNameCurve Events_DEPRECATED;
+
+	UPROPERTY()
+	FMovieSceneEventSectionData EventData;
+
+	TOptional<TCurveInterface<FEventPayload, float>> CurveInterface;
 };
diff --git a/Engine/Source/Runtime/MovieSceneTracks/Public/Sections/MovieSceneIntegerSection.h b/Engine/Source/Runtime/MovieSceneTracks/Public/Sections/MovieSceneIntegerSection.h
new file mode 100644
--- /dev/null
+++ b/Engine/Source/Runtime/MovieSceneTracks/Public/Sections/MovieSceneIntegerSection.h
@@ -0,0 +1,51 @@
+// Copyright 1998-2016 Epic Games, Inc. All Rights Reserved.
+
+#pragma once
+
+#include "Curves/IntegralCurve.h"
+#include "MovieSceneSection.h"
+#include "IKeyframeSection.h"
+#include "MovieSceneIntegerSection.generated.h"
+
+
+/**
+ * A single integer section.
+ */
+UCLASS(MinimalAPI)
+class UMovieSceneIntegerSection 
+	: public UMovieSceneSection
+	, public IKeyframeSection<int32>
+{
+	GENERATED_UCLASS_BODY()
+public:
+
+	/** Gets all the keys of this integer section */
+	FIntegralCurve& GetCurve() { return IntegerCurve; }
+	const FIntegralCurve& GetCurve() const { return IntegerCurve; }
+
+public:
+
+	//~ IKeyframeSection interface
+
+	virtual void AddKey(float Time, const int32& Value, EMovieSceneKeyInterpolation KeyInterpolation) override;
+	virtual bool NewKeyIsNewData(float Time, const int32& Value) const override;
+	virtual bool HasKeys(const int32& Value) const override;
+	virtual void SetDefault(const int32& Value) override;
+	virtual void ClearDefaults() override;
+
+public:
+
+	//~ UMovieSceneSection interface 
+
+	virtual void MoveSection(float DeltaPosition, TSet<FKeyHandle>& KeyHandles) override;
+	virtual void DilateSection(float DilationFactor, float Origin, TSet<FKeyHandle>& KeyHandles) override;
+	virtual void GetKeyHandles(TSet<FKeyHandle>& OutKeyHandles, TRange<float> TimeRange) const override;
+	virtual TOptional<float> GetKeyTime(FKeyHandle KeyHandle) const override;
+	virtual void SetKeyTime(FKeyHandle KeyHandle, float Time) override;
+
+private:
+
+	/** Ordered curve data */
+	UPROPERTY()
+	FIntegralCurve IntegerCurve;
+};
diff --git a/Engine/Source/Runtime/MovieSceneTracks/Public/Sections/MovieSceneLevelVisibilitySection.h b/Engine/Source/Runtime/MovieSceneTracks/Public/Sections/MovieSceneLevelVisibilitySection.h
--- a/Engine/Source/Runtime/MovieSceneTracks/Public/Sections/MovieSceneLevelVisibilitySection.h
+++ b/Engine/Source/Runtime/MovieSceneTracks/Public/Sections/MovieSceneLevelVisibilitySection.h
@@ -34,11 +34,12 @@ public:
 	void SetVisibility(ELevelVisibility InVisibility);
 
 	TArray<FName>* GetLevelNames();
+	const TArray<FName>& GetLevelNames() const { return LevelNames; }
 
 public:
 
 	//~ UMovieSceneSection interface
-
+	virtual FMovieSceneEvalTemplatePtr GenerateTemplate() const override;
 	virtual TOptional<float> GetKeyTime(FKeyHandle KeyHandle) const override;
 	virtual void SetKeyTime(FKeyHandle KeyHandle, float Time) override;
 
diff --git a/Engine/Source/Runtime/MovieSceneTracks/Public/Sections/MovieSceneParameterSection.h b/Engine/Source/Runtime/MovieSceneTracks/Public/Sections/MovieSceneParameterSection.h
--- a/Engine/Source/Runtime/MovieSceneTracks/Public/Sections/MovieSceneParameterSection.h
+++ b/Engine/Source/Runtime/MovieSceneTracks/Public/Sections/MovieSceneParameterSection.h
@@ -218,29 +218,19 @@ public:
 
 	/** Gets the animated scalar parameters and their associated curves. */
 	MOVIESCENETRACKS_API TArray<FScalarParameterNameAndCurve>* GetScalarParameterNamesAndCurves();
+	MOVIESCENETRACKS_API const TArray<FScalarParameterNameAndCurve>& GetScalarParameterNamesAndCurves() const;
 
-	/** Gets the animated color parameters and their associated curves. */
+	/** Gets the animated vector parameters and their associated curves. */
 	MOVIESCENETRACKS_API TArray<FVectorParameterNameAndCurves>* GetVectorParameterNamesAndCurves();
+	MOVIESCENETRACKS_API const TArray<FVectorParameterNameAndCurves>& GetVectorParameterNamesAndCurves() const;
 
 	/** Gets the animated color parameters and their associated curves. */
 	MOVIESCENETRACKS_API TArray<FColorParameterNameAndCurves>* GetColorParameterNamesAndCurves();
+	MOVIESCENETRACKS_API const TArray<FColorParameterNameAndCurves>& GetColorParameterNamesAndCurves() const;
 
 	/** Gets the set of all parameter names used by this section. */
 	MOVIESCENETRACKS_API void GetParameterNames(TSet<FName>& ParameterNames) const;
 
-	/**
-	 * Gets the animated values for this track.
-	 *
-	 * @param Position The playback position to use for evaluation.
-	 * @param OutScalarValues An array of FScalarParameterNameAndValue objects representing each animated scalar parameter and it's animated value.
-	 * @param OutVectorValues An array of FVectorParameterNameAndValue objects representing each animated vector parameter and it's animated value.
-	 * @param OutColorValues An array of FColorParameterNameAndValue objects representing each animated color parameter and it's animated value.
-	 */
-	void Eval(float Position, 
-		TArray<FScalarParameterNameAndValue>& OutScalarValues,
-		TArray<FVectorParameterNameAndValue>& OutVectorValues,
-		TArray<FColorParameterNameAndValue>& OutColorValues) const;
-
 public:
 
 	//~ UMovieSceneSection interface
diff --git a/Engine/Source/Runtime/MovieSceneTracks/Public/Sections/MovieSceneParticleSection.h b/Engine/Source/Runtime/MovieSceneTracks/Public/Sections/MovieSceneParticleSection.h
--- a/Engine/Source/Runtime/MovieSceneTracks/Public/Sections/MovieSceneParticleSection.h
+++ b/Engine/Source/Runtime/MovieSceneTracks/Public/Sections/MovieSceneParticleSection.h
@@ -33,7 +33,8 @@ class UMovieSceneParticleSection
 
 	MOVIESCENETRACKS_API void AddKey(float Time, EParticleKey::Type KeyType);
 
-	MOVIESCENETRACKS_API FIntegralCurve& GetParticleCurve();
+	FIntegralCurve& GetParticleCurve() { return ParticleKeys; }
+	const FIntegralCurve& GetParticleCurve() const { return ParticleKeys; }
 
 public:
 
@@ -44,7 +45,8 @@ public:
 	virtual void GetKeyHandles(TSet<FKeyHandle>& OutKeyHandles, TRange<float> TimeRange) const override;
 	virtual TOptional<float> GetKeyTime(FKeyHandle KeyHandle) const override;
 	virtual void SetKeyTime(FKeyHandle KeyHandle, float Time) override;
-
+	virtual FMovieSceneEvalTemplatePtr GenerateTemplate() const override;
+	
 private:
 
 	/** Curve containing the particle keys. */
diff --git a/Engine/Source/Runtime/MovieSceneTracks/Public/Sections/MovieSceneSkeletalAnimationSection.h b/Engine/Source/Runtime/MovieSceneTracks/Public/Sections/MovieSceneSkeletalAnimationSection.h
--- a/Engine/Source/Runtime/MovieSceneTracks/Public/Sections/MovieSceneSkeletalAnimationSection.h
+++ b/Engine/Source/Runtime/MovieSceneTracks/Public/Sections/MovieSceneSkeletalAnimationSection.h
@@ -7,6 +7,44 @@
 
 class UAnimSequenceBase;
 
+USTRUCT()
+struct FMovieSceneSkeletalAnimationParams
+{
+	GENERATED_BODY()
+
+	FMovieSceneSkeletalAnimationParams();
+
+	/** Gets the animation duration, modified by play rate */
+	float GetDuration() const { return FMath::IsNearlyZero(PlayRate) || Animation == nullptr ? 0.f : Animation->SequenceLength / PlayRate; }
+
+	/** Gets the animation sequence length, not modified by play rate */
+	float GetSequenceLength() const { return Animation != nullptr ? Animation->SequenceLength : 0.f; }
+
+	/** The animation this section plays */
+	UPROPERTY(EditAnywhere, Category="Animation")
+	UAnimSequenceBase* Animation;
+
+	/** The offset into the beginning of the animation clip */
+	UPROPERTY(EditAnywhere, Category="Animation")
+	float StartOffset;
+	
+	/** The offset into the end of the animation clip */
+	UPROPERTY(EditAnywhere, Category="Animation")
+	float EndOffset;
+	
+	/** The playback rate of the animation clip */
+	UPROPERTY(EditAnywhere, Category="Animation")
+	float PlayRate;
+
+	/** Reverse the playback of the animation clip */
+	UPROPERTY(EditAnywhere, Category="Animation")
+	uint32 bReverse:1;
+
+	/** The slot name to use for the animation */
+	UPROPERTY( EditAnywhere, Category = "Animation" )
+	FName SlotName;
+};
+
 /**
  * Movie scene section that control skeletal animation
  */
@@ -18,47 +56,8 @@
 
 public:
 
-	/** Sets the animation sequence for this section */
-	void SetAnimSequence(UAnimSequenceBase* InAnimation) { Animation = InAnimation; }
-
-	/** Gets the animation sequence for this section */
-	UAnimSequenceBase* GetAnimSequence() { return Animation; }
-	
-	/** Gets the start offset into the animation clip */
-	float GetStartOffset() const { return StartOffset; }
-
-	/** Sets the start offset into the animation clip */
-	void SetStartOffset(float InStartOffset) { StartOffset = InStartOffset; }
-	
-	/** Gets the end offset into the animation clip */
-	float GetEndOffset() const { return EndOffset; }
-
-	/** Sets the end offset into the animation clip */
-	void SetEndOffset(float InEndOffset) { EndOffset = InEndOffset; }
-	
-	/** Gets the animation duration, modified by play rate */
-	float GetDuration() const { return FMath::IsNearlyZero(PlayRate) || Animation == nullptr ? 0.f : Animation->SequenceLength / PlayRate; }
-
-	/** Gets the animation sequence length, not modified by play rate */
-	float GetSequenceLength() const { return Animation != nullptr ? Animation->SequenceLength : 0.f; }
-
-	/** Sets the play rate of the animation clip */
-	float GetPlayRate() const { return PlayRate; }
-
-	/** Sets the play rate of the animation clip */
-	void SetPlayRate(float InPlayRate) { PlayRate = InPlayRate; }
-
-	/** Gets whether the playback is reversed */
-	bool GetReverse() const { return bReverse; }
-
-	/** Sets whether the playback is reversed */
-	void SetReverse(bool bInReverse) { bReverse = bInReverse; }
-
-	/** Gets the anim BP slot name. */
-	FName GetSlotName() const { return SlotName; }
-
-	/** Sets the anim BP slot name. */
-	void SetSlotName( FName InSlotName ) { SlotName = InSlotName; }
+	UPROPERTY(EditAnywhere, Category = "Animation", meta=(ShowOnlyInnerProperties))
+	FMovieSceneSkeletalAnimationParams Params;
 
 public:
 
@@ -71,6 +70,7 @@
 	virtual void GetSnapTimes(TArray<float>& OutSnapTimes, bool bGetSectionBorders) const override;
 	virtual TOptional<float> GetKeyTime( FKeyHandle KeyHandle ) const override { return TOptional<float>(); }
 	virtual void SetKeyTime( FKeyHandle KeyHandle, float Time ) override { }
+	virtual FMovieSceneEvalTemplatePtr GenerateTemplate() const override;
 
 	/** ~UObject interface */
 	virtual void PostLoad() override;
@@ -89,32 +89,24 @@
 
 private:
 
-	static FName DefaultSlotName;
-
 	UPROPERTY()
 	class UAnimSequence* AnimSequence_DEPRECATED;
 
-	/** The animation this section plays */
-	UPROPERTY(EditAnywhere, Category="Animation")
-	UAnimSequenceBase* Animation;
+	UPROPERTY()
+	UAnimSequenceBase* Animation_DEPRECATED;
 
-	/** The offset into the beginning of the animation clip */
-	UPROPERTY(EditAnywhere, Category="Animation")
-	float StartOffset;
+	UPROPERTY()
+	float StartOffset_DEPRECATED;
 	
-	/** The offset into the end of the animation clip */
-	UPROPERTY(EditAnywhere, Category="Animation")
-	float EndOffset;
+	UPROPERTY()
+	float EndOffset_DEPRECATED;
 	
-	/** The playback rate of the animation clip */
-	UPROPERTY(EditAnywhere, Category="Animation")
-	float PlayRate;
+	UPROPERTY()
+	float PlayRate_DEPRECATED;
 
-	/** Reverse the playback of the animation clip */
-	UPROPERTY(EditAnywhere, Category="Animation")
-	uint32 bReverse:1;
+	UPROPERTY()
+	uint32 bReverse_DEPRECATED:1;
 
-	/** The slot name to use for the animation */
-	UPROPERTY( EditAnywhere, Category = "Animation" )
-	FName SlotName;
+	UPROPERTY()
+	FName SlotName_DEPRECATED;
 };
diff --git a/Engine/Source/Runtime/MovieSceneTracks/Public/Sections/MovieSceneStringSection.h b/Engine/Source/Runtime/MovieSceneTracks/Public/Sections/MovieSceneStringSection.h
--- a/Engine/Source/Runtime/MovieSceneTracks/Public/Sections/MovieSceneStringSection.h
+++ b/Engine/Source/Runtime/MovieSceneTracks/Public/Sections/MovieSceneStringSection.h
@@ -37,6 +37,10 @@ public:
 	{
 		return StringCurve;
 	}
+	const FStringCurve& GetStringCurve() const
+	{
+		return StringCurve;
+	}
 
 public:
 
diff --git a/Engine/Source/Runtime/MovieSceneTracks/Public/Sections/MovieSceneSubSection.h b/Engine/Source/Runtime/MovieSceneTracks/Public/Sections/MovieSceneSubSection.h
--- a/Engine/Source/Runtime/MovieSceneTracks/Public/Sections/MovieSceneSubSection.h
+++ b/Engine/Source/Runtime/MovieSceneTracks/Public/Sections/MovieSceneSubSection.h
@@ -3,6 +3,7 @@
 #pragma once
 
 #include "MovieSceneSection.h"
+#include "MovieSceneSectionParameters.h"
 #include "MovieSceneSubSection.generated.h"
 
 
@@ -33,6 +34,16 @@ public:
 	 */
 	UMovieSceneSequence* GetSequence() const;
 
+	/**
+	 * Get the path name to this sub section from the outer moviescene
+	 */
+	FString GetPathNameInMovieScene() const;
+
+	/**
+	 * Get this sub section's sequence ID
+	 */
+	FMovieSceneSequenceID GetSequenceID() const;
+
 public:
 
 	/**
@@ -55,6 +66,8 @@ public:
 	/** Check if we are primed for recording */
 	static bool IsSetAsRecording();
 
+	virtual void PostLoad() override;
+
 #if WITH_EDITOR
 	virtual void PostEditChangeProperty(FPropertyChangedEvent& PropertyChangedEvent) override;
 
@@ -103,17 +116,19 @@ public:
 
 public:
 
-	/** Number of seconds to skip at the beginning of the sub-sequence. */
-	UPROPERTY(EditAnywhere, Category="Clipping")
-	float StartOffset;
+	UPROPERTY(EditAnywhere, Category="General", meta=(ShowOnlyInnerProperties))
+	FMovieSceneSectionParameters Parameters;
+
+private:
+
+	UPROPERTY()
+	float StartOffset_DEPRECATED;
 
-	/** Playback time scaling factor. */
-	UPROPERTY(EditAnywhere, Category="Timing")
-	float TimeScale;
+	UPROPERTY()
+	float TimeScale_DEPRECATED;
 
-	/** Preroll time. */
-	UPROPERTY(EditAnywhere, Category="Timing")
-	float PrerollTime;
+	UPROPERTY()
+	float PrerollTime_DEPRECATED;
 
 protected:
 
diff --git a/Engine/Source/Runtime/MovieSceneTracks/Public/Sections/MovieSceneVectorSection.h b/Engine/Source/Runtime/MovieSceneTracks/Public/Sections/MovieSceneVectorSection.h
--- a/Engine/Source/Runtime/MovieSceneTracks/Public/Sections/MovieSceneVectorSection.h
+++ b/Engine/Source/Runtime/MovieSceneTracks/Public/Sections/MovieSceneVectorSection.h
@@ -129,16 +129,10 @@ class UMovieSceneVectorSection
 	GENERATED_UCLASS_BODY()
 
 public:
-	
-	/**
-	 * Updates this section
-	 *
-	 * @param Position	The position in time within the movie scene
-	 */
-	FVector4 Eval(float Position, const FVector4& DefaultVector) const;
 
 	/** Gets one of four curves in this section */
-	FRichCurve& GetCurve(const int32& Index) { return Curves[Index]; }
+	FRichCurve& GetCurve(int32 Index) { check(Index >= 0 && Index < 4); return Curves[Index]; }
+	const FRichCurve& GetCurve(int32 Index) const { check(Index >= 0 && Index < 4); return Curves[Index]; }
 
 	/** Sets how many channels are to be used */
 	void SetChannelsUsed(int32 InChannelsUsed) 
diff --git a/Engine/Source/Runtime/MovieSceneTracks/Public/Tracks/MovieScene3DAttachTrack.h b/Engine/Source/Runtime/MovieSceneTracks/Public/Tracks/MovieScene3DAttachTrack.h
--- a/Engine/Source/Runtime/MovieSceneTracks/Public/Tracks/MovieScene3DAttachTrack.h
+++ b/Engine/Source/Runtime/MovieSceneTracks/Public/Tracks/MovieScene3DAttachTrack.h
@@ -27,7 +27,7 @@ public:
 
 	// UMovieSceneTrack interface
 
-	virtual TSharedPtr<IMovieSceneTrackInstance> CreateInstance() override;
+	virtual FMovieSceneEvalTemplatePtr CreateTemplateForSection(const UMovieSceneSection& InSection) const override;
 
 #if WITH_EDITORONLY_DATA
 	virtual FText GetDisplayName() const override;
diff --git a/Engine/Source/Runtime/MovieSceneTracks/Public/Tracks/MovieScene3DPathTrack.h b/Engine/Source/Runtime/MovieSceneTracks/Public/Tracks/MovieScene3DPathTrack.h
--- a/Engine/Source/Runtime/MovieSceneTracks/Public/Tracks/MovieScene3DPathTrack.h
+++ b/Engine/Source/Runtime/MovieSceneTracks/Public/Tracks/MovieScene3DPathTrack.h
@@ -25,7 +25,7 @@ public:
 
 	// UMovieSceneTrack interface
 
-	virtual TSharedPtr<IMovieSceneTrackInstance> CreateInstance() override;
+	virtual FMovieSceneEvalTemplatePtr CreateTemplateForSection(const UMovieSceneSection& InSection) const override;
 
 #if WITH_EDITORONLY_DATA
 	virtual FText GetDisplayName() const override;
diff --git a/Engine/Source/Runtime/MovieSceneTracks/Public/Tracks/MovieScene3DTransformTrack.h b/Engine/Source/Runtime/MovieSceneTracks/Public/Tracks/MovieScene3DTransformTrack.h
--- a/Engine/Source/Runtime/MovieSceneTracks/Public/Tracks/MovieScene3DTransformTrack.h
+++ b/Engine/Source/Runtime/MovieSceneTracks/Public/Tracks/MovieScene3DTransformTrack.h
@@ -2,12 +2,9 @@
 
 #pragma once
 
-#include "MovieScene.h"
-#include "MovieSceneTrack.h"
 #include "MovieScenePropertyTrack.h"
 #include "MovieScene3DTransformTrack.generated.h"
 
-
 /**
  * Handles manipulation of component transforms in a movie scene
  */
@@ -17,24 +14,9 @@ class UMovieScene3DTransformTrack
 {
 	GENERATED_UCLASS_BODY()
 
-public:
-
-	/**
-	 * Evaluates the track at the playback position
-	 *
-	 * @param Position The current playback position.
-	 * @param LastPosition The last playback position.
-	 * @param OutTranslation The evaluated translation component of the transform.
-	 * @param OutRotation The evaluated rotation component of the transform.
-	 * @param OutScale The evaluated scale component of the transform.
-	 * @return true if anything was evaluated. Note: if false is returned OutBool remains unchanged.
-	 */
-	virtual bool Eval(float Position, float LastPosition, FVector& InOutTranslation, FRotator& InOutRotation, FVector& InOutScale) const;
-
 public:
 
 	// UMovieSceneTrack interface
 
 	virtual UMovieSceneSection* CreateNewSection() override;
-	virtual TSharedPtr<IMovieSceneTrackInstance> CreateInstance() override;
-};
+};
\ No newline at end of file
diff --git a/Engine/Source/Runtime/MovieSceneTracks/Public/Tracks/MovieSceneActorReferenceTrack.h b/Engine/Source/Runtime/MovieSceneTracks/Public/Tracks/MovieSceneActorReferenceTrack.h
--- a/Engine/Source/Runtime/MovieSceneTracks/Public/Tracks/MovieSceneActorReferenceTrack.h
+++ b/Engine/Source/Runtime/MovieSceneTracks/Public/Tracks/MovieSceneActorReferenceTrack.h
@@ -18,15 +18,5 @@ class UMovieSceneActorReferenceTrack : public UMovieScenePropertyTrack
 public:
 	/** UMovieSceneTrack interface */
 	virtual UMovieSceneSection* CreateNewSection() override;
-	virtual TSharedPtr<IMovieSceneTrackInstance> CreateInstance() override;
-
-	/**
-	 * Evaluates the track at the playback position
-	 *
-	 * @param Position	The current playback position
-	 * @param LastPosition	The last plackback position
-	 * @param OutActorReference 	The value at the playback position
-	 * @return true if anything was evaluated. Note: if false is returned OutActorReference remains unchanged
-	 */
-	virtual bool Eval( float Position, float LastPosition, FGuid& OutActorReferenceGuid ) const;
+	virtual FMovieSceneEvalTemplatePtr CreateTemplateForSection(const UMovieSceneSection& InSection) const override;
 };
diff --git a/Engine/Source/Runtime/MovieSceneTracks/Public/Tracks/MovieSceneAudioTrack.h b/Engine/Source/Runtime/MovieSceneTracks/Public/Tracks/MovieSceneAudioTrack.h
--- a/Engine/Source/Runtime/MovieSceneTracks/Public/Tracks/MovieSceneAudioTrack.h
+++ b/Engine/Source/Runtime/MovieSceneTracks/Public/Tracks/MovieSceneAudioTrack.h
@@ -43,7 +43,6 @@ public:
 
 	// UMovieSceneTrack interface
 
-	virtual TSharedPtr<IMovieSceneTrackInstance> CreateInstance() override;
 	virtual void RemoveAllAnimationData() override;
 	virtual bool HasSection(const UMovieSceneSection& Section) const override;
 	virtual void AddSection(UMovieSceneSection& Section) override;
@@ -52,7 +51,7 @@ public:
 	virtual TRange<float> GetSectionBoundaries() const override;
 	virtual const TArray<UMovieSceneSection*>& GetAllSections() const override;
 	virtual bool SupportsMultipleRows() const override;
-
+	
 private:
 
 	/** List of all master audio sections */
diff --git a/Engine/Source/Runtime/MovieSceneTracks/Public/Tracks/MovieSceneBoolTrack.h b/Engine/Source/Runtime/MovieSceneTracks/Public/Tracks/MovieSceneBoolTrack.h
--- a/Engine/Source/Runtime/MovieSceneTracks/Public/Tracks/MovieSceneBoolTrack.h
+++ b/Engine/Source/Runtime/MovieSceneTracks/Public/Tracks/MovieSceneBoolTrack.h
@@ -22,15 +22,8 @@ public:
 	// UMovieSceneTrack interface
 
 	virtual UMovieSceneSection* CreateNewSection() override;
-	virtual TSharedPtr<IMovieSceneTrackInstance> CreateInstance() override;
-
-	/**
-	 * Evaluates the track at the playback position
-	 *
-	 * @param Position	The current playback position
-	 * @param LastPosition	The last plackback position
-	 * @param InOutBool 	The current value at the playback position
-	 * @return true if anything was evaluated. Note: if false is returned InOutBool remains unchanged
-	 */
+	virtual FMovieSceneEvalTemplatePtr CreateTemplateForSection(const UMovieSceneSection& InSection) const override;
+
+	DEPRECATED(4.15, "Direct evaluation of boolean tracks is no longer supported. Please create an evaluation template (see FMovieSceneBoolPropertySectionTemplate).")
 	virtual bool Eval( float Position, float LastPostion, bool& InOutBool ) const;
 };
diff --git a/Engine/Source/Runtime/MovieSceneTracks/Public/Tracks/MovieSceneByteTrack.h b/Engine/Source/Runtime/MovieSceneTracks/Public/Tracks/MovieSceneByteTrack.h
--- a/Engine/Source/Runtime/MovieSceneTracks/Public/Tracks/MovieSceneByteTrack.h
+++ b/Engine/Source/Runtime/MovieSceneTracks/Public/Tracks/MovieSceneByteTrack.h
@@ -18,17 +18,7 @@ class MOVIESCENETRACKS_API UMovieSceneByteTrack : public UMovieScenePropertyTrac
 public:
 	/** UMovieSceneTrack interface */
 	virtual UMovieSceneSection* CreateNewSection() override;
-	virtual TSharedPtr<IMovieSceneTrackInstance> CreateInstance() override;
-
-	/**
-	 * Evaluates the track at the playback position
-	 *
-	 * @param Position	The current playback position
-	 * @param LastPosition	The last plackback position
-	 * @param InOutByte 	The value at the playback position
-	 * @return true if anything was evaluated. Note: if false is returned InOutByte remains unchanged
-	 */
-	virtual bool Eval( float Position, float LastPostion, uint8& InOutByte ) const;
+	virtual FMovieSceneEvalTemplatePtr CreateTemplateForSection(const UMovieSceneSection& InSection) const override;
 
 	void SetEnum(UEnum* Enum);
 
diff --git a/Engine/Source/Runtime/MovieSceneTracks/Public/Tracks/MovieSceneCameraAnimTrack.h b/Engine/Source/Runtime/MovieSceneTracks/Public/Tracks/MovieSceneCameraAnimTrack.h
--- a/Engine/Source/Runtime/MovieSceneTracks/Public/Tracks/MovieSceneCameraAnimTrack.h
+++ b/Engine/Source/Runtime/MovieSceneTracks/Public/Tracks/MovieSceneCameraAnimTrack.h
@@ -3,6 +3,7 @@
 #pragma once
 
 #include "MovieSceneNameableTrack.h"
+#include "MovieSceneSegmentCompiler.h"
 #include "MovieSceneCameraAnimTrack.generated.h"
 
 class UMovieSceneCameraAnimSection;
@@ -22,8 +23,6 @@ public:
 public:
 
 	// UMovieSceneTrack interface
-	virtual TSharedPtr<IMovieSceneTrackInstance> CreateInstance() override;
-
 	virtual bool HasSection(const UMovieSceneSection& Section) const override;
 	virtual void AddSection(UMovieSceneSection& Section) override;
 	virtual void RemoveSection(UMovieSceneSection& Section) override;
@@ -33,7 +32,8 @@ public:
 	virtual const TArray<UMovieSceneSection*>& GetAllSections() const override;
 	virtual UMovieSceneSection* CreateNewSection() override;
 	virtual void RemoveAllAnimationData() override;
-
+	virtual TInlineValue<FMovieSceneSegmentCompilerRules> GetTrackCompilerRules() const override;
+	virtual void PostCompile(FMovieSceneEvaluationTrack& OutTrack, const FMovieSceneTrackCompilerArgs& Args) const override;
 	
 #if WITH_EDITORONLY_DATA
 	virtual FText GetDisplayName() const override;
diff --git a/Engine/Source/Runtime/MovieSceneTracks/Public/Tracks/MovieSceneCameraCutTrack.h b/Engine/Source/Runtime/MovieSceneTracks/Public/Tracks/MovieSceneCameraCutTrack.h
--- a/Engine/Source/Runtime/MovieSceneTracks/Public/Tracks/MovieSceneCameraCutTrack.h
+++ b/Engine/Source/Runtime/MovieSceneTracks/Public/Tracks/MovieSceneCameraCutTrack.h
@@ -8,6 +8,7 @@
 
 class IMovieSceneTrackInstance;
 class UMovieSceneSection;
+class UMovieSceneCameraCutSection;
 
 
 /**
@@ -35,7 +36,6 @@ public:
 	// UMovieSceneTrack interface
 
 	virtual void AddSection(UMovieSceneSection& Section) override;
-	virtual TSharedPtr<IMovieSceneTrackInstance> CreateInstance() override;
 	virtual UMovieSceneSection* CreateNewSection() override;
 	virtual const TArray<UMovieSceneSection*>& GetAllSections() const override;
 	virtual void RemoveSection(UMovieSceneSection& Section) override;
@@ -57,4 +57,4 @@ private:
 	/** All movie scene sections. */
 	UPROPERTY()
 	TArray<UMovieSceneSection*> Sections;
-};
+};
\ No newline at end of file
diff --git a/Engine/Source/Runtime/MovieSceneTracks/Public/Tracks/MovieSceneCameraShakeTrack.h b/Engine/Source/Runtime/MovieSceneTracks/Public/Tracks/MovieSceneCameraShakeTrack.h
--- a/Engine/Source/Runtime/MovieSceneTracks/Public/Tracks/MovieSceneCameraShakeTrack.h
+++ b/Engine/Source/Runtime/MovieSceneTracks/Public/Tracks/MovieSceneCameraShakeTrack.h
@@ -18,13 +18,10 @@ class UMovieSceneCameraShakeTrack : public UMovieSceneNameableTrack
 
 public:
 	virtual void AddNewCameraShake(float KeyTime, TSubclassOf<UCameraShake> ShakeClass);
-	void GetCameraShakeSectionsAtTime(float Time, TArray<UMovieSceneCameraShakeSection*>& OutSections);
 	
 public:
 
 	// UMovieSceneTrack interface
-	virtual TSharedPtr<IMovieSceneTrackInstance> CreateInstance() override;
-
 	virtual bool HasSection(const UMovieSceneSection& Section) const override;
 	virtual void AddSection(UMovieSceneSection& Section) override;
 	virtual void RemoveSection(UMovieSceneSection& Section) override;
@@ -34,6 +31,8 @@ public:
 	virtual const TArray<UMovieSceneSection*>& GetAllSections() const override;
 	virtual UMovieSceneSection* CreateNewSection() override;
 	virtual void RemoveAllAnimationData() override;
+	virtual TInlineValue<FMovieSceneSegmentCompilerRules> GetTrackCompilerRules() const override;
+	virtual void PostCompile(FMovieSceneEvaluationTrack& OutTrack, const FMovieSceneTrackCompilerArgs& Args) const override;
 
 #if WITH_EDITORONLY_DATA
 	virtual FText GetDisplayName() const override;
diff --git a/Engine/Source/Runtime/MovieSceneTracks/Public/Tracks/MovieSceneCinematicShotTrack.h b/Engine/Source/Runtime/MovieSceneTracks/Public/Tracks/MovieSceneCinematicShotTrack.h
--- a/Engine/Source/Runtime/MovieSceneTracks/Public/Tracks/MovieSceneCinematicShotTrack.h
+++ b/Engine/Source/Runtime/MovieSceneTracks/Public/Tracks/MovieSceneCinematicShotTrack.h
@@ -3,6 +3,7 @@
 #pragma once
 
 #include "MovieSceneSubTrack.h"
+#include "MovieSceneSegmentCompiler.h"
 #include "MovieSceneCinematicShotTrack.generated.h"
 
 
@@ -30,11 +31,12 @@ public:
 	// UMovieSceneTrack interface
 
 	virtual void AddSection(UMovieSceneSection& Section) override;
-	virtual TSharedPtr<IMovieSceneTrackInstance> CreateInstance() override;
 	virtual UMovieSceneSection* CreateNewSection() override;
 	virtual void RemoveSection(UMovieSceneSection& Section) override;
 	virtual bool SupportsMultipleRows() const override;
-
+	virtual TInlineValue<FMovieSceneSegmentCompilerRules> GetRowCompilerRules() const override;
+	virtual TInlineValue<FMovieSceneSegmentCompilerRules> GetTrackCompilerRules() const override;
+	
 #if WITH_EDITOR
 	virtual void OnSectionMoved(UMovieSceneSection& Section) override;
 #endif
diff --git a/Engine/Source/Runtime/MovieSceneTracks/Public/Tracks/MovieSceneColorTrack.h b/Engine/Source/Runtime/MovieSceneTracks/Public/Tracks/MovieSceneColorTrack.h
--- a/Engine/Source/Runtime/MovieSceneTracks/Public/Tracks/MovieSceneColorTrack.h
+++ b/Engine/Source/Runtime/MovieSceneTracks/Public/Tracks/MovieSceneColorTrack.h
@@ -18,17 +18,7 @@ class UMovieSceneColorTrack : public UMovieScenePropertyTrack
 public:
 	/** UMovieSceneTrack interface */
 	virtual UMovieSceneSection* CreateNewSection() override;
-	virtual TSharedPtr<IMovieSceneTrackInstance> CreateInstance() override;
-
-	/**
-	 * Evaluates the track at the playback position
-	 *
-	 * @param Position	The current playback position
-	 * @param LastPosition	The last plackback position
-	 * @param InOutColor The color at the playback position
-	 * @return true if anything was evaluated. Note: if false is returned InOutColor remains unchanged
-	 */
-	virtual bool Eval( float Position, float LastPostion, FLinearColor& InOutColor ) const;
+	virtual FMovieSceneEvalTemplatePtr CreateTemplateForSection(const UMovieSceneSection& Section) const override;
 
 private:
 	UPROPERTY()
diff --git a/Engine/Source/Runtime/MovieSceneTracks/Public/Tracks/MovieSceneEventTrack.h b/Engine/Source/Runtime/MovieSceneTracks/Public/Tracks/MovieSceneEventTrack.h
--- a/Engine/Source/Runtime/MovieSceneTracks/Public/Tracks/MovieSceneEventTrack.h
+++ b/Engine/Source/Runtime/MovieSceneTracks/Public/Tracks/MovieSceneEventTrack.h
@@ -3,6 +3,7 @@
 #pragma once
 
 #include "MovieSceneNameableTrack.h"
+#include "MovieSceneSpawnTrack.h"
 #include "MovieSceneEventTrack.generated.h"
 
 
@@ -29,32 +30,13 @@ public:
 
 public:
 
-	/**
-	 * Adds an event to the appropriate section.
-	 *
-	 * This method will create a new section if no appropriate one exists.
-	 *
-	 * @param Time The time at which the event should be triggered.
-	 * @param EventName The name of the event to be triggered.
-	 * @param KeyParams The keying parameters
-	 */
-	bool AddKeyToSection(float Time, FName EventName);
-
-	/**
-	 * Trigger the events that fall into the given time range.
-	 *
-	 * @param Position The current position in time.
-	 * @param LastPosition The time at the last update.
-	 * @param Player The movie scene player that has the event contexts where the events should be invoked from.
-	 */
-	void TriggerEvents(float Position, float LastPosition, IMovieScenePlayer& Player);
+	static uint16 GetEvaluationPriority() { return UMovieSceneSpawnTrack::GetEvaluationPriority() - 100; }
 
 public:
 
 	// UMovieSceneTrack interface
 
 	virtual void AddSection(UMovieSceneSection& Section) override;
-	virtual TSharedPtr<IMovieSceneTrackInstance> CreateInstance() override;
 	virtual UMovieSceneSection* CreateNewSection() override;
 	virtual const TArray<UMovieSceneSection*>& GetAllSections() const override;
 	virtual TRange<float> GetSectionBoundaries() const override;
@@ -62,12 +44,14 @@ public:
 	virtual bool IsEmpty() const override;
 	virtual void RemoveAllAnimationData() override;
 	virtual void RemoveSection(UMovieSceneSection& Section) override;
+	virtual FMovieSceneEvalTemplatePtr CreateTemplateForSection(const UMovieSceneSection& InSection) const override;
+	virtual void PostCompile(FMovieSceneEvaluationTrack& Track, const FMovieSceneTrackCompilerArgs& Args) const override;
 
 #if WITH_EDITORONLY_DATA
 	virtual FText GetDefaultDisplayName() const override;
 #endif
 
-private:
+public:
 
 	/** If events should be fired when passed playing the sequence forwards. */
 	UPROPERTY(EditAnywhere, Category=TrackEvent)
@@ -77,6 +61,8 @@ private:
 	UPROPERTY(EditAnywhere, Category=TrackEvent)
 	uint32 bFireEventsWhenBackwards:1;
 
+private:
+	
 	/** The track's sections. */
 	UPROPERTY()
 	TArray<UMovieSceneSection*> Sections;
diff --git a/Engine/Source/Runtime/MovieSceneTracks/Public/Tracks/MovieSceneFadeTrack.h b/Engine/Source/Runtime/MovieSceneTracks/Public/Tracks/MovieSceneFadeTrack.h
--- a/Engine/Source/Runtime/MovieSceneTracks/Public/Tracks/MovieSceneFadeTrack.h
+++ b/Engine/Source/Runtime/MovieSceneTracks/Public/Tracks/MovieSceneFadeTrack.h
@@ -17,10 +17,12 @@ class UMovieSceneFadeTrack
 
 public:
 
+	UMovieSceneFadeTrack(const FObjectInitializer& Init);
+
 	// UMovieSceneTrack interface
 
-	virtual TSharedPtr<IMovieSceneTrackInstance> CreateInstance() override;
 	virtual UMovieSceneSection* CreateNewSection() override;
+	virtual FMovieSceneEvalTemplatePtr CreateTemplateForSection(const UMovieSceneSection& InSection) const override;
 
 #if WITH_EDITORONLY_DATA
 	virtual FText GetDefaultDisplayName() const override;
diff --git a/Engine/Source/Runtime/MovieSceneTracks/Public/Tracks/MovieSceneFloatTrack.h b/Engine/Source/Runtime/MovieSceneTracks/Public/Tracks/MovieSceneFloatTrack.h
--- a/Engine/Source/Runtime/MovieSceneTracks/Public/Tracks/MovieSceneFloatTrack.h
+++ b/Engine/Source/Runtime/MovieSceneTracks/Public/Tracks/MovieSceneFloatTrack.h
@@ -18,15 +18,8 @@ class UMovieSceneFloatTrack : public UMovieScenePropertyTrack
 public:
 	/** UMovieSceneTrack interface */
 	virtual UMovieSceneSection* CreateNewSection() override;
-	virtual TSharedPtr<IMovieSceneTrackInstance> CreateInstance() override;
+	virtual FMovieSceneEvalTemplatePtr CreateTemplateForSection(const UMovieSceneSection& InSection) const override;
 
-	/**
-	 * Evaluates the track at the playback position
-	 *
-	 * @param Position	The current playback position
-	 * @param LastPosition	The last plackback position
-	 * @param InOutFloat 	The value at the playback position
-	 * @return true if anything was evaluated. Note: if false is returned InOutFloat remains unchanged
-	 */
+	DEPRECATED(4.15, "Direct evaluation of float tracks is no longer supported. Please create an evaluation template (see FMovieSceneFloatPropertySectionTemplate).")
 	virtual bool Eval( float Position, float LastPosition, float& InOutFloat ) const;
 };
diff --git a/Engine/Source/Runtime/MovieSceneTracks/Public/Tracks/MovieSceneIntegerTrack.h b/Engine/Source/Runtime/MovieSceneTracks/Public/Tracks/MovieSceneIntegerTrack.h
new file mode 100644
--- /dev/null
+++ b/Engine/Source/Runtime/MovieSceneTracks/Public/Tracks/MovieSceneIntegerTrack.h
@@ -0,0 +1,22 @@
+// Copyright 1998-2016 Epic Games, Inc. All Rights Reserved.
+
+#pragma once
+
+#include "MovieScene.h"
+#include "MovieSceneTrack.h"
+#include "MovieScenePropertyTrack.h"
+#include "MovieSceneIntegerTrack.generated.h"
+
+/**
+ * Handles manipulation of integer properties in a movie scene
+ */
+UCLASS()
+class MOVIESCENETRACKS_API UMovieSceneIntegerTrack : public UMovieScenePropertyTrack
+{
+	GENERATED_UCLASS_BODY()
+
+public:
+	/** UMovieSceneTrack interface */
+	virtual UMovieSceneSection* CreateNewSection() override;
+	virtual FMovieSceneEvalTemplatePtr CreateTemplateForSection(const UMovieSceneSection& InSection) const override;
+};
diff --git a/Engine/Source/Runtime/MovieSceneTracks/Public/Tracks/MovieSceneLevelVisibilityTrack.h b/Engine/Source/Runtime/MovieSceneTracks/Public/Tracks/MovieSceneLevelVisibilityTrack.h
--- a/Engine/Source/Runtime/MovieSceneTracks/Public/Tracks/MovieSceneLevelVisibilityTrack.h
+++ b/Engine/Source/Runtime/MovieSceneTracks/Public/Tracks/MovieSceneLevelVisibilityTrack.h
@@ -4,6 +4,7 @@
 
 #include "MovieScene.h"
 #include "MovieSceneNameableTrack.h"
+#include "MovieSceneSpawnTrack.h"
 #include "MovieScenePropertyTrack.h"
 #include "MovieSceneLevelVisibilityTrack.generated.h"
 
@@ -16,8 +17,11 @@ class UMovieSceneLevelVisibilityTrack : public UMovieSceneNameableTrack
 	GENERATED_UCLASS_BODY()
 
 public:
+
+	static uint16 GetEvaluationPriority() { return UMovieSceneSpawnTrack::GetEvaluationPriority() + 100; }
+
 	/** UMovieSceneTrack interface */
-	virtual TSharedPtr<IMovieSceneTrackInstance> CreateInstance() override;
+	virtual void PostCompile(FMovieSceneEvaluationTrack& Track, const FMovieSceneTrackCompilerArgs& Args) const override;
 	virtual bool IsEmpty() const override;
 	virtual void AddSection(UMovieSceneSection& Section) override;
 	virtual void RemoveSection( UMovieSceneSection& Section ) override;
diff --git a/Engine/Source/Runtime/MovieSceneTracks/Public/Tracks/MovieSceneMaterialTrack.h b/Engine/Source/Runtime/MovieSceneTracks/Public/Tracks/MovieSceneMaterialTrack.h
--- a/Engine/Source/Runtime/MovieSceneTracks/Public/Tracks/MovieSceneMaterialTrack.h
+++ b/Engine/Source/Runtime/MovieSceneTracks/Public/Tracks/MovieSceneMaterialTrack.h
@@ -51,14 +51,6 @@ public:
 	*/
 	void AddColorParameterKey(FName ParameterName, float Position, FLinearColor Value);
 
-	/**
-	 * Gets the animated values for this track.
-	 * @param Position The playback position to use for evaluation.
-	 * @param OutScalarValues An array of FScalarParameterNameAndValue objects representing each animated scalar parameter and it's animated value.
-	 * @param OutVectorValues An array of FVectorParameterNameAndValue objects representing each animated vector parameter and it's animated value.
-	 */
-	void Eval(float Position, TArray<FScalarParameterNameAndValue>& OutScalarValues, TArray<FColorParameterNameAndValue>& OutVectorValues) const;
-
 private:
 
 	/** The sections owned by this track .*/
@@ -80,7 +72,7 @@ public:
 
 	// UMovieSceneTrack interface
 
-	virtual TSharedPtr<IMovieSceneTrackInstance> CreateInstance() override;
+	virtual FMovieSceneEvalTemplatePtr CreateTemplateForSection(const UMovieSceneSection& InSection) const override;
 	virtual FName GetTrackName() const { return FName( *FString::FromInt(MaterialIndex) ); }
 
 #if WITH_EDITORONLY_DATA
diff --git a/Engine/Source/Runtime/MovieSceneTracks/Public/Tracks/MovieSceneParticleParameterTrack.h b/Engine/Source/Runtime/MovieSceneTracks/Public/Tracks/MovieSceneParticleParameterTrack.h
--- a/Engine/Source/Runtime/MovieSceneTracks/Public/Tracks/MovieSceneParticleParameterTrack.h
+++ b/Engine/Source/Runtime/MovieSceneTracks/Public/Tracks/MovieSceneParticleParameterTrack.h
@@ -20,7 +20,7 @@ public:
 
 	// UMovieSceneTrack interface
 
-	virtual TSharedPtr<IMovieSceneTrackInstance> CreateInstance() override;
+	virtual FMovieSceneEvalTemplatePtr CreateTemplateForSection(const UMovieSceneSection& InSection) const override;
 	virtual UMovieSceneSection* CreateNewSection() override;
 	virtual void RemoveAllAnimationData() override;
 	virtual bool HasSection(const UMovieSceneSection& Section) const override;
@@ -61,15 +61,6 @@ public:
 	*/
 	void MOVIESCENETRACKS_API AddColorParameterKey( FName ParameterName, float Position, FLinearColor Value );
 
-	/**
-	 * Gets the animated values for this track.
-	 * @param Position The playback position to use for evaluation.
-	 * @param OutScalarValues An array of FScalarParameterNameAndValue objects representing each animated scalar parameter and it's animated value.
-	 * @param OutVectorValues An array of FVectorParameterNameAndValue objects representing each animated vector parameter and it's animated value.
-	 * @param OutVectorValues An array of FVectorParameterNameAndValue objects representing each animated color parameter and it's animated value.
-	 */
-	void Eval( float Position, TArray<FScalarParameterNameAndValue>& OutScalarValues, TArray<FVectorParameterNameAndValue>& OutVectorValues, TArray<FColorParameterNameAndValue>& OutColorValues ) const;
-
 private:
 
 	/** The sections owned by this track .*/
diff --git a/Engine/Source/Runtime/MovieSceneTracks/Public/Tracks/MovieSceneParticleTrack.h b/Engine/Source/Runtime/MovieSceneTracks/Public/Tracks/MovieSceneParticleTrack.h
--- a/Engine/Source/Runtime/MovieSceneTracks/Public/Tracks/MovieSceneParticleTrack.h
+++ b/Engine/Source/Runtime/MovieSceneTracks/Public/Tracks/MovieSceneParticleTrack.h
@@ -34,7 +34,6 @@ public:
 
 	// UMovieSceneTrack interface
 
-	virtual TSharedPtr<IMovieSceneTrackInstance> CreateInstance() override;
 	virtual void RemoveAllAnimationData() override;
 	virtual bool HasSection(const UMovieSceneSection& Section) const override;
 	virtual void AddSection(UMovieSceneSection& Section) override;
diff --git a/Engine/Source/Runtime/MovieSceneTracks/Public/Tracks/MovieSceneSkeletalAnimationTrack.h b/Engine/Source/Runtime/MovieSceneTracks/Public/Tracks/MovieSceneSkeletalAnimationTrack.h
--- a/Engine/Source/Runtime/MovieSceneTracks/Public/Tracks/MovieSceneSkeletalAnimationTrack.h
+++ b/Engine/Source/Runtime/MovieSceneTracks/Public/Tracks/MovieSceneSkeletalAnimationTrack.h
@@ -9,7 +9,7 @@
 
 class IMovieSceneTrackInstance;
 class UMovieSceneSection;
-
+class UMovieSceneSkeletalAnimationSection;
 
 /**
  * Handles animation of skeletal mesh actors
@@ -32,7 +32,6 @@ public:
 
 	// UMovieSceneTrack interface
 
-	virtual TSharedPtr<IMovieSceneTrackInstance> CreateInstance() override;
 	virtual void RemoveAllAnimationData() override;
 	virtual bool HasSection(const UMovieSceneSection& Section) const override;
 	virtual void AddSection(UMovieSceneSection& Section) override;
@@ -52,4 +51,4 @@ private:
 	/** List of all animation sections */
 	UPROPERTY()
 	TArray<UMovieSceneSection*> AnimationSections;
-};
+};
\ No newline at end of file
diff --git a/Engine/Source/Runtime/MovieSceneTracks/Public/Tracks/MovieSceneSlomoTrack.h b/Engine/Source/Runtime/MovieSceneTracks/Public/Tracks/MovieSceneSlomoTrack.h
--- a/Engine/Source/Runtime/MovieSceneTracks/Public/Tracks/MovieSceneSlomoTrack.h
+++ b/Engine/Source/Runtime/MovieSceneTracks/Public/Tracks/MovieSceneSlomoTrack.h
@@ -20,11 +20,13 @@ class UMovieSceneSlomoTrack
 	GENERATED_BODY()
 
 public:
+	
+	UMovieSceneSlomoTrack(const FObjectInitializer& Init);
 
 	// UMovieSceneTrack interface
 
-	virtual TSharedPtr<IMovieSceneTrackInstance> CreateInstance() override;
 	virtual UMovieSceneSection* CreateNewSection() override;
+	virtual FMovieSceneEvalTemplatePtr CreateTemplateForSection(const UMovieSceneSection& InSection) const override;
 
 #if WITH_EDITORONLY_DATA
 	virtual FText GetDefaultDisplayName() const override;
diff --git a/Engine/Source/Runtime/MovieSceneTracks/Public/Tracks/MovieSceneSpawnTrack.h b/Engine/Source/Runtime/MovieSceneTracks/Public/Tracks/MovieSceneSpawnTrack.h
--- a/Engine/Source/Runtime/MovieSceneTracks/Public/Tracks/MovieSceneSpawnTrack.h
+++ b/Engine/Source/Runtime/MovieSceneTracks/Public/Tracks/MovieSceneSpawnTrack.h
@@ -3,10 +3,12 @@
 #pragma once
 
 #include "MovieSceneTrack.h"
+#include "Curves/IntegralCurve.h"
 #include "MovieSceneSpawnTrack.generated.h"
 
 
 class IMovieSceneTrackInstance;
+class UMovieSceneBoolSection;
 class UMovieSceneSection;
 
 
@@ -24,9 +26,6 @@ public:
 
 public:
 
-	/** Evaluate whether the controlled object should currently be spawned or not */
-	bool Eval(float Position, float LastPostion, bool& bOutSpawned) const;
-
 	/** Get the object identifier that this spawn track controls */
 	const FGuid& GetObjectId() const
 	{
@@ -39,18 +38,25 @@ public:
 		ObjectGuid = InGuid;
 	}
 
+	static uint16 GetEvaluationPriority() { return uint16(0xFFF); }
+	
 public:
 
 	// UMovieSceneTrack interface
-
 	virtual UMovieSceneSection* CreateNewSection() override;
-	virtual TSharedPtr<IMovieSceneTrackInstance> CreateInstance() override;
 	virtual bool HasSection(const UMovieSceneSection& Section) const override;
 	virtual void AddSection(UMovieSceneSection& Section) override;
 	virtual void RemoveSection(UMovieSceneSection& Section) override;
 	virtual bool IsEmpty() const override;
 	virtual TRange<float> GetSectionBoundaries() const override;
 	virtual const TArray<UMovieSceneSection*>& GetAllSections() const override;
+	virtual void GenerateTemplate(const FMovieSceneTrackCompilerArgs& Args) const override;
+	virtual void PostCompile(FMovieSceneEvaluationTrack& OutTrack, const FMovieSceneTrackCompilerArgs& Args) const override;
+	virtual FMovieSceneEvalTemplatePtr CreateTemplateForSection(const UMovieSceneSection& InSection) const override;
+
+	//~ UObject interface
+	virtual void PostLoad() override;
+
 #if WITH_EDITOR
 	virtual ECookOptimizationFlags GetCookOptimizationFlags() const override;
 #endif
diff --git a/Engine/Source/Runtime/MovieSceneTracks/Public/Tracks/MovieSceneStringTrack.h b/Engine/Source/Runtime/MovieSceneTracks/Public/Tracks/MovieSceneStringTrack.h
--- a/Engine/Source/Runtime/MovieSceneTracks/Public/Tracks/MovieSceneStringTrack.h
+++ b/Engine/Source/Runtime/MovieSceneTracks/Public/Tracks/MovieSceneStringTrack.h
@@ -41,11 +41,12 @@ public:
 	/**
 	 * Evaluates the track at the playback position
 	 *
-	 * @param Position The psootion at which to evaluate this track.
+	 * @param Position The position at which to evaluate this track.
 	 * @param LastPositionThe last playback position.
 	 * @param InOutString The string at the evaluation time.
 	 * @return true if anything was evaluated. Note: if false is returned InOutString remains unchanged
 	 */
+	DEPRECATED(4.15, "Direct evaluation of string tracks is no longer supported. Please create an evaluation template (see FMovieSceneStringPropertySectionTemplate).")
 	virtual bool Eval(float Position, float LastPostion, FString& InOutString) const;
 
 public:
@@ -53,8 +54,8 @@ public:
 	//~ UMovieSceneTrack interface
 
 	virtual void AddSection(UMovieSceneSection& Section) override;
-	virtual TSharedPtr<IMovieSceneTrackInstance> CreateInstance() override;
 	virtual UMovieSceneSection* CreateNewSection() override;
+	virtual FMovieSceneEvalTemplatePtr CreateTemplateForSection(const UMovieSceneSection& InSection) const override;
 	virtual const TArray<UMovieSceneSection*>& GetAllSections() const override;
 	virtual TRange<float> GetSectionBoundaries() const override;
 	virtual bool HasSection(const UMovieSceneSection& Section) const override;
diff --git a/Engine/Source/Runtime/MovieSceneTracks/Public/Tracks/MovieSceneSubTrack.h b/Engine/Source/Runtime/MovieSceneTracks/Public/Tracks/MovieSceneSubTrack.h
--- a/Engine/Source/Runtime/MovieSceneTracks/Public/Tracks/MovieSceneSubTrack.h
+++ b/Engine/Source/Runtime/MovieSceneTracks/Public/Tracks/MovieSceneSubTrack.h
@@ -10,6 +10,7 @@ class UMovieSceneSection;
 class UMovieSceneSequence;
 class UMovieSceneSubSection;
 
+struct FMovieSceneSegmentCompilerRules;
 
 /**
  * A track that holds sub-sequences within a larger sequence.
@@ -54,7 +55,6 @@ public:
 	// UMovieSceneTrack interface
 
 	virtual void AddSection(UMovieSceneSection& Section) override;
-	virtual TSharedPtr<IMovieSceneTrackInstance> CreateInstance() override;
 	virtual UMovieSceneSection* CreateNewSection() override;
 	virtual const TArray<UMovieSceneSection*>& GetAllSections() const override;
 	virtual TRange<float> GetSectionBoundaries() const override;
@@ -63,14 +63,18 @@ public:
 	virtual void RemoveAllAnimationData() override;
 	virtual void RemoveSection(UMovieSceneSection& Section) override;
 	virtual bool SupportsMultipleRows() const override;
+	virtual void GenerateTemplate(const FMovieSceneTrackCompilerArgs& Args) const override;
 
 #if WITH_EDITORONLY_DATA
 	virtual FText GetDefaultDisplayName() const override;
 #endif
-	
+
+protected:
+	virtual TInlineValue<FMovieSceneSegmentCompilerRules> GetTrackCompilerRules() const override;
+
 protected:
 
 	/** All movie scene sections. */
 	UPROPERTY()
 	TArray<UMovieSceneSection*> Sections;
-};
+};
\ No newline at end of file
diff --git a/Engine/Source/Runtime/MovieSceneTracks/Public/Tracks/MovieSceneVectorTrack.h b/Engine/Source/Runtime/MovieSceneTracks/Public/Tracks/MovieSceneVectorTrack.h
--- a/Engine/Source/Runtime/MovieSceneTracks/Public/Tracks/MovieSceneVectorTrack.h
+++ b/Engine/Source/Runtime/MovieSceneTracks/Public/Tracks/MovieSceneVectorTrack.h
@@ -20,17 +20,7 @@ class UMovieSceneVectorTrack : public UMovieScenePropertyTrack
 public:
 	/** UMovieSceneTrack interface */
 	virtual UMovieSceneSection* CreateNewSection() override;
-	virtual TSharedPtr<IMovieSceneTrackInstance> CreateInstance() override;
-
-	/**
-	 * Evaluates the track at the playback position
-	 *
-	 * @param Position	The current playback position
-	 * @param LastPosition	The last plackback position
-	 * @param OutVector 	The value at the playback position
-	 * @return true if anything was evaluated. Note: if false is returned OutVector remains unchanged
-	 */
-	virtual bool Eval( float Position, float LastPostion, FVector4& InOutVector ) const;
+	virtual FMovieSceneEvalTemplatePtr CreateTemplateForSection(const UMovieSceneSection& InSection) const override;
 
 	/** @return Get the number of channels used by the vector */
 	int32 GetNumChannelsUsed() const { return NumChannelsUsed; }
diff --git a/Engine/Source/Runtime/MovieSceneTracks/Public/Tracks/MovieSceneVisibilityTrack.h b/Engine/Source/Runtime/MovieSceneTracks/Public/Tracks/MovieSceneVisibilityTrack.h
--- a/Engine/Source/Runtime/MovieSceneTracks/Public/Tracks/MovieSceneVisibilityTrack.h
+++ b/Engine/Source/Runtime/MovieSceneTracks/Public/Tracks/MovieSceneVisibilityTrack.h
@@ -22,8 +22,7 @@ class UMovieSceneVisibilityTrack
 public:
 
 	// UMovieSceneTrack interface
-	virtual UMovieSceneSection* CreateNewSection() override;
-	virtual TSharedPtr<IMovieSceneTrackInstance> CreateInstance() override;
+	virtual FMovieSceneEvalTemplatePtr CreateTemplateForSection(const UMovieSceneSection& InSection) const override;
 
 #if WITH_EDITORONLY_DATA
 	virtual FText GetDisplayName() const override;
diff --git a/Engine/Source/Runtime/Niagara/Private/NiagaraSequence.cpp b/Engine/Source/Runtime/Niagara/Private/NiagaraSequence.cpp
--- a/Engine/Source/Runtime/Niagara/Private/NiagaraSequence.cpp
+++ b/Engine/Source/Runtime/Niagara/Private/NiagaraSequence.cpp
@@ -22,23 +22,11 @@ void UNiagaraSequence::BindPossessableObject(const FGuid& ObjectId, UObject& Pos
 }
 
 
-bool UNiagaraSequence::CanPossessObject(UObject& Object) const
+bool UNiagaraSequence::CanPossessObject(UObject& Object, UObject* InPlaybackContext) const
 {
 	return false;
 }
 
-
-UObject* UNiagaraSequence::FindPossessableObject(const FGuid& ObjectId, UObject* Context) const
-{
-	return nullptr;
-}
-
-
-FGuid UNiagaraSequence::FindPossessableObjectId(UObject& Object) const
-{
-	return FGuid();
-}
-
 UMovieScene* UNiagaraSequence::GetMovieScene() const
 {
 	return MovieScene;
diff --git a/Engine/Source/Runtime/Niagara/Public/NiagaraSequence.h b/Engine/Source/Runtime/Niagara/Public/NiagaraSequence.h
--- a/Engine/Source/Runtime/Niagara/Public/NiagaraSequence.h
+++ b/Engine/Source/Runtime/Niagara/Public/NiagaraSequence.h
@@ -20,9 +20,7 @@ public:
 	// UMovieSceneSequence overrides
 
 	virtual void BindPossessableObject(const FGuid& ObjectId, UObject& PossessedObject, UObject* Context) override;
-	virtual bool CanPossessObject(UObject& Object) const override;
-	virtual UObject* FindPossessableObject(const FGuid& ObjectId, UObject* Context) const override;
-	virtual FGuid FindPossessableObjectId(UObject& Object) const override;
+	virtual bool CanPossessObject(UObject& Object, UObject* InPlaybackContext) const override;
 	virtual UMovieScene* GetMovieScene() const override;
 	virtual UObject* GetParentObject(UObject* Object) const override;
 	virtual void UnbindPossessableObjects(const FGuid& ObjectId) override;
diff --git a/Engine/Source/Runtime/SessionMessages/Private/SessionMessagesModule.cpp b/Engine/Source/Runtime/SessionMessages/Private/SessionMessagesModule.cpp
--- a/Engine/Source/Runtime/SessionMessages/Private/SessionMessagesModule.cpp
+++ b/Engine/Source/Runtime/SessionMessages/Private/SessionMessagesModule.cpp
@@ -11,12 +11,12 @@ class FSessionMessagesModule
 {
 public:
 
-	// IModuleInterface interface
+	//~ IModuleInterface interface
 
-	virtual void StartupModule( ) override { }
-	virtual void ShutdownModule( ) override { }
+	virtual void StartupModule() override { }
+	virtual void ShutdownModule() override { }
 
-	virtual bool SupportsDynamicReloading( ) override
+	virtual bool SupportsDynamicReloading() override
 	{
 		return true;
 	}
diff --git a/Engine/Source/Runtime/SessionMessages/Private/SessionMessagesPrivatePCH.h b/Engine/Source/Runtime/SessionMessages/Private/SessionMessagesPrivatePCH.h
--- a/Engine/Source/Runtime/SessionMessages/Private/SessionMessagesPrivatePCH.h
+++ b/Engine/Source/Runtime/SessionMessages/Private/SessionMessagesPrivatePCH.h
@@ -2,10 +2,4 @@
 
 #pragma once
 
-#include "SessionMessages.h"
-
-
-/* Dependencies
- *****************************************************************************/
-
-#include "ModuleManager.h"
+#include "Runtime/Core/Public/Modules/ModuleManager.h"
diff --git a/Engine/Source/Runtime/SessionServices/Private/SessionInfo.cpp b/Engine/Source/Runtime/SessionServices/Private/SessionInfo.cpp
--- a/Engine/Source/Runtime/SessionServices/Private/SessionInfo.cpp
+++ b/Engine/Source/Runtime/SessionServices/Private/SessionInfo.cpp
@@ -1,12 +1,14 @@
 // Copyright 1998-2016 Epic Games, Inc. All Rights Reserved.
 
 #include "SessionServicesPrivatePCH.h"
+#include "SessionInfo.h"
+#include "SessionInstanceInfo.h"
 
 
 /* FSessionInfo structors
  *****************************************************************************/
 
-FSessionInfo::FSessionInfo( const FGuid& InSessionId, const IMessageBusRef& InMessageBus )
+FSessionInfo::FSessionInfo(const FGuid& InSessionId, const TSharedRef<IMessageBus, ESPMode::ThreadSafe>& InMessageBus)
 	: MessageBusPtr(InMessageBus)
 	, SessionId(InSessionId)
 { }
@@ -15,7 +17,7 @@ FSessionInfo::FSessionInfo( const FGuid& InSessionId, const IMessageBusRef& InMe
 /* FSessionInfo interface
  *****************************************************************************/
 
-void FSessionInfo::UpdateFromMessage( const FEngineServicePong& Message, const IMessageContextRef& Context )
+void FSessionInfo::UpdateFromMessage(const FEngineServicePong& Message, const TSharedRef<IMessageContext, ESPMode::ThreadSafe>& Context)
 {
 	if (Message.SessionId != SessionId)
 	{
@@ -41,7 +43,7 @@ void FSessionInfo::UpdateFromMessage( const FEngineServicePong& Message, const I
 }
 
 
-void FSessionInfo::UpdateFromMessage( const FSessionServicePong& Message, const IMessageContextRef& Context )
+void FSessionInfo::UpdateFromMessage(const FSessionServicePong& Message, const TSharedRef<IMessageContext, ESPMode::ThreadSafe>& Context)
 {
 	if (Message.SessionId != SessionId)
 	{
@@ -62,7 +64,7 @@ void FSessionInfo::UpdateFromMessage( const FSessionServicePong& Message, const
 	}
 	else
 	{
-		IMessageBusPtr MessageBus = MessageBusPtr.Pin();
+		auto MessageBus = MessageBusPtr.Pin();
 
 		if (MessageBus.IsValid())
 		{
@@ -81,7 +83,7 @@ void FSessionInfo::UpdateFromMessage( const FSessionServicePong& Message, const
 /* ISessionInfo interface
  *****************************************************************************/
 
-void FSessionInfo::GetInstances( TArray<ISessionInstanceInfoPtr>& OutInstances ) const
+void FSessionInfo::GetInstances(TArray<TSharedPtr<ISessionInstanceInfo>>& OutInstances) const
 {
 	OutInstances.Empty();
 
@@ -140,7 +142,7 @@ void FSessionInfo::Terminate()
 /* FSessionInfo callbacks
  *****************************************************************************/
 
-void FSessionInfo::HandleLogReceived( const ISessionInstanceInfoRef& Instance, const FSessionLogMessageRef& LogMessage )
+void FSessionInfo::HandleLogReceived(const TSharedRef<ISessionInstanceInfo>& Instance, const TSharedRef<FSessionLogMessage>& LogMessage)
 {
 	LogReceivedEvent.Broadcast(AsShared(), Instance, LogMessage);
 }
diff --git a/Engine/Source/Runtime/SessionServices/Private/SessionInfo.h b/Engine/Source/Runtime/SessionServices/Private/SessionInfo.h
--- a/Engine/Source/Runtime/SessionServices/Private/SessionInfo.h
+++ b/Engine/Source/Runtime/SessionServices/Private/SessionInfo.h
@@ -2,6 +2,17 @@
 
 #pragma once
 
+#include "SessionInfo.h"
+
+
+struct FEngineServicePong;
+struct FSessionLogMessage;
+class FSessionInstanceInfo;
+struct FSessionServicePong;
+class IMessageBus;
+class IMessageContext;
+class ISessionInstanceInfo;
+
 
 /**
  * Implements a class to maintain all info related to a game session
@@ -18,7 +29,7 @@ public:
 	 * @param InSessionId The session's identifier.
 	 * @param InMessageBus The message bus to use.
 	 */
-	FSessionInfo( const FGuid& InSessionId, const IMessageBusRef& InMessageBus );
+	FSessionInfo(const FGuid& InSessionId, const TSharedRef<IMessageBus, ESPMode::ThreadSafe>& InMessageBus);
 
 public:
 
@@ -28,7 +39,7 @@ public:
 	 * @param Message The message containing engine information.
 	 * @param Context The message context.
 	 */
-	void UpdateFromMessage( const FEngineServicePong& Message, const IMessageContextRef& Context );
+	void UpdateFromMessage(const FEngineServicePong& Message, const TSharedRef<IMessageContext, ESPMode::ThreadSafe>& Context);
 
 	/**
 	 * Updates this session info with the data in the specified message.
@@ -36,13 +47,13 @@ public:
 	 * @param Message The message containing session information.
 	 * @param Context The message context.
 	 */
-	void UpdateFromMessage( const FSessionServicePong& Message, const IMessageContextRef& Context );
+	void UpdateFromMessage(const FSessionServicePong& Message, const TSharedRef<IMessageContext, ESPMode::ThreadSafe>& Context);
 
 public:	
 
-	// ISessionInfo interface
+	//~ ISessionInfo interface
 
-	virtual void GetInstances( TArray<ISessionInstanceInfoPtr>& OutInstances ) const override;
+	virtual void GetInstances(TArray<TSharedPtr<ISessionInstanceInfo>>& OutInstances) const override;
 	virtual const FDateTime& GetLastUpdateTime() const override;
 	virtual const int32 GetNumInstances() const override;
 	virtual const FGuid& GetSessionId() const override;
@@ -67,7 +78,7 @@ public:
 private:
 
 	/** Handles received log messages. */
-	void HandleLogReceived( const ISessionInstanceInfoRef& Instance, const FSessionLogMessageRef& LogMessage );
+	void HandleLogReceived(const TSharedRef<ISessionInstanceInfo>& Instance, const TSharedRef<FSessionLogMessage>& LogMessage);
 
 private:
 
@@ -78,7 +89,7 @@ private:
 	FDateTime LastUpdateTime;
 
 	/** Holds a weak pointer to the message bus. */
-	IMessageBusWeakPtr MessageBusPtr;
+	TWeakPtr<IMessageBus, ESPMode::ThreadSafe> MessageBusPtr;
 
 	/** Holds the session identifier. */
 	FGuid SessionId;
diff --git a/Engine/Source/Runtime/SessionServices/Private/SessionInstanceInfo.cpp b/Engine/Source/Runtime/SessionServices/Private/SessionInstanceInfo.cpp
--- a/Engine/Source/Runtime/SessionServices/Private/SessionInstanceInfo.cpp
+++ b/Engine/Source/Runtime/SessionServices/Private/SessionInstanceInfo.cpp
@@ -1,12 +1,13 @@
 // Copyright 1998-2016 Epic Games, Inc. All Rights Reserved.
 
 #include "SessionServicesPrivatePCH.h"
+#include "SessionInstanceInfo.h"
 
 
 /* FSessionInstanceInfo structors
  *****************************************************************************/
 
-FSessionInstanceInfo::FSessionInstanceInfo(const FGuid& InInstanceId, const TSharedRef<ISessionInfo>& InOwner, const IMessageBusRef& InMessageBus)
+FSessionInstanceInfo::FSessionInstanceInfo(const FGuid& InInstanceId, const TSharedRef<ISessionInfo>& InOwner, const TSharedRef<IMessageBus, ESPMode::ThreadSafe>& InMessageBus)
 	: Authorized(false)
 	, EngineVersion(0)
 	, InstanceId(InInstanceId)
@@ -86,7 +87,7 @@ void FSessionInstanceInfo::Terminate()
 
 void FSessionInstanceInfo::HandleSessionLogMessage(const FSessionServiceLog& Message, const IMessageContextRef& Context)
 {
-	FSessionLogMessageRef LogMessage = MakeShareable(
+	TSharedRef<FSessionLogMessage> LogMessage = MakeShareable(
 		new FSessionLogMessage(
 			InstanceId,
 			InstanceName,
diff --git a/Engine/Source/Runtime/SessionServices/Private/SessionInstanceInfo.h b/Engine/Source/Runtime/SessionServices/Private/SessionInstanceInfo.h
--- a/Engine/Source/Runtime/SessionServices/Private/SessionInstanceInfo.h
+++ b/Engine/Source/Runtime/SessionServices/Private/SessionInstanceInfo.h
@@ -5,8 +5,11 @@
 #include "ISessionInstanceInfo.h"
 
 
+class IMessageBus;
 class ISessionInfo;
+struct FEngineServicePong;
 struct FSessionLogMessage;
+struct FSessionServicePong;
 
 
 /**
@@ -28,7 +31,7 @@ public:
 	 * @param InOwner The session that owns this instance.
 	 * @param InMessageBus The message bus to use.
 	 */
-	FSessionInstanceInfo(const FGuid& InInstanceId, const TSharedRef<ISessionInfo>& InOwner, const IMessageBusRef& InMessageBus);
+	FSessionInstanceInfo(const FGuid& InInstanceId, const TSharedRef<ISessionInfo>& InOwner, const TSharedRef<IMessageBus, ESPMode::ThreadSafe>& InMessageBus);
 
 public:
 
@@ -50,7 +53,7 @@ public:
 
 public:	
 
-	// IGameInstanceInfo interface
+	//~ IGameInstanceInfo interface
 
 	virtual void ExecuteCommand(const FString& CommandString) override;
 
diff --git a/Engine/Source/Runtime/SessionServices/Private/SessionManager.cpp b/Engine/Source/Runtime/SessionServices/Private/SessionManager.cpp
--- a/Engine/Source/Runtime/SessionServices/Private/SessionManager.cpp
+++ b/Engine/Source/Runtime/SessionServices/Private/SessionManager.cpp
@@ -1,12 +1,14 @@
 // Copyright 1998-2016 Epic Games, Inc. All Rights Reserved.
 
 #include "SessionServicesPrivatePCH.h"
+#include "SessionInfo.h"
+#include "SessionManager.h"
 
 
 /* FSessionManager structors
  *****************************************************************************/
 
-FSessionManager::FSessionManager(const IMessageBusRef& InMessageBus)
+FSessionManager::FSessionManager(const TSharedRef<IMessageBus, ESPMode::ThreadSafe>& InMessageBus)
 	: MessageBusPtr(InMessageBus)
 {
 	// fill in the owner array
@@ -55,19 +57,19 @@ void FSessionManager::AddOwner(const FString& InOwner)
 }
 
 
-const TArray<ISessionInstanceInfoPtr>& FSessionManager::GetSelectedInstances() const
+const TArray<TSharedPtr<ISessionInstanceInfo>>& FSessionManager::GetSelectedInstances() const
 {
 	return SelectedInstances;
 }
 
 
-const ISessionInfoPtr& FSessionManager::GetSelectedSession() const
+const TSharedPtr<ISessionInfo>& FSessionManager::GetSelectedSession() const
 {
 	return SelectedSession;
 }
 
 
-void FSessionManager::GetSessions(TArray<ISessionInfoPtr>& OutSessions) const
+void FSessionManager::GetSessions(TArray<TSharedPtr<ISessionInfo>>& OutSessions) const
 {
 	OutSessions.Empty(Sessions.Num());
 
@@ -103,7 +105,7 @@ void FSessionManager::RemoveOwner(const FString& InOwner)
 }
 
 
-bool FSessionManager::SelectSession(const ISessionInfoPtr& Session)
+bool FSessionManager::SelectSession(const TSharedPtr<ISessionInfo>& Session)
 {
 	// already selected?
 	if (Session == SelectedSession)
@@ -135,7 +137,7 @@ bool FSessionManager::SelectSession(const ISessionInfoPtr& Session)
 }
 
 
-bool FSessionManager::SetInstanceSelected(const ISessionInstanceInfoRef& Instance, bool Selected)
+bool FSessionManager::SetInstanceSelected(const TSharedRef<ISessionInstanceInfo>& Instance, bool Selected)
 {
 	if (Instance->GetOwnerSession() != SelectedSession)
 	{
@@ -257,7 +259,7 @@ void FSessionManager::HandleEnginePongMessage(const FEngineServicePong& Message,
 }
 
 
-void FSessionManager::HandleLogReceived(const ISessionInfoRef& Session, const ISessionInstanceInfoRef& Instance, const FSessionLogMessageRef& Message)
+void FSessionManager::HandleLogReceived(const TSharedRef<ISessionInfo>& Session, const TSharedRef<ISessionInstanceInfo>& Instance, const TSharedRef<FSessionLogMessage>& Message)
 {
 	if (Session == SelectedSession)
 	{
@@ -278,7 +280,7 @@ void FSessionManager::HandleSessionPongMessage(const FSessionServicePong& Messag
 		return;
 	}
 
-	IMessageBusPtr MessageBus = MessageBusPtr.Pin();
+	auto MessageBus = MessageBusPtr.Pin();
 
 	if (!MessageBus.IsValid())
 	{
diff --git a/Engine/Source/Runtime/SessionServices/Private/SessionManager.h b/Engine/Source/Runtime/SessionServices/Private/SessionManager.h
--- a/Engine/Source/Runtime/SessionServices/Private/SessionManager.h
+++ b/Engine/Source/Runtime/SessionServices/Private/SessionManager.h
@@ -5,6 +5,15 @@
 #include "ISessionManager.h"
 
 
+struct FEngineServicePong;
+class FSessionInfo;
+struct FSessionServicePong;
+class IMessageBus;
+class IMessageContext;
+class ISessionInfo;
+class ISessionInstanceInfo;
+
+
 /**
  * Implement the session manager
  */
@@ -19,18 +28,18 @@ public:
 	 *
 	 * @param InMessageBus The message bus to use.
 	 */
-	FSessionManager(const IMessageBusRef& InMessageBus);
+	FSessionManager(const TSharedRef<IMessageBus, ESPMode::ThreadSafe>& InMessageBus);
 	
 	/** Destructor */
 	~FSessionManager();
 
 public:
 
-	// ISessionManager interface
+	//~ ISessionManager interface
 
 	virtual void AddOwner(const FString& InOwner) override;
-	virtual const TArray<ISessionInstanceInfoPtr>& GetSelectedInstances() const override;
-	virtual const ISessionInfoPtr& GetSelectedSession() const override;
+	virtual const TArray<TSharedPtr<ISessionInstanceInfo>>& GetSelectedInstances() const override;
+	virtual const TSharedPtr<ISessionInfo>& GetSelectedSession() const override;
 	virtual void GetSessions(TArray<TSharedPtr<ISessionInfo>>& OutSessions) const override;
 	virtual bool IsInstanceSelected(const TSharedRef<ISessionInstanceInfo>& Instance) const override;
 
@@ -61,8 +70,8 @@ public:
 	virtual FSimpleMulticastDelegate& OnSessionsUpdated() override;
 	virtual FSimpleMulticastDelegate& OnSessionInstanceUpdated() override;
 	virtual void RemoveOwner(const FString& InOwner) override;
-	virtual bool SelectSession(const ISessionInfoPtr& Session) override;
-	virtual bool SetInstanceSelected(const ISessionInstanceInfoRef& Instance, bool Selected) override;
+	virtual bool SelectSession(const TSharedPtr<ISessionInfo>& Session) override;
+	virtual bool SetInstanceSelected(const TSharedRef<ISessionInstanceInfo>& Instance, bool Selected) override;
 
 protected:
 
@@ -89,13 +98,13 @@ protected:
 private:
 
 	/** Callback for handling FSessionServicePong messages. */
-	void HandleEnginePongMessage(const FEngineServicePong& Message, const IMessageContextRef& Context);
+	void HandleEnginePongMessage(const FEngineServicePong& Message, const TSharedRef<IMessageContext, ESPMode::ThreadSafe>& Context);
 
 	/** Callback received log entries. */
-	void HandleLogReceived(const ISessionInfoRef& Session, const ISessionInstanceInfoRef& Instance, const FSessionLogMessageRef& Message);
+	void HandleLogReceived(const TSharedRef<ISessionInfo>& Session, const TSharedRef<ISessionInstanceInfo>& Instance, const TSharedRef<FSessionLogMessage>& Message);
 
 	/** Callback for handling FSessionServicePong messages. */
-	void HandleSessionPongMessage(const FSessionServicePong& Message, const IMessageContextRef& Context);
+	void HandleSessionPongMessage(const FSessionServicePong& Message, const TSharedRef<IMessageContext, ESPMode::ThreadSafe>& Context);
 
 	/** Callback for ticks from the ticker. */
 	bool HandleTicker(float DeltaTime);
@@ -109,16 +118,16 @@ private:
 	FDateTime LastPingTime;
 
 	/** Holds a pointer to the message bus. */
-	IMessageBusWeakPtr MessageBusPtr;
+	TWeakPtr<IMessageBus, ESPMode::ThreadSafe> MessageBusPtr;
 
 	/** Holds the messaging endpoint. */
 	FMessageEndpointPtr MessageEndpoint;
 
 	/** Holds the list of currently selected instances. */
-	TArray<ISessionInstanceInfoPtr> SelectedInstances;
+	TArray<TSharedPtr<ISessionInstanceInfo>> SelectedInstances;
 
 	/** Holds a reference to the currently selected session. */
-	ISessionInfoPtr SelectedSession;
+	TSharedPtr<ISessionInfo> SelectedSession;
 
 	/** Holds the collection of discovered sessions. */
 	TMap<FGuid, TSharedPtr<FSessionInfo>> Sessions;
diff --git a/Engine/Source/Runtime/SessionServices/Private/SessionService.cpp b/Engine/Source/Runtime/SessionServices/Private/SessionService.cpp
--- a/Engine/Source/Runtime/SessionServices/Private/SessionService.cpp
+++ b/Engine/Source/Runtime/SessionServices/Private/SessionService.cpp
@@ -1,12 +1,13 @@
 // Copyright 1998-2016 Epic Games, Inc. All Rights Reserved.
 
 #include "SessionServicesPrivatePCH.h"
+#include "SessionService.h"
 
 
 /* FSessionService structors
  *****************************************************************************/
 
-FSessionService::FSessionService(const IMessageBusRef& InMessageBus)
+FSessionService::FSessionService(const TSharedRef<IMessageBus, ESPMode::ThreadSafe>& InMessageBus)
 	: MessageBusPtr(InMessageBus)
 { }
 
@@ -22,7 +23,7 @@ FSessionService::~FSessionService()
 
 bool FSessionService::Start()
 {
-	IMessageBusPtr MessageBus = MessageBusPtr.Pin();
+	auto MessageBus = MessageBusPtr.Pin();
 
 	if (!MessageBus.IsValid())
 	{
@@ -105,7 +106,7 @@ void FSessionService::SendNotification(const TCHAR* NotificationText, const FMes
 }
 
 
-void FSessionService::SendPong(const IMessageContextRef& Context, const FString& UserName)
+void FSessionService::SendPong(const TSharedRef<IMessageContext, ESPMode::ThreadSafe>& Context, const FString& UserName)
 {
 	if (!MessageEndpoint.IsValid())
 	{
@@ -146,21 +147,21 @@ void FSessionService::HandleMessageEndpointShutdown()
 }
 
 
-void FSessionService::HandleSessionLogSubscribeMessage(const FSessionServiceLogSubscribe& Message, const IMessageContextRef& Context)
+void FSessionService::HandleSessionLogSubscribeMessage(const FSessionServiceLogSubscribe& Message, const TSharedRef<IMessageContext, ESPMode::ThreadSafe>& Context)
 {
 	FScopeLock Lock(&LogSubscribersLock);
 	LogSubscribers.AddUnique(Context->GetSender());
 }
 
 
-void FSessionService::HandleSessionLogUnsubscribeMessage(const FSessionServiceLogUnsubscribe& Message, const IMessageContextRef& Context)
+void FSessionService::HandleSessionLogUnsubscribeMessage(const FSessionServiceLogUnsubscribe& Message, const TSharedRef<IMessageContext, ESPMode::ThreadSafe>& Context)
 {
 	FScopeLock Lock(&LogSubscribersLock);
 	LogSubscribers.Remove(Context->GetSender());
 }
 
 
-void FSessionService::HandleSessionPingMessage(const FSessionServicePing& Message, const IMessageContextRef& Context)
+void FSessionService::HandleSessionPingMessage(const FSessionServicePing& Message, const TSharedRef<IMessageContext, ESPMode::ThreadSafe>& Context)
 {
 	SendPong(Context, Message.UserName);
 }
diff --git a/Engine/Source/Runtime/SessionServices/Private/SessionService.h b/Engine/Source/Runtime/SessionServices/Private/SessionService.h
--- a/Engine/Source/Runtime/SessionServices/Private/SessionService.h
+++ b/Engine/Source/Runtime/SessionServices/Private/SessionService.h
@@ -5,6 +5,15 @@
 #include "ISessionService.h"
 
 
+struct FMessageAddress;
+class FMessageEndpoint;
+struct FSessionServiceLogSubscribe;
+struct FSessionServiceLogUnsubscribe;
+struct FSessionServicePing;
+class IMessageBus;
+class IMessageContext;
+
+
 /**
  * Implements an application session service.
  */
@@ -19,14 +28,14 @@ public:
 	 *
 	 * @param InMessageBus The message bus to use.
 	 */
-	FSessionService(const IMessageBusRef& InMessageBus);
+	FSessionService(const TSharedRef<IMessageBus, ESPMode::ThreadSafe>& InMessageBus);
 
 	/** Virtual destructor. */
 	virtual ~FSessionService();
 
 public:
 
-	// FOutputDevice interface
+	//~ FOutputDevice interface
 
 	virtual void Serialize(const TCHAR* Data, ELogVerbosity::Type Verbosity, const class FName& Category) override
 	{
@@ -35,7 +44,7 @@ public:
 
 public:
 
-	// ISessionService interface
+	//~ ISessionService interface
 
 	virtual bool IsRunning() override
 	{
@@ -70,7 +79,7 @@ protected:
 	 * @param Context The context of the received Ping message.
 	 * @param UserName The name of the user that sent the ping.
 	 */
-	void SendPong(const IMessageContextRef& Context, const FString& UserName);
+	void SendPong(const TSharedRef<IMessageContext, ESPMode::ThreadSafe>& Context, const FString& UserName);
 
 private:
 
@@ -78,13 +87,13 @@ private:
 	void HandleMessageEndpointShutdown();
 
 	/** Handles FSessionServiceLogSubscribe messages. */
-	void HandleSessionLogSubscribeMessage(const FSessionServiceLogSubscribe& Message, const IMessageContextRef& Context);
+	void HandleSessionLogSubscribeMessage(const FSessionServiceLogSubscribe& Message, const TSharedRef<IMessageContext, ESPMode::ThreadSafe>& Context);
 
 	/** Handles FSessionServiceLogUnsubscribe messages. */
-	void HandleSessionLogUnsubscribeMessage(const FSessionServiceLogUnsubscribe& Message, const IMessageContextRef& Context);
+	void HandleSessionLogUnsubscribeMessage(const FSessionServiceLogUnsubscribe& Message, const TSharedRef<IMessageContext, ESPMode::ThreadSafe>& Context);
 
 	/** Handles FSessionServicePing messages. */
-	void HandleSessionPingMessage(const FSessionServicePing& Message, const IMessageContextRef& Context);
+	void HandleSessionPingMessage(const FSessionServicePing& Message, const TSharedRef<IMessageContext, ESPMode::ThreadSafe>& Context);
 
 private:
 
@@ -95,8 +104,8 @@ private:
 	FCriticalSection LogSubscribersLock;
 
 	/** Holds a weak pointer to the message bus. */
-	IMessageBusWeakPtr MessageBusPtr;
+	TWeakPtr<IMessageBus, ESPMode::ThreadSafe> MessageBusPtr;
 
 	/** Holds the message endpoint. */
-	FMessageEndpointPtr MessageEndpoint;
+	TSharedPtr<FMessageEndpoint, ESPMode::ThreadSafe> MessageEndpoint;
 };
diff --git a/Engine/Source/Runtime/SessionServices/Private/SessionServicesModule.cpp b/Engine/Source/Runtime/SessionServices/Private/SessionServicesModule.cpp
--- a/Engine/Source/Runtime/SessionServices/Private/SessionServicesModule.cpp
+++ b/Engine/Source/Runtime/SessionServices/Private/SessionServicesModule.cpp
@@ -2,7 +2,8 @@
 
 #include "SessionServicesPrivatePCH.h"
 #include "ISessionServicesModule.h"
-#include "ModuleManager.h"
+#include "SessionManager.h"
+#include "SessionService.h"
 
 
 /**
@@ -24,7 +25,7 @@ public:
 
 public:
 
-	// FSelfRegisteringExec interface
+	//~ FSelfRegisteringExec interface
 
 	virtual bool Exec(UWorld* InWorld, const TCHAR* Cmd, FOutputDevice& Ar) override
 	{
@@ -81,9 +82,9 @@ public:
 
 public:
 
-	// ISessionServicesModule interface
+	//~ ISessionServicesModule interface
 
-	virtual ISessionManagerRef GetSessionManager() override
+	virtual TSharedRef<ISessionManager> GetSessionManager() override
 	{
 		if (!SessionManager.IsValid())
 		{
@@ -93,7 +94,7 @@ public:
 		return SessionManager.ToSharedRef();
 	}
 
-	virtual ISessionServiceRef GetSessionService() override
+	virtual TSharedRef<ISessionService> GetSessionService() override
 	{
 		if (!SessionService.IsValid())
 		{
@@ -105,7 +106,7 @@ public:
 
 public:
 
-	// IModuleInterface interface
+	//~ IModuleInterface interface
 
 	virtual void StartupModule() override
 	{
@@ -122,13 +123,13 @@ public:
 private:
 	
 	/** Holds a weak pointer to the message bus. */
-	IMessageBusWeakPtr MessageBusPtr;
+	TWeakPtr<IMessageBus, ESPMode::ThreadSafe> MessageBusPtr;
 
 	/** Holds the session manager singleton. */
-	ISessionManagerPtr SessionManager;
+	TSharedPtr<ISessionManager> SessionManager;
 
 	/** Holds the session service singleton. */
-	ISessionServicePtr SessionService;
+	TSharedPtr<ISessionService> SessionService;
 };
 
 
diff --git a/Engine/Source/Runtime/SessionServices/Private/SessionServicesPrivatePCH.h b/Engine/Source/Runtime/SessionServices/Private/SessionServicesPrivatePCH.h
--- a/Engine/Source/Runtime/SessionServices/Private/SessionServicesPrivatePCH.h
+++ b/Engine/Source/Runtime/SessionServices/Private/SessionServicesPrivatePCH.h
@@ -2,25 +2,21 @@
 
 #pragma once
 
-
-/* Dependencies
- *****************************************************************************/
-
-#include "CoreUObject.h"
-#include "EngineServiceMessages.h"
-#include "Messaging.h"
-#include "SessionMessages.h"
-#include "Ticker.h"
-
-
-/* Private includes
- *****************************************************************************/
+#include "Runtime/Core/Public/Core.h"
+#include "Runtime/Core/Public/Containers/Ticker.h"
+#include "Runtime/Core/Public/Modules/ModuleManager.h"
+#include "Runtime/CoreUObject/Public/CoreUObject.h"
+#include "Runtime/EngineMessages/Public/EngineServiceMessages.h"
+#include "Runtime/Messaging/Public/IMessageBus.h"
+#include "Runtime/Messaging/Public/IMessageContext.h"
+#include "Runtime/Messaging/Public/IMessagingModule.h"
+#include "Runtime/Messaging/Public/Helpers/MessageEndpoint.h"
+#include "Runtime/Messaging/Public/Helpers/MessageEndpointBuilder.h"
+#include "Runtime/SessionMessages/Public/SessionServiceMessages.h"
 
 #include "ISessionInfo.h"
 #include "ISessionInstanceInfo.h"
+#include "ISessionManager.h"
+#include "ISessionService.h"
+#include "ISessionServicesModule.h"
 #include "SessionLogMessage.h"
-
-#include "SessionInstanceInfo.h"
-#include "SessionInfo.h"
-#include "SessionManager.h"
-#include "SessionService.h"
diff --git a/Engine/Source/Runtime/SessionServices/Public/ISessionInfo.h b/Engine/Source/Runtime/SessionServices/Public/ISessionInfo.h
--- a/Engine/Source/Runtime/SessionServices/Public/ISessionInfo.h
+++ b/Engine/Source/Runtime/SessionServices/Public/ISessionInfo.h
@@ -19,7 +19,7 @@ public:
 	 *
 	 * @param OutInstances Will hold the collection of instances.
 	 */
-	virtual void GetInstances( TArray<TSharedPtr<ISessionInstanceInfo>>& OutInstances ) const = 0;
+	virtual void GetInstances(TArray<TSharedPtr<ISessionInstanceInfo>>& OutInstances) const = 0;
 
 	/**
 	 * Gets the time at which the last update was received from this instance.
@@ -93,10 +93,3 @@ public:
 	/** Virtual destructor. */
 	virtual ~ISessionInfo() { }
 };
-
-
-/** Type definition for shared pointers to instances of ISessionInfo. */
-typedef TSharedPtr<ISessionInfo> ISessionInfoPtr;
-
-/** Type definition for shared references to instances of ISessionInfo. */
-typedef TSharedRef<ISessionInfo> ISessionInfoRef;
diff --git a/Engine/Source/Runtime/SessionServices/Public/ISessionInstanceInfo.h b/Engine/Source/Runtime/SessionServices/Public/ISessionInstanceInfo.h
--- a/Engine/Source/Runtime/SessionServices/Public/ISessionInstanceInfo.h
+++ b/Engine/Source/Runtime/SessionServices/Public/ISessionInstanceInfo.h
@@ -108,7 +108,7 @@ public:
 	/**
 	 * Check whether the current user is authorized to interact with this instance.
 	 *
-	 * @return true if the uesr is authorized, false otherwise.
+	 * @return true if the user is authorized, false otherwise.
 	 */
 	virtual bool IsAuthorized() const = 0;
 
@@ -144,10 +144,3 @@ public:
 	/** Virtual destructor. */
 	virtual ~ISessionInstanceInfo() { }
 };
-
-
-/** Type definition for shared pointers to instances of IGameInstanceInfo. */
-typedef TSharedPtr<ISessionInstanceInfo> ISessionInstanceInfoPtr;
-
-/** Type definition for shared references to instances of IGameInstanceInfo. */
-typedef TSharedRef<ISessionInstanceInfo> ISessionInstanceInfoRef;
diff --git a/Engine/Source/Runtime/SessionServices/Public/ISessionManager.h b/Engine/Source/Runtime/SessionServices/Public/ISessionManager.h
--- a/Engine/Source/Runtime/SessionServices/Public/ISessionManager.h
+++ b/Engine/Source/Runtime/SessionServices/Public/ISessionManager.h
@@ -127,10 +127,3 @@ public:
 	/** Virtual destructor. */
 	virtual ~ISessionManager() { }
 };
-
-
-/** Type definition for shared pointers to instances of ISessionManager. */
-typedef TSharedPtr<ISessionManager> ISessionManagerPtr;
-
-/** Type definition for shared references to instances of ISessionManager. */
-typedef TSharedRef<ISessionManager> ISessionManagerRef;
diff --git a/Engine/Source/Runtime/SessionServices/Public/ISessionService.h b/Engine/Source/Runtime/SessionServices/Public/ISessionService.h
--- a/Engine/Source/Runtime/SessionServices/Public/ISessionService.h
+++ b/Engine/Source/Runtime/SessionServices/Public/ISessionService.h
@@ -38,10 +38,3 @@ public:
 	/** Virtual destructor. */
 	virtual ~ISessionService() { }
 };
-
-
-/** Type definition for shared pointers to instances of ISessionService. */
-typedef TSharedPtr<ISessionService> ISessionServicePtr;
-
-/** Type definition for shared references to instances of ISessionService. */
-typedef TSharedRef<ISessionService> ISessionServiceRef;
diff --git a/Engine/Source/Runtime/SessionServices/Public/ISessionServicesModule.h b/Engine/Source/Runtime/SessionServices/Public/ISessionServicesModule.h
--- a/Engine/Source/Runtime/SessionServices/Public/ISessionServicesModule.h
+++ b/Engine/Source/Runtime/SessionServices/Public/ISessionServicesModule.h
@@ -3,7 +3,6 @@
 #pragma once
 
 
-// forward declarations
 class ISessionManager;
 class ISessionService;
 
@@ -20,6 +19,7 @@ public:
 	 * Gets the session manager.
 	 *
 	 * @return The session manager.
+	 * @see GetSessionService
 	 */
 	virtual TSharedRef<ISessionManager> GetSessionManager() = 0;
 
@@ -27,6 +27,7 @@ public:
 	 * Gets the session service.
 	 *
 	 * @return The session service.
+	 * @see GetSessionManager
 	 */
 	virtual TSharedRef<ISessionService> GetSessionService() = 0;
 
diff --git a/Engine/Source/Runtime/SessionServices/Public/SessionLogMessage.h b/Engine/Source/Runtime/SessionServices/Public/SessionLogMessage.h
--- a/Engine/Source/Runtime/SessionServices/Public/SessionLogMessage.h
+++ b/Engine/Source/Runtime/SessionServices/Public/SessionLogMessage.h
@@ -64,10 +64,3 @@ public:
 		}
 	};
 };
-
-
-/** Type definition for shared pointers to instances of FLauncherLogMessage. */
-typedef TSharedPtr<FSessionLogMessage> FSessionLogMessagePtr;
-
-/** Type definition for shared references to instances of FLauncherLogMessage. */
-typedef TSharedRef<FSessionLogMessage> FSessionLogMessageRef;
diff --git a/Engine/Source/Runtime/Slate/Private/Widgets/Text/SInlineEditableTextBlock.cpp b/Engine/Source/Runtime/Slate/Private/Widgets/Text/SInlineEditableTextBlock.cpp
--- a/Engine/Source/Runtime/Slate/Private/Widgets/Text/SInlineEditableTextBlock.cpp
+++ b/Engine/Source/Runtime/Slate/Private/Widgets/Text/SInlineEditableTextBlock.cpp
@@ -141,6 +141,11 @@ bool SInlineEditableTextBlock::IsInEditMode() const
 	return TextBlock->GetVisibility() == EVisibility::Collapsed;
 }
 
+void SInlineEditableTextBlock::SetReadOnly(bool bInIsReadOnly)
+{
+	bIsReadOnly = bInIsReadOnly;
+}
+
 void SInlineEditableTextBlock::SetText( const TAttribute< FText >& InText )
 {
 	Text = InText;
diff --git a/Engine/Source/Runtime/Slate/Public/Widgets/Text/SInlineEditableTextBlock.h b/Engine/Source/Runtime/Slate/Public/Widgets/Text/SInlineEditableTextBlock.h
--- a/Engine/Source/Runtime/Slate/Public/Widgets/Text/SInlineEditableTextBlock.h
+++ b/Engine/Source/Runtime/Slate/Public/Widgets/Text/SInlineEditableTextBlock.h
@@ -115,6 +115,8 @@ class SLATE_API SInlineEditableTextBlock: public SCompoundWidget
 	/** Checks if the widget is in edit mode */
 	bool IsInEditMode() const;
 
+	void SetReadOnly(bool bInIsReadOnly);
+
 	void SetText( const TAttribute< FText >& InText );
 	void SetText( const FString& InText );
 
