Copying //UE4/Dev-Sequencer to //UE4/Dev-Main (Source: //UE4/Dev-Sequencer @ 3237992)

diff --git a/Engine/Source/Runtime/MovieSceneTracks/Private/Evaluation/MovieSceneSkeletalAnimationTemplate.cpp b/Engine/Source/Runtime/MovieSceneTracks/Private/Evaluation/MovieSceneSkeletalAnimationTemplate.cpp
--- a/Engine/Source/Runtime/MovieSceneTracks/Private/Evaluation/MovieSceneSkeletalAnimationTemplate.cpp
+++ b/Engine/Source/Runtime/MovieSceneTracks/Private/Evaluation/MovieSceneSkeletalAnimationTemplate.cpp
@@ -4,8 +4,12 @@
 
 #include "Compilation/MovieSceneCompilerRules.h"
 #include "Animation/AnimMontage.h"
+#include "Animation/AnimSingleNodeInstance.h"
 #include "Components/SkeletalMeshComponent.h"
 #include "Animation/AnimSingleNodeInstance.h"
+#include "Runtime/AnimGraphRuntime/Public/AnimSequencerInstance.h"
+#include "MovieSceneEvaluation.h"
+#include "IMovieScenePlayer.h"
 
 
 bool ShouldUsePreviewPlayback(IMovieScenePlayer& Player, UObject& RuntimeObject)
@@ -21,150 +25,143 @@ bool CanPlayAnimation(USkeletalMeshComponent* SkeletalMeshComponent, UAnimSequen
 		(!AnimAssetBase || SkeletalMeshComponent->SkeletalMesh->Skeleton->IsCompatible(AnimAssetBase->GetSkeleton())));
 }
 
-struct FAnimationTrackTrackState : IPersistentEvaluationData
+void FinishAnimControl(USkeletalMeshComponent* SkeletalMeshComponent)
 {
-	struct FComponentAndMontage
-	{
-		FComponentAndMontage(USkeletalMeshComponent* InComponent) : Component(InComponent) {}
-
-		TWeakObjectPtr<UAnimMontage> Montage;
-		TWeakObjectPtr<USkeletalMeshComponent> Component;
-	};
-
-	TMap<FObjectKey, FComponentAndMontage> CurrentlyPlayingMontages;
-
-	~FAnimationTrackTrackState()
+	if (SkeletalMeshComponent->GetAnimationMode() == EAnimationMode::Type::AnimationCustomMode)
 	{
-		for (auto& Pair : CurrentlyPlayingMontages)
+		UAnimSequencerInstance* SequencerInstance = Cast<UAnimSequencerInstance>(SkeletalMeshComponent->GetAnimInstance());
+		if (SequencerInstance)
 		{
-			USkeletalMeshComponent* Component = Pair.Value.Component.Get();
-
-			if (Component)
-			{
-				FinishAnimControl(Component);
-			}
+			SkeletalMeshComponent->AnimScriptInstance = nullptr;
 		}
 	}
-
-	void EnsureAnimControl(USkeletalMeshComponent* SkeletalMeshComponent, UAnimSequenceBase* InAnimSequence, IMovieScenePlayer& Player)
+	else if (SkeletalMeshComponent->GetAnimationMode() == EAnimationMode::Type::AnimationBlueprint)
 	{
-		SkeletalMeshComponent->MeshComponentUpdateFlag = EMeshComponentUpdateFlag::AlwaysTickPoseAndRefreshBones;
-
-		if (!CurrentlyPlayingMontages.Contains(SkeletalMeshComponent))
-		{
-			BeginAnimControl(SkeletalMeshComponent, InAnimSequence);
-		}
-	}
-
-private:
-
-	void BeginAnimControl(USkeletalMeshComponent* SkeletalMeshComponent, UAnimSequenceBase* InAnimSequence)
-	{
-		if (CanPlayAnimation(SkeletalMeshComponent, InAnimSequence))
+		UAnimInstance* AnimInstance = SkeletalMeshComponent->GetAnimInstance();
+		if (AnimInstance)
 		{
-			UAnimInstance* AnimInstance = SkeletalMeshComponent->GetAnimInstance();
-			if (!AnimInstance)
-			{
-				SkeletalMeshComponent->SetAnimationMode(EAnimationMode::Type::AnimationSingleNode);
-			}
+			AnimInstance->Montage_Stop(0.f);
+			AnimInstance->UpdateAnimation(0.f, false);
 		}
 
-		CurrentlyPlayingMontages.Add(SkeletalMeshComponent, FComponentAndMontage(SkeletalMeshComponent));
+		// Update space bases to reset it back to ref pose
+		SkeletalMeshComponent->RefreshBoneTransforms();
+		SkeletalMeshComponent->RefreshSlaveComponents();
+		SkeletalMeshComponent->UpdateComponentToWorld();
 	}
+}
 
-	void FinishAnimControl(USkeletalMeshComponent* SkeletalMeshComponent)
-	{
-		if (SkeletalMeshComponent->GetAnimationMode() == EAnimationMode::Type::AnimationBlueprint)
-		{
-			UAnimInstance* AnimInstance = SkeletalMeshComponent->GetAnimInstance();
-			if(AnimInstance)
-			{
-				AnimInstance->Montage_Stop(0.f);
-				AnimInstance->UpdateAnimation(0.f, false);
-			}
-
-			// Update space bases to reset it back to ref pose
-			SkeletalMeshComponent->RefreshBoneTransforms();
-			SkeletalMeshComponent->RefreshSlaveComponents();
-			SkeletalMeshComponent->UpdateComponentToWorld();
-		}
-	}
+struct FMinimalAnimParameters
+{
+	FMinimalAnimParameters(UAnimSequenceBase* InAnimation, float InEvalTime, float InBlendWeight, uint32 InSectionId, FName InSlotName)
+		: Animation(InAnimation)
+		, EvalTime(InEvalTime)
+		, BlendWeight(InBlendWeight)
+		, SectionId(InSectionId)
+		, SlotName(InSlotName)
+	{}
+	
+	UAnimSequenceBase* Animation;
+	float EvalTime;
+	float BlendWeight;
+	uint32 SectionId;
+	FName SlotName;
 };
 
 struct FPreAnimatedAnimationTokenProducer : IMovieScenePreAnimatedTokenProducer
 {
 	virtual IMovieScenePreAnimatedTokenPtr CacheExistingState(UObject& Object) const
 	{
-		struct FPreAnimatedToken : IMovieScenePreAnimatedToken
+		struct FToken : IMovieScenePreAnimatedToken
 		{
-			EMeshComponentUpdateFlag::Type MeshComponentUpdateFlag;
+			FToken(USkeletalMeshComponent* InComponent)
+			{
+				// Cache this object's current update flag and animation mode
+				MeshComponentUpdateFlag = InComponent->MeshComponentUpdateFlag;
+				AnimationMode = InComponent->GetAnimationMode();
+			}
 
-			virtual void RestoreState(UObject& InObject, IMovieScenePlayer& Player)
+			virtual void RestoreState(UObject& ObjectToRestore, IMovieScenePlayer& Player)
 			{
-				CastChecked<USkeletalMeshComponent>(&InObject)->MeshComponentUpdateFlag = MeshComponentUpdateFlag;
+				USkeletalMeshComponent* Component = CastChecked<USkeletalMeshComponent>(&ObjectToRestore);
+
+				UAnimSequencerInstance* SequencerInst = Cast<UAnimSequencerInstance>(Component->GetAnimInstance());
+				if (SequencerInst)
+				{
+					SequencerInst->ResetNodes();
+				}
+
+				FinishAnimControl(Component);
+
+				// Reset the mesh component update flag and animation mode to what they were before we animated the object
+				Component->MeshComponentUpdateFlag = MeshComponentUpdateFlag;
+				Component->SetAnimationMode(AnimationMode);
+
 			}
+
+			EMeshComponentUpdateFlag::Type MeshComponentUpdateFlag;
+			EAnimationMode::Type AnimationMode;
 		};
 
-		FPreAnimatedToken Token;
-		Token.MeshComponentUpdateFlag = CastChecked<USkeletalMeshComponent>(&Object)->MeshComponentUpdateFlag;
-		return Token;
+		return FToken(CastChecked<USkeletalMeshComponent>(&Object));
 	}
 };
 
-struct FAnimationTrackTrackToken : IMovieSceneExecutionToken
+struct FSkeletalAnimationTrackData : IPersistentEvaluationData
 {
-	FAnimationTrackTrackState* TrackState;
-
-	FAnimationTrackTrackToken(const FMovieSceneSkeletalAnimationSectionTemplateParameters& InParams)
-		: Params(InParams)
+	struct FComponentDataContainer
 	{
-	}
+		FComponentDataContainer(USkeletalMeshComponent* InComponent) : Component(InComponent) {}
 
-	virtual void Execute(const FMovieSceneContext& Context, const FMovieSceneEvaluationOperand& Operand, FPersistentEvaluationData& PersistentData, IMovieScenePlayer& Player) override
-	{
-		TrackState = &PersistentData.GetOrAddTrackData<FAnimationTrackTrackState>();
+		// per slot name, montage reference if the component is using this
+		TMap<FName, TWeakObjectPtr<class UAnimMontage>> CurrentlyPlayingMontages;
 
-		for (TWeakObjectPtr<> Object : Player.FindBoundObjects(Operand))
-		{
-			UObject* ObjectPtr = Object.Get();
-			if (!ObjectPtr)
-			{
-				continue;
-			}
+		// component reference
+		TWeakObjectPtr<USkeletalMeshComponent> Component;
 
-			// first try to control the component directly
-			USkeletalMeshComponent* SkeletalMeshComponent = Cast<USkeletalMeshComponent>(ObjectPtr);
-			if (!SkeletalMeshComponent)
-			{
-				// then check to see if we are controlling an actor & if so use its first USkeletalMeshComponent 
-				if (AActor* Actor = Cast<AActor>(ObjectPtr))
-				{
-					SkeletalMeshComponent = Actor->FindComponentByClass<USkeletalMeshComponent>();
-				}
-			}
+		// parameters to apply per section
+		TArray<FMinimalAnimParameters, TInlineAllocator<1>> SectionParams;
+	};
 
-			UAnimSequenceBase* AnimSequence = Params.Animation;
+	TMap<FObjectKey, FComponentDataContainer> ComponentDataContainer;
 
-			if (!SkeletalMeshComponent || !AnimSequence)
-			{
-				continue;
-			}
-
-			const float EvalTime = Params.MapTimeToAnimation(Context.GetTime());
+	// Store the section parameters and map them to the skeletal component
+	void AssignSkeletalAnimationData(USkeletalMeshComponent* SkeletalMeshComponent, FMinimalAnimParameters AnimParameters)
+	{
+		check(SkeletalMeshComponent);
 
-			Player.SavePreAnimatedState(*SkeletalMeshComponent, TMovieSceneAnimTypeID<FAnimationTrackTrackToken>(), FPreAnimatedAnimationTokenProducer());
+		FObjectKey ObjectKey(SkeletalMeshComponent);
 
-			TrackState->EnsureAnimControl(SkeletalMeshComponent, AnimSequence, Player);
+		FComponentDataContainer* ComponentData = ComponentDataContainer.Find(ObjectKey);
+		if (!ComponentData)
+		{
+			// add to cache data container
+			ComponentData = &ComponentDataContainer.Add(ObjectKey, FComponentDataContainer(SkeletalMeshComponent));
+		}
 
-			static const int32 ChannelIndex = 0;
-			static const bool bLooping = false;
+		ComponentData->SectionParams.Add(AnimParameters);
+	}
 
-			if (ShouldUsePreviewPlayback(Player, *SkeletalMeshComponent))
+	// Apply the skeletal animation 
+	void ApplySkeletalAnimationData(const FMovieSceneContext& Context, const FMovieSceneEvaluationOperand& Operand, IMovieScenePlayer& Player)
+	{
+		for (auto& Pair : ComponentDataContainer)
+		{
+			const FComponentDataContainer& DataContainer = Pair.Value;
+			USkeletalMeshComponent* SkeletalMeshComponent = DataContainer.Component.Get();
+			
+			if (SkeletalMeshComponent)
 			{
+				BeginAnimControl(SkeletalMeshComponent);
+
+				static const bool bLooping = false;
+
+				bool bPreviewPlayback = ShouldUsePreviewPlayback(Player, *SkeletalMeshComponent);
+
 				const EMovieScenePlayerStatus::Type PlayerStatus = Player.GetPlaybackStatus();
 
 				// If the playback status is jumping, ie. one such occurrence is setting the time for thumbnail generation, disable anim notifies updates because it could fire audio
-				const bool bFireNotifies = PlayerStatus != EMovieScenePlayerStatus::Jumping && PlayerStatus != EMovieScenePlayerStatus::Stopped;
+				const bool bFireNotifies = !bPreviewPlayback || (PlayerStatus != EMovieScenePlayerStatus::Jumping && PlayerStatus != EMovieScenePlayerStatus::Stopped);
 
 				// When jumping from one cut to another cut, the delta time should be 0 so that anim notifies before the current position are not evaluated. Note, anim notifies at the current time should still be evaluated.
 				const float DeltaTime = Context.HasJumped() ? 0.f : Context.GetRange().Size<float>();
@@ -174,87 +171,269 @@ struct FAnimationTrackTrackToken : IMovieSceneExecutionToken
 											PlayerStatus == EMovieScenePlayerStatus::Scrubbing || 
 											(DeltaTime == 0.0f && PlayerStatus != EMovieScenePlayerStatus::Stopped); 
 			
-				PreviewSetAnimPosition(SkeletalMeshComponent, Params.SlotName, ChannelIndex, AnimSequence, EvalTime, bLooping, bFireNotifies, DeltaTime, Player.GetPlaybackStatus() == EMovieScenePlayerStatus::Playing, bResetDynamics);
+				for (const FMinimalAnimParameters& SectionParam : DataContainer.SectionParams)
+				{
+					UAnimSequenceBase* AnimSequence = SectionParam.Animation;
+					float EvalTime = SectionParam.EvalTime;
+					float Weight = SectionParam.BlendWeight;
+
+					uint32 SectionId = SectionParam.SectionId;
+
+					if (bPreviewPlayback)
+					{
+						PreviewSetAnimPosition(SkeletalMeshComponent, SectionParam.SlotName, SectionId, AnimSequence, EvalTime, Weight, bLooping, bFireNotifies, DeltaTime, Player.GetPlaybackStatus() == EMovieScenePlayerStatus::Playing, bResetDynamics);
+					}
+					else
+					{
+						SetAnimPosition(SkeletalMeshComponent, SectionParam.SlotName, SectionId, AnimSequence, EvalTime, Weight, bLooping, bFireNotifies);
+					}
+				}
 			}
-			else
+		}
+
+		// Clear all existing section data ready for the next frame
+		for (auto& Pair : ComponentDataContainer)
+		{
+			Pair.Value.SectionParams.Empty();
+		}
+	}
+
+	void StopAnimationForSlot(USkeletalMeshComponent* SkeletalMeshComponent, FName SlotName)
+	{
+		UAnimInstance* AnimInstance = SkeletalMeshComponent->GetAnimInstance();
+		UAnimMontage* Montage = GetCurrentlyPlayingMontage(SkeletalMeshComponent, SlotName).Get();
+		if (AnimInstance && Montage)
+		{
+			AnimInstance->Montage_Stop(0.f, Montage);
+		}
+	}
+
+private:
+
+	TWeakObjectPtr<class UAnimMontage>& GetCurrentlyPlayingMontage(USkeletalMeshComponent* SkeletalMeshComponent, const FName& SlotName)
+	{
+		FComponentDataContainer* DataContainer = ComponentDataContainer.Find(SkeletalMeshComponent);
+		check(DataContainer);
+		return DataContainer->CurrentlyPlayingMontages.FindOrAdd(SlotName);
+	}
+
+	bool ShouldUseSequenceInstancePlayer(const USkeletalMeshComponent* SkeletalMeshComponent)
+	{
+		const USkeletalMesh* SkeletalMesh = SkeletalMeshComponent->SkeletalMesh;
+		// create proper anim instance to animate
+		UAnimInstance* AnimInstance = SkeletalMeshComponent->GetAnimInstance();
+
+		return (AnimInstance == nullptr || SkeletalMeshComponent->GetAnimationMode() != EAnimationMode::AnimationBlueprint ||
+			AnimInstance->GetClass() != SkeletalMeshComponent->AnimClass || !SkeletalMesh->Skeleton->IsCompatible(AnimInstance->CurrentSkeleton));
+	}
+
+	void BeginAnimControl(USkeletalMeshComponent* SkeletalMeshComponent)
+	{
+		// make sure to tick and refresh all the time when ticks
+		SkeletalMeshComponent->MeshComponentUpdateFlag = EMeshComponentUpdateFlag::AlwaysTickPoseAndRefreshBones;
+		// we use sequence instance if it's using anim blueprint that matches. Otherwise, we create sequence player
+		// this might need more check - i.e. making sure if it's same skeleton and so on, 
+		// Ideally we could just call NeedToSpawnAnimScriptInstance call, which is protected now
+		const bool bShouldUseSequenceInstance = ShouldUseSequenceInstancePlayer(SkeletalMeshComponent);
+
+		// if it should use sequence instance
+		if (bShouldUseSequenceInstance)
+		{
+			// this has to wrap around with this because we don't want to reinitialize everytime they come here
+			// SetAnimationMode will reinitiaize it even if it's same, so make sure we just call SetAnimationMode if not AnimationCustomMode
+			if (SkeletalMeshComponent->GetAnimationMode() != EAnimationMode::AnimationCustomMode)
+			{
+				SkeletalMeshComponent->SetAnimationMode(EAnimationMode::AnimationCustomMode);
+			}
+
+			if (SkeletalMeshComponent->AnimScriptInstance == nullptr)
 			{
-				// Don't fire notifies at runtime since they will be fired through the standard animation tick path.
-				const bool bFireNotifies = false;
-				SetAnimPosition(SkeletalMeshComponent, Params.SlotName, ChannelIndex, AnimSequence, EvalTime, bLooping, bFireNotifies);
+				UAnimSequencerInstance* SequencerInstance = NewObject<UAnimSequencerInstance>(SkeletalMeshComponent);
+				SkeletalMeshComponent->AnimScriptInstance = SequencerInstance;
+				SkeletalMeshComponent->AnimScriptInstance->InitializeAnimation();
 			}
 		}
 	}
 
-	void SetAnimPosition(USkeletalMeshComponent* SkeletalMeshComponent, FName SlotName, int32 ChannelIndex, UAnimSequenceBase* InAnimSequence, float InPosition, bool bLooping, bool bFireNotifies)
+	void SetAnimPosition(USkeletalMeshComponent* SkeletalMeshComponent, FName SlotName, int32 SequenceIndex, UAnimSequenceBase* InAnimSequence, float InPosition, float Weight, bool bLooping, bool bFireNotifies)
 	{
 		if (CanPlayAnimation(SkeletalMeshComponent, InAnimSequence))
 		{
-			UAnimMontage* Montage = FAnimMontageInstance::SetMatineeAnimPositionInner(SlotName, SkeletalMeshComponent, InAnimSequence, InPosition, bLooping);
+			UAnimSequencerInstance* SequencerInst = Cast<UAnimSequencerInstance>(SkeletalMeshComponent->GetAnimInstance());
 
-			// Ensure the sequence is not stopped
-			UAnimInstance* AnimInst = SkeletalMeshComponent->GetAnimInstance();
-			UAnimSingleNodeInstance* SingleNodeInst = SkeletalMeshComponent->GetSingleNodeInstance();
-			if(SingleNodeInst)
+			if (SequencerInst)
 			{
-				SingleNodeInst->SetPlaying(true);
+				// Set position and weight
+				SequencerInst->UpdateAnimTrack(InAnimSequence, SequenceIndex, InPosition, Weight, bFireNotifies);
 			}
-			else if (AnimInst && Montage)
+			else
 			{
-				AnimInst->Montage_Resume(Montage);
+				TWeakObjectPtr<class UAnimMontage>& CurrentlyPlayingMontage = GetCurrentlyPlayingMontage(SkeletalMeshComponent, SlotName);
+				CurrentlyPlayingMontage = FAnimMontageInstance::SetMatineeAnimPositionInner(SlotName, SkeletalMeshComponent, InAnimSequence, InPosition, bLooping);
+
+				// Ensure the sequence is not stopped
+				UAnimInstance* AnimInst = SkeletalMeshComponent->GetAnimInstance();
+				if (AnimInst && CurrentlyPlayingMontage.IsValid())
+				{
+					AnimInst->Montage_Resume(CurrentlyPlayingMontage.Get());
+				}
 			}
 		}
 	}
 
-	void PreviewSetAnimPosition(USkeletalMeshComponent* SkeletalMeshComponent, FName SlotName, int32 ChannelIndex, UAnimSequenceBase* InAnimSequence, float InPosition, bool bLooping, bool bFireNotifies, float DeltaTime, bool bPlaying, bool bResetDynamics)
+	void PreviewSetAnimPosition(USkeletalMeshComponent* SkeletalMeshComponent, FName SlotName, int32 SequenceIndex, UAnimSequenceBase* InAnimSequence, float InPosition, float Weight, bool bLooping, bool bFireNotifies, float DeltaTime, bool bPlaying, bool bResetDynamics)
 	{
 		if(!CanPlayAnimation(SkeletalMeshComponent, InAnimSequence))
 		{
 			return;
 		}
 
-		UAnimMontage* Montage = FAnimMontageInstance::PreviewMatineeSetAnimPositionInner(SlotName, SkeletalMeshComponent, InAnimSequence, InPosition, bLooping, bFireNotifies, DeltaTime);
+		UAnimSequencerInstance* SequencerInst = Cast<UAnimSequencerInstance>(SkeletalMeshComponent->GetAnimInstance());
 
-		// if we are not playing, make sure we dont continue (as skeletal meshes can still tick us onwards)
-		UAnimInstance* AnimInst = SkeletalMeshComponent->GetAnimInstance();
-		UAnimSingleNodeInstance * SingleNodeInst = SkeletalMeshComponent->GetSingleNodeInstance();
-		if(SingleNodeInst)
+		if (SequencerInst)
 		{
-			SingleNodeInst->SetPlaying(bPlaying);
+			// Set position and weight
+			SequencerInst->UpdateAnimTrack(InAnimSequence, SequenceIndex, InPosition, Weight, bFireNotifies);
+
+			// Update space bases so new animation position has an effect.
+			// @todo - hack - this will be removed at some point
+			SkeletalMeshComponent->TickAnimation(0.03f, false);
+
+			SkeletalMeshComponent->RefreshBoneTransforms();
+			SkeletalMeshComponent->RefreshSlaveComponents();
+			SkeletalMeshComponent->UpdateComponentToWorld();
+			SkeletalMeshComponent->FinalizeBoneTransform();
+			SkeletalMeshComponent->MarkRenderTransformDirty();
+			SkeletalMeshComponent->MarkRenderDynamicDataDirty();
 		}
-		else if (AnimInst)
+		else
 		{
-			if(Montage)
+			TWeakObjectPtr<class UAnimMontage>& CurrentlyPlayingMontage = GetCurrentlyPlayingMontage(SkeletalMeshComponent, SlotName);
+			CurrentlyPlayingMontage = FAnimMontageInstance::PreviewMatineeSetAnimPositionInner(SlotName, SkeletalMeshComponent, InAnimSequence, InPosition, bLooping, bFireNotifies, DeltaTime);
+
+			// add to montage
+			// if we are not playing, make sure we dont continue (as skeletal meshes can still tick us onwards)
+			UAnimInstance* AnimInst = SkeletalMeshComponent->GetAnimInstance();
+			if (AnimInst)
 			{
-				if(bPlaying)
+				if (CurrentlyPlayingMontage.IsValid())
 				{
-					AnimInst->Montage_Resume(Montage);
+					if (bPlaying)
+					{
+						AnimInst->Montage_Resume(CurrentlyPlayingMontage.Get());
+					}
+					else
+					{
+						AnimInst->Montage_Pause(CurrentlyPlayingMontage.Get());
+					}
 				}
-				else
+
+				if (bResetDynamics)
 				{
-					AnimInst->Montage_Pause(Montage);
+					// make sure we reset any simulations
+					AnimInst->ResetDynamics();
 				}
 			}
+		}
+	}
+};
+
+struct FAnimationSectionData : IPersistentEvaluationData
+{
+	TSet<FObjectKey> AllAnimatedComponents;
+};
+
 
-			if(bResetDynamics)
+// Execution token that assigns blended animation to all skeletal mesh components for this section
+struct FSkeletalAnimationAccumulatorToken : IMovieSceneExecutionToken
+{
+	FSkeletalAnimationAccumulatorToken(const FMinimalAnimParameters& InAnimParams) : AnimParams(InAnimParams){}
+
+	virtual void Execute(const FMovieSceneContext& Context, const FMovieSceneEvaluationOperand& Operand, FPersistentEvaluationData& PersistentData, IMovieScenePlayer& Player) override
+	{
+		FSkeletalAnimationTrackData& SharedData = PersistentData.GetOrAdd<FSkeletalAnimationTrackData>(FMovieSceneSkeletalAnimationSharedTrack::GetSharedDataKey());
+
+		FAnimationSectionData& ThisSectionData = PersistentData.GetOrAddSectionData<FAnimationSectionData>();
+
+		for (TWeakObjectPtr<> Object : Player.FindBoundObjects(Operand))
+		{
+			UObject* ObjectPtr = Object.Get();
+			if (!ObjectPtr)
 			{
-				// make sure we reset any simulations
-				AnimInst->ResetDynamics();
+				continue;
+			}
+
+			// first try to control the component directly
+			USkeletalMeshComponent* SkeletalMeshComponent = Cast<USkeletalMeshComponent>(ObjectPtr);
+			if (!SkeletalMeshComponent)
+			{
+				// then check to see if we are controlling an actor & if so use its first USkeletalMeshComponent 
+				if (AActor* Actor = Cast<AActor>(ObjectPtr))
+				{
+					SkeletalMeshComponent = Actor->FindComponentByClass<USkeletalMeshComponent>();
+				}
+			}
+
+			if (SkeletalMeshComponent)
+			{
+				// Save preanimated state here rather than the shared track as the shared track will still be evaluated even if there are no animations happening
+				Player.SavePreAnimatedState(*SkeletalMeshComponent, TMovieSceneAnimTypeID<FSkeletalAnimationAccumulatorToken>(), FPreAnimatedAnimationTokenProducer());
+				ThisSectionData.AllAnimatedComponents.Add(SkeletalMeshComponent);
+
+				SharedData.AssignSkeletalAnimationData(SkeletalMeshComponent, AnimParams);
 			}
 		}
 	}
 
-	FMovieSceneSkeletalAnimationSectionTemplateParameters Params;
+	FMinimalAnimParameters AnimParams;
 };
 
-
 FMovieSceneSkeletalAnimationSectionTemplate::FMovieSceneSkeletalAnimationSectionTemplate(const UMovieSceneSkeletalAnimationSection& InSection)
-	: Params(InSection.Params, InSection.GetStartTime())
+	: Params(InSection.Params, InSection.GetStartTime(), InSection.GetEndTime())
 {
 }
 
 void FMovieSceneSkeletalAnimationSectionTemplate::Evaluate(const FMovieSceneEvaluationOperand& Operand, const FMovieSceneContext& Context, const FPersistentEvaluationData& PersistentData, FMovieSceneExecutionTokens& ExecutionTokens) const
 {
-	ExecutionTokens.Add(FAnimationTrackTrackToken(Params));
+	if (Params.Animation)
+	{
+		// calculate the time at which to evaluate the animation
+		float EvalTime = Params.MapTimeToAnimation(Context.GetTime());
+		float BlendWeight = Params.Weight.Eval(Context.GetTime());
+
+		// Create the execution token with the minimal anim params
+		FMinimalAnimParameters AnimParams(
+			Params.Animation, EvalTime, BlendWeight, GetTypeHash(PersistentData.GetSectionKey()), Params.SlotName
+			);
+		ExecutionTokens.Add(FSkeletalAnimationAccumulatorToken(AnimParams));
+	}
+}
+
+void FMovieSceneSkeletalAnimationSectionTemplate::TearDown(FPersistentEvaluationData& PersistentData, IMovieScenePlayer& Player) const
+{
+	// Reset blend weights on any nodes that we animated
+	FAnimationSectionData* ThisSectionData = PersistentData.FindSectionData<FAnimationSectionData>();
+	FSkeletalAnimationTrackData* SharedData = PersistentData.Find<FSkeletalAnimationTrackData>(FMovieSceneSkeletalAnimationSharedTrack::GetSharedDataKey());
+
+	if (ThisSectionData)
+	{
+		for (FObjectKey Key : ThisSectionData->AllAnimatedComponents)
+		{
+			if (USkeletalMeshComponent* Component = Cast<USkeletalMeshComponent>(Key.ResolveObjectPtr()))
+			{
+				UAnimSequencerInstance* SequencerInst = Cast<UAnimSequencerInstance>(Component->GetAnimInstance());
+				if (SequencerInst)
+				{
+					SequencerInst->ResetNodes();
+				}
+
+				if (SharedData)
+				{
+					SharedData->StopAnimationForSlot(Component, Params.SlotName);
+				}
+			}
+		}
+	}
 }
 
 float FMovieSceneSkeletalAnimationSectionTemplateParameters::MapTimeToAnimation(float ThisPosition) const
@@ -262,13 +441,18 @@ float FMovieSceneSkeletalAnimationSectionTemplateParameters::MapTimeToAnimation(
 	// @todo: Sequencer: what is this for??
 	//ThisPosition -= 1 / 1000.0f;
 
+	ThisPosition = FMath::Clamp(ThisPosition, SectionStartTime, SectionEndTime);
+
 	const float SectionPlayRate = PlayRate;
 	const float AnimPlayRate = FMath::IsNearlyZero(SectionPlayRate) ? 1.0f : SectionPlayRate;
 
 	const float SeqLength = GetSequenceLength() - (StartOffset + EndOffset);
 
 	ThisPosition = (ThisPosition - SectionStartTime) * AnimPlayRate;
-	ThisPosition = FMath::Fmod(ThisPosition, SeqLength);
+	if (SeqLength > 0.f)
+	{
+		ThisPosition = FMath::Fmod(ThisPosition, SeqLength);
+	}
 	ThisPosition += StartOffset;
 	if (bReverse)
 	{
@@ -277,3 +461,26 @@ float FMovieSceneSkeletalAnimationSectionTemplateParameters::MapTimeToAnimation(
 
 	return ThisPosition;
 }
+
+struct FSkeletalAnimationExecutionToken : IMovieSceneExecutionToken
+{
+	virtual void Execute(const FMovieSceneContext& Context, const FMovieSceneEvaluationOperand& Operand, FPersistentEvaluationData& PersistentData, IMovieScenePlayer& Player) override
+	{
+		FSkeletalAnimationTrackData* TrackData = PersistentData.Find<FSkeletalAnimationTrackData>(FMovieSceneSkeletalAnimationSharedTrack::GetSharedDataKey());
+		if (TrackData)
+		{
+			TrackData->ApplySkeletalAnimationData(Context, Operand, Player);
+		}
+	}
+};
+
+void FMovieSceneSkeletalAnimationSharedTrack::Evaluate(const FMovieSceneEvaluationOperand& Operand, const FMovieSceneContext& Context, const FPersistentEvaluationData& PersistentData, FMovieSceneExecutionTokens& ExecutionTokens) const
+{
+	ExecutionTokens.Add(FSkeletalAnimationExecutionToken());
+}
+
+FSharedPersistentDataKey FMovieSceneSkeletalAnimationSharedTrack::GetSharedDataKey()
+{
+	static FMovieSceneSharedDataId DataId(FMovieSceneSharedDataId::Allocate());
+	return FSharedPersistentDataKey(DataId, FMovieSceneEvaluationOperand());
+}
\ No newline at end of file
diff --git a/Engine/Source/Runtime/MovieSceneTracks/Private/Evaluation/MovieSceneSkeletalAnimationTemplate.h b/Engine/Source/Runtime/MovieSceneTracks/Private/Evaluation/MovieSceneSkeletalAnimationTemplate.h
--- a/Engine/Source/Runtime/MovieSceneTracks/Private/Evaluation/MovieSceneSkeletalAnimationTemplate.h
+++ b/Engine/Source/Runtime/MovieSceneTracks/Private/Evaluation/MovieSceneSkeletalAnimationTemplate.h
@@ -14,15 +14,19 @@ struct FMovieSceneSkeletalAnimationSectionTemplateParameters : public FMovieScen
 	GENERATED_BODY()
 
 	FMovieSceneSkeletalAnimationSectionTemplateParameters() {}
-	FMovieSceneSkeletalAnimationSectionTemplateParameters(const FMovieSceneSkeletalAnimationParams& BaseParams, float InSectionStartTime)
+	FMovieSceneSkeletalAnimationSectionTemplateParameters(const FMovieSceneSkeletalAnimationParams& BaseParams, float InSectionStartTime, float InSectionEndTime)
 		: FMovieSceneSkeletalAnimationParams(BaseParams)
 		, SectionStartTime(InSectionStartTime)
+		, SectionEndTime(InSectionEndTime)
 	{}
 
 	float MapTimeToAnimation(float InPosition) const;
 
 	UPROPERTY()
 	float SectionStartTime;
+
+	UPROPERTY()
+	float SectionEndTime;
 };
 
 USTRUCT()
@@ -34,8 +38,29 @@ struct FMovieSceneSkeletalAnimationSectionTemplate : public FMovieSceneEvalTempl
 	FMovieSceneSkeletalAnimationSectionTemplate(const UMovieSceneSkeletalAnimationSection& Section);
 
 	virtual UScriptStruct& GetScriptStructImpl() const override { return *StaticStruct(); }
+
+	virtual void SetupOverrides() override
+	{
+		EnableOverrides(RequiresTearDownFlag);
+	}
+	
 	virtual void Evaluate(const FMovieSceneEvaluationOperand& Operand, const FMovieSceneContext& Context, const FPersistentEvaluationData& PersistentData, FMovieSceneExecutionTokens& ExecutionTokens) const override;
+	virtual void TearDown(FPersistentEvaluationData& PersistentData, IMovieScenePlayer& Player) const override;
 
 	UPROPERTY()
 	FMovieSceneSkeletalAnimationSectionTemplateParameters Params;
 };
+
+USTRUCT()
+struct FMovieSceneSkeletalAnimationSharedTrack
+	: public FMovieSceneEvalTemplate
+{
+	GENERATED_BODY()
+
+	MOVIESCENETRACKS_API static FSharedPersistentDataKey GetSharedDataKey();
+
+private:
+
+	virtual UScriptStruct& GetScriptStructImpl() const override { return *StaticStruct(); }
+	virtual void Evaluate(const FMovieSceneEvaluationOperand& Operand, const FMovieSceneContext& Context, const FPersistentEvaluationData& PersistentData, FMovieSceneExecutionTokens& ExecutionTokens) const override;
+};
diff --git a/Engine/Source/Runtime/MovieSceneTracks/Private/Evaluation/MovieSceneSlomoTemplate.cpp b/Engine/Source/Runtime/MovieSceneTracks/Private/Evaluation/MovieSceneSlomoTemplate.cpp
--- a/Engine/Source/Runtime/MovieSceneTracks/Private/Evaluation/MovieSceneSlomoTemplate.cpp
+++ b/Engine/Source/Runtime/MovieSceneTracks/Private/Evaluation/MovieSceneSlomoTemplate.cpp
@@ -7,9 +7,13 @@
 #include "Engine/Engine.h"
 #include "GameFramework/WorldSettings.h"
 #include "EngineGlobals.h"
+#include "MovieSceneEvaluation.h"
+#include "IMovieScenePlayer.h"
+
 
 DECLARE_CYCLE_STAT(TEXT("Slomo Track Token Execute"), MovieSceneEval_SlomoTrack_TokenExecute, STATGROUP_MovieSceneEval);
 
+
 struct FSlomoTrackToken
 {
 	FSlomoTrackToken(float InSlomoValue)
@@ -23,18 +27,15 @@ struct FSlomoTrackToken
 		UObject* PlaybackContext = Player.GetPlaybackContext();
 		UWorld* World = PlaybackContext ? PlaybackContext->GetWorld() : nullptr;
 
-		if (!GIsEditor || 
-			(World && World->GetWorld()->GetNetMode() == NM_Client) ||
-			!GEngine ||
-			SlomoValue <= 0.f)
+		if (!World || (!GIsEditor && World->GetNetMode() == NM_Client) || SlomoValue <= 0.f)
 		{
 			return;
 		}
 
-		AWorldSettings* WorldSettings = World ? World->GetWorldSettings() : nullptr;
-			
+		AWorldSettings* WorldSettings = World->GetWorldSettings();
+
 		if (WorldSettings)
-		{	
+		{
 			WorldSettings->MatineeTimeDilation = SlomoValue;
 			WorldSettings->ForceNetUpdate();
 		}
diff --git a/Engine/Source/Runtime/MovieSceneTracks/Private/Evaluation/MovieSceneSpawnTemplate.cpp b/Engine/Source/Runtime/MovieSceneTracks/Private/Evaluation/MovieSceneSpawnTemplate.cpp
--- a/Engine/Source/Runtime/MovieSceneTracks/Private/Evaluation/MovieSceneSpawnTemplate.cpp
+++ b/Engine/Source/Runtime/MovieSceneTracks/Private/Evaluation/MovieSceneSpawnTemplate.cpp
@@ -1,9 +1,10 @@
 // Copyright 1998-2017 Epic Games, Inc. All Rights Reserved.
 
 #include "Evaluation/MovieSceneSpawnTemplate.h"
-
 #include "MovieSceneSequence.h"
 #include "Sections/MovieSceneSpawnSection.h"
+#include "MovieSceneEvaluation.h"
+#include "IMovieScenePlayer.h"
 
 DECLARE_CYCLE_STAT(TEXT("Spawn Track Evaluate"), MovieSceneEval_SpawnTrack_Evaluate, STATGROUP_MovieSceneEval);
 DECLARE_CYCLE_STAT(TEXT("Spawn Track Token Execute"), MovieSceneEval_SpawnTrack_TokenExecute, STATGROUP_MovieSceneEval);
diff --git a/Engine/Source/Runtime/MovieSceneTracks/Private/Evaluation/MovieSceneTemplateCommon.cpp b/Engine/Source/Runtime/MovieSceneTracks/Private/Evaluation/MovieSceneTemplateCommon.cpp
--- a/Engine/Source/Runtime/MovieSceneTracks/Private/Evaluation/MovieSceneTemplateCommon.cpp
+++ b/Engine/Source/Runtime/MovieSceneTracks/Private/Evaluation/MovieSceneTemplateCommon.cpp
@@ -2,6 +2,7 @@
 
 #include "Evaluation/MovieSceneTemplateCommon.h"
 #include "Components/SceneComponent.h"
+#include "GameFramework/Actor.h"
 
 /** A movie scene pre-animated token that stores a pre-animated mobility */
 struct FMobilityPreAnimatedToken : IMovieScenePreAnimatedToken
@@ -39,7 +40,13 @@ void F3DTransformTrackToken::Apply(USceneComponent& SceneComponent, float DeltaT
 	/* Cache initial absolute position */
 	FVector PreviousPosition = SceneComponent.GetComponentLocation();
 
-	SceneComponent.SetRelativeLocationAndRotation(Translation, Rotation);
+	// If this is a simulating component, teleport since sequencer takes over. 
+	// Teleport will not have no velocity, but it's computed below by sequencer so that it will be correct for physics.
+	AActor* Actor = SceneComponent.GetOwner();
+	USceneComponent* RootComponent = Actor ? Actor->GetRootComponent() : nullptr;
+	bool bIsSimulatingPhysics = RootComponent ? RootComponent->IsSimulatingPhysics() : false;
+
+	SceneComponent.SetRelativeLocationAndRotation(Translation, Rotation, false, nullptr, bIsSimulatingPhysics ? ETeleportType::TeleportPhysics : ETeleportType::None);
 	SceneComponent.SetRelativeScale3D(Scale);
 
 	// Force the location and rotation values to avoid Rot->Quat->Rot conversions
diff --git a/Engine/Source/Runtime/MovieSceneTracks/Private/Evaluation/MovieSceneVisibilityTemplate.cpp b/Engine/Source/Runtime/MovieSceneTracks/Private/Evaluation/MovieSceneVisibilityTemplate.cpp
--- a/Engine/Source/Runtime/MovieSceneTracks/Private/Evaluation/MovieSceneVisibilityTemplate.cpp
+++ b/Engine/Source/Runtime/MovieSceneTracks/Private/Evaluation/MovieSceneVisibilityTemplate.cpp
@@ -5,10 +5,13 @@
 #include "Tracks/MovieScenePropertyTrack.h"
 #include "GameFramework/Actor.h"
 #include "Engine/World.h"
+#include "MovieSceneEvaluation.h"
+
 
 DECLARE_CYCLE_STAT(TEXT("Visibility Track Evaluate"), MovieSceneEval_VisibilityTrack_Evaluate, STATGROUP_MovieSceneEval);
 DECLARE_CYCLE_STAT(TEXT("Visibility Track Token Execute"), MovieSceneEval_VisibilityTrack_TokenExecute, STATGROUP_MovieSceneEval);
 
+
 /** A movie scene pre-animated token that stores a pre-animated actor's temporarily hidden in game */
 struct FTemporarilyHiddenInGamePreAnimatedToken : IMovieScenePreAnimatedToken
 {
@@ -98,10 +101,13 @@ void FMovieSceneVisibilitySectionTemplate::Evaluate(const FMovieSceneEvaluationO
 
 	for (TCachedValue<bool>& ObjectAndValue : PersistentData.GetSectionData<TCachedSectionData<bool>>().ObjectsAndValues)
 	{
-		ObjectAndValue.Value = !!BoolCurve.Evaluate(Context.GetTime(), ObjectAndValue.Value ? 1 : 0);
+		bool& bIsHidden = ObjectAndValue.Value;
+
+		// Pass the inverse of our current visibility since we negate the result of the curve evaluation.
+		uint8 InverseCurrentVisibility = bIsHidden ? 1 : 0;
 
 		// Invert this evaluation since the property is "bHiddenInGame" and we want the visualization to be the inverse of that. Green means visible.
-		ObjectAndValue.Value = !ObjectAndValue.Value;
+		bIsHidden = !BoolCurve.Evaluate(Context.GetTime(), InverseCurrentVisibility);
 	}
 	
 	ExecutionTokens.Add(FTemporarilyHiddenInGameExecutionToken());
diff --git a/Engine/Source/Runtime/MovieSceneTracks/Private/Sections/MovieSceneAudioSection.cpp b/Engine/Source/Runtime/MovieSceneTracks/Private/Sections/MovieSceneAudioSection.cpp
--- a/Engine/Source/Runtime/MovieSceneTracks/Private/Sections/MovieSceneAudioSection.cpp
+++ b/Engine/Source/Runtime/MovieSceneTracks/Private/Sections/MovieSceneAudioSection.cpp
@@ -7,11 +7,13 @@
 UMovieSceneAudioSection::UMovieSceneAudioSection( const FObjectInitializer& ObjectInitializer )
 	: Super( ObjectInitializer )
 {
-
 	Sound = nullptr;
-	AudioStartTime = 0.f;
-	AudioDilationFactor = 1.f;
-	AudioVolume = 1.f;
+	StartOffset = 0.f;
+	AudioStartTime_DEPRECATED = 0.f;
+	AudioDilationFactor_DEPRECATED = 1.f;
+	AudioVolume_DEPRECATED = 1.f;
+	SoundVolume.SetDefaultValue(1.f);
+	PitchMultiplier.SetDefaultValue(1.f);
 #if WITH_EDITORONLY_DATA
 	bShowIntensity = false;
 #endif
@@ -27,55 +29,113 @@ FMovieSceneEvalTemplatePtr UMovieSceneAudioSection::GenerateTemplate() const
 
 TRange<float> UMovieSceneAudioSection::GetAudioRange() const
 {
-	return !Sound ? TRange<float>::Empty() :
-		TRange<float>(FMath::Max(GetStartTime(), AudioStartTime),
-		FMath::Min(AudioStartTime + Sound->GetDuration() * AudioDilationFactor, GetEndTime()));
-}
+	if (!Sound)
+	{
+		return TRange<float>::Empty();
+	}
+
+	float AudioRangeStartTime = GetStartTime();
+	if (StartOffset < 0)
+	{
+		AudioRangeStartTime += FMath::Abs(StartOffset);
+	}
 
+	float AudioRangeEndTime = FMath::Min(AudioRangeStartTime + Sound->GetDuration() * PitchMultiplier.GetDefaultValue(), GetEndTime());
 
-TRange<float> UMovieSceneAudioSection::GetAudioTrueRange() const
-{
-	return !Sound ? TRange<float>::Empty() : TRange<float>(AudioStartTime, AudioStartTime + Sound->GetDuration() * AudioDilationFactor);
+	return TRange<float>(AudioRangeStartTime, AudioRangeEndTime);
 }
 
+void UMovieSceneAudioSection::PostLoad()
+{
+	Super::PostLoad();
+
+	if (AudioDilationFactor_DEPRECATED != FLT_MAX)
+	{
+		PitchMultiplier.SetDefaultValue(AudioDilationFactor_DEPRECATED);
+
+		AudioDilationFactor_DEPRECATED = FLT_MAX;
+	}
+
+	if (AudioVolume_DEPRECATED != FLT_MAX)
+	{
+		SoundVolume.SetDefaultValue(AudioVolume_DEPRECATED);
+
+		AudioVolume_DEPRECATED = FLT_MAX;
+	}
+
+	if (AudioStartTime_DEPRECATED != FLT_MAX)
+	{
+		// Previously, start time in relation to the sequence. Start time was used to calculate the offset into the 
+		// clip at the start of the section evaluation as such: Section Start Time - Start Time. 
+		if (AudioStartTime_DEPRECATED != 0.f)
+		{
+			StartOffset = GetStartTime() - AudioStartTime_DEPRECATED;
+		}
+		AudioStartTime_DEPRECATED = FLT_MAX;
+	}
+}
 
 void UMovieSceneAudioSection::MoveSection( float DeltaTime, TSet<FKeyHandle>& KeyHandles )
 {
 	Super::MoveSection(DeltaTime, KeyHandles);
 
-	AudioStartTime += DeltaTime;
+	SoundVolume.ShiftCurve(DeltaTime, KeyHandles);
+	PitchMultiplier.ShiftCurve(DeltaTime, KeyHandles);
 }
 
 
 void UMovieSceneAudioSection::DilateSection( float DilationFactor, float Origin, TSet<FKeyHandle>& KeyHandles )
 {
-	AudioStartTime = (AudioStartTime - Origin) * DilationFactor + Origin;
-	AudioDilationFactor *= DilationFactor;
-
 	Super::DilateSection(DilationFactor, Origin, KeyHandles);
+
+	SoundVolume.ScaleCurve(Origin, DilationFactor, KeyHandles);
+	PitchMultiplier.ScaleCurve(Origin, DilationFactor, KeyHandles);
 }
 
 
-void UMovieSceneAudioSection::GetKeyHandles(TSet<FKeyHandle>& KeyHandles, TRange<float> TimeRange) const
+UMovieSceneSection* UMovieSceneAudioSection::SplitSection(float SplitTime)
 {
-	// do nothing
-}
+	float NewOffset = StartOffset + (SplitTime - GetStartTime());
 
+	UMovieSceneSection* NewSection = Super::SplitSection(SplitTime);
+	if (NewSection != nullptr)
+	{
+		UMovieSceneAudioSection* NewAudioSection = Cast<UMovieSceneAudioSection>(NewSection);
+		NewAudioSection->StartOffset = NewOffset;
+	}
+	return NewSection;
+}
 
-void UMovieSceneAudioSection::GetSnapTimes(TArray<float>& OutSnapTimes, bool bGetSectionBorders) const
+void UMovieSceneAudioSection::GetKeyHandles(TSet<FKeyHandle>& OutKeyHandles, TRange<float> TimeRange) const
 {
-	Super::GetSnapTimes(OutSnapTimes, bGetSectionBorders);
-	if (AudioStartTime > GetStartTime())
+	if (!TimeRange.Overlaps(GetRange()))
 	{
-		OutSnapTimes.Add(AudioStartTime);
+		return;
 	}
 
-	float AudioEndTime = GetAudioTrueRange().GetUpperBoundValue();
+	for (auto It(SoundVolume.GetKeyHandleIterator()); It; ++It)
+	{
+		float Time = SoundVolume.GetKeyTime(It.Key());
+		if (TimeRange.Contains(Time))
+		{
+			OutKeyHandles.Add(It.Key());
+		}
+	}
 
-	if (AudioEndTime < GetEndTime())
+	for (auto It(PitchMultiplier.GetKeyHandleIterator()); It; ++It)
 	{
-		OutSnapTimes.Add(AudioEndTime);
+		float Time = PitchMultiplier.GetKeyTime(It.Key());
+		if (TimeRange.Contains(Time))
+		{
+			OutKeyHandles.Add(It.Key());
+		}
 	}
+}
+
+
+void UMovieSceneAudioSection::GetSnapTimes(TArray<float>& OutSnapTimes, bool bGetSectionBorders) const
+{
+	Super::GetSnapTimes(OutSnapTimes, bGetSectionBorders);
 
 	// @todo Sequencer handle snapping for time dilation
 	// @todo Don't add redundant times (can't use AddUnique due to floating point equality issues)
diff --git a/Engine/Source/Runtime/MovieSceneTracks/Private/Sections/MovieSceneSkeletalAnimationSection.cpp b/Engine/Source/Runtime/MovieSceneTracks/Private/Sections/MovieSceneSkeletalAnimationSection.cpp
--- a/Engine/Source/Runtime/MovieSceneTracks/Private/Sections/MovieSceneSkeletalAnimationSection.cpp
+++ b/Engine/Source/Runtime/MovieSceneTracks/Private/Sections/MovieSceneSkeletalAnimationSection.cpp
@@ -3,6 +3,11 @@
 #include "Sections/MovieSceneSkeletalAnimationSection.h"
 #include "Animation/AnimSequence.h"
 #include "Evaluation/MovieSceneSkeletalAnimationTemplate.h"
+#include "MessageLog.h"
+#include "MovieScene.h"
+#include "SequencerObjectVersion.h"
+
+#define LOCTEXT_NAMESPACE "MovieSceneSkeletalAnimationSection"
 
 namespace
 {
@@ -17,6 +22,7 @@ FMovieSceneSkeletalAnimationParams::FMovieSceneSkeletalAnimationParams()
 	PlayRate = 1.f;
 	bReverse = false;
 	SlotName = DefaultSlotName;
+	Weight.SetDefaultValue(1.f);
 }
 
 UMovieSceneSkeletalAnimationSection::UMovieSceneSkeletalAnimationSection( const FObjectInitializer& ObjectInitializer )
@@ -30,11 +36,20 @@ UMovieSceneSkeletalAnimationSection::UMovieSceneSkeletalAnimationSection( const
 	bReverse_DEPRECATED = false;
 	SlotName_DEPRECATED = DefaultSlotName;
 
+	// Section template relies on always restoring state for objects when they are no longer animating. This is how it releases animation control.
+	EvalOptions.CompletionMode = EMovieSceneCompletionMode::RestoreState;
+
 #if WITH_EDITOR
 	PreviousPlayRate = Params.PlayRate;
 #endif
 }
 
+void UMovieSceneSkeletalAnimationSection::Serialize(FArchive& Ar)
+{
+	Ar.UsingCustomVersion(FSequencerObjectVersion::GUID);
+	Super::Serialize(Ar);
+}
+
 void UMovieSceneSkeletalAnimationSection::PostLoad()
 {
 	if (AnimSequence_DEPRECATED)
@@ -72,6 +87,39 @@ void UMovieSceneSkeletalAnimationSection::PostLoad()
 		Params.SlotName = SlotName_DEPRECATED;
 	}
 
+	// if version is less than this
+	if (GetLinkerCustomVersion(FSequencerObjectVersion::GUID) < FSequencerObjectVersion::ConvertEnableRootMotionToForceRootLock)
+	{
+		UAnimSequence* AnimSeq = Cast<UAnimSequence>(Params.Animation);
+		if (AnimSeq && AnimSeq->bEnableRootMotion)
+		{
+			// this is not ideal, but previously single player node was using this flag to whether or not to extract root motion
+			// with new anim sequencer instance, this would break because we use the instance flag to extract root motion or not
+			// so instead of setting that flag, we use bForceRootLock flag to asset
+			// this can have side effect, where users didn't want that to be on to start with
+			// So we'll notify users to let them know this has to be saved
+			AnimSeq->bForceRootLock = true;
+			AnimSeq->MarkPackageDirty();
+			// warning to users
+#if WITH_EDITOR			
+			if (!IsRunningGame())
+			{
+				static FName NAME_LoadErrors("LoadErrors");
+				FMessageLog LoadErrors(NAME_LoadErrors);
+
+				TSharedRef<FTokenizedMessage> Message = LoadErrors.Warning();
+				Message->AddToken(FTextToken::Create(LOCTEXT("RootMotionFixUp1", "The Animation ")));
+				Message->AddToken(FAssetNameToken::Create(AnimSeq->GetPathName(), FText::FromString(GetNameSafe(AnimSeq))));
+				Message->AddToken(FTextToken::Create(LOCTEXT("RootMotionFixUp2", "will be set to ForceRootLock on. Please save the animation if you want to keep this change.")));
+				Message->SetSeverity(EMessageSeverity::Warning);
+				LoadErrors.Notify();
+			}
+#endif // WITH_EDITOR
+
+			UE_LOG(LogMovieScene, Warning, TEXT("%s Animation has set ForceRootLock to be used in Sequencer. If this animation is used in anywhere else using root motion, that will cause conflict."), *AnimSeq->GetName());
+		}
+	}
+
 	Super::PostLoad();
 }
 
@@ -80,10 +128,32 @@ FMovieSceneEvalTemplatePtr UMovieSceneSkeletalAnimationSection::GenerateTemplate
 	return FMovieSceneSkeletalAnimationSectionTemplate(*this);
 }
 
+void UMovieSceneSkeletalAnimationSection::PostLoadUpgradeTrackRow(const TRange<float>& InEvaluationRange)
+{
+	if (!Params.Weight.IsKeyHandleValid(Params.Weight.FindKey(GetStartTime())))
+	{
+		FKeyHandle StartKeyHandle = Params.Weight.UpdateOrAddKey(GetStartTime(), 1.f);
+		Params.Weight.SetKeyInterpMode(StartKeyHandle, RCIM_Constant);
+	}
+
+	if (InEvaluationRange.HasLowerBound())
+	{
+		FKeyHandle EvalInKeyHandle = Params.Weight.UpdateOrAddKey(InEvaluationRange.GetLowerBoundValue(), 1.f);
+		Params.Weight.SetKeyInterpMode(EvalInKeyHandle, RCIM_Constant);
+	}
+
+	if (InEvaluationRange.HasUpperBound())
+	{
+		FKeyHandle EvalOutKeyHandle = Params.Weight.UpdateOrAddKey(InEvaluationRange.GetUpperBoundValue(), 0.f);
+		Params.Weight.SetKeyInterpMode(EvalOutKeyHandle, RCIM_Constant);
+	}
+}
 
 void UMovieSceneSkeletalAnimationSection::MoveSection( float DeltaTime, TSet<FKeyHandle>& KeyHandles )
 {
 	Super::MoveSection(DeltaTime, KeyHandles);
+
+	Params.Weight.ShiftCurve(DeltaTime, KeyHandles);
 }
 
 
@@ -92,6 +162,8 @@ void UMovieSceneSkeletalAnimationSection::DilateSection( float DilationFactor, f
 	Params.PlayRate /= DilationFactor;
 
 	Super::DilateSection(DilationFactor, Origin, KeyHandles);
+	
+	Params.Weight.ScaleCurve(Origin, DilationFactor, KeyHandles);
 }
 
 UMovieSceneSection* UMovieSceneSkeletalAnimationSection::SplitSection(float SplitTime)
@@ -113,9 +185,21 @@ UMovieSceneSection* UMovieSceneSkeletalAnimationSection::SplitSection(float Spli
 }
 
 
-void UMovieSceneSkeletalAnimationSection::GetKeyHandles(TSet<FKeyHandle>& KeyHandles, TRange<float> TimeRange) const
+void UMovieSceneSkeletalAnimationSection::GetKeyHandles(TSet<FKeyHandle>& OutKeyHandles, TRange<float> TimeRange) const
 {
-	// do nothing
+	if (!TimeRange.Overlaps(GetRange()))
+	{
+		return;
+	}
+
+	for (auto It(Params.Weight.GetKeyHandleIterator()); It; ++It)
+	{
+		float Time = Params.Weight.GetKeyTime(It.Key());
+		if (TimeRange.Contains(Time))
+		{
+			OutKeyHandles.Add(It.Key());
+		}
+	}
 }
 
 
@@ -128,7 +212,7 @@ void UMovieSceneSkeletalAnimationSection::GetSnapTimes(TArray<float>& OutSnapTim
 	float SeqLength = (Params.GetSequenceLength() - (Params.StartOffset + Params.EndOffset)) / AnimPlayRate;
 
 	// Snap to the repeat times
-	while (CurrentTime <= GetEndTime() && !FMath::IsNearlyZero(Params.GetDuration()) && SeqLength > 0)
+	while (CurrentTime <= GetEndTime() && !FMath::IsNearlyZero(SeqLength, KINDA_SMALL_NUMBER) && SeqLength > 0)
 	{
 		if (CurrentTime >= GetStartTime())
 		{
@@ -170,3 +254,5 @@ void UMovieSceneSkeletalAnimationSection::PostEditChangeProperty(FPropertyChange
 	Super::PostEditChangeProperty(PropertyChangedEvent);
 }
 #endif
+
+#undef LOCTEXT_NAMESPACE 
\ No newline at end of file
diff --git a/Engine/Source/Runtime/MovieSceneTracks/Private/Tracks/MovieSceneAudioTrack.cpp b/Engine/Source/Runtime/MovieSceneTracks/Private/Tracks/MovieSceneAudioTrack.cpp
--- a/Engine/Source/Runtime/MovieSceneTracks/Private/Tracks/MovieSceneAudioTrack.cpp
+++ b/Engine/Source/Runtime/MovieSceneTracks/Private/Tracks/MovieSceneAudioTrack.cpp
@@ -29,6 +29,12 @@ const TArray<UMovieSceneSection*>& UMovieSceneAudioTrack::GetAllSections() const
 }
 
 
+bool UMovieSceneAudioTrack::SupportsMultipleRows() const
+{
+	return true;
+}
+
+
 void UMovieSceneAudioTrack::RemoveAllAnimationData()
 {
 	// do nothing
@@ -116,7 +122,6 @@ void UMovieSceneAudioTrack::AddNewSound(USoundBase* Sound, float Time)
 	// add the section
 	UMovieSceneAudioSection* NewSection = NewObject<UMovieSceneAudioSection>(this);
 	NewSection->InitialPlacement( AudioSections, Time, Time + DurationToUse, SupportsMultipleRows() );
-	NewSection->SetAudioStartTime(Time);
 	NewSection->SetSound(Sound);
 
 	AudioSections.Add(NewSection);
@@ -129,10 +134,4 @@ bool UMovieSceneAudioTrack::IsAMasterTrack() const
 }
 
 
-bool UMovieSceneAudioTrack::SupportsMultipleRows() const
-{
-	return true;
-}
-
-
 #undef LOCTEXT_NAMESPACE
diff --git a/Engine/Source/Runtime/MovieSceneTracks/Private/Tracks/MovieSceneSkeletalAnimationTrack.cpp b/Engine/Source/Runtime/MovieSceneTracks/Private/Tracks/MovieSceneSkeletalAnimationTrack.cpp
--- a/Engine/Source/Runtime/MovieSceneTracks/Private/Tracks/MovieSceneSkeletalAnimationTrack.cpp
+++ b/Engine/Source/Runtime/MovieSceneTracks/Private/Tracks/MovieSceneSkeletalAnimationTrack.cpp
@@ -3,7 +3,9 @@
 #include "Tracks/MovieSceneSkeletalAnimationTrack.h"
 #include "Sections/MovieSceneSkeletalAnimationSection.h"
 #include "Compilation/MovieSceneCompilerRules.h"
-
+#include "Evaluation/MovieSceneEvaluationTrack.h"
+#include "Evaluation/MovieSceneSkeletalAnimationTemplate.h"
+#include "Compilation/IMovieSceneTemplateGenerator.h"
 
 #define LOCTEXT_NAMESPACE "MovieSceneSkeletalAnimationTrack"
 
@@ -62,6 +64,12 @@ const TArray<UMovieSceneSection*>& UMovieSceneSkeletalAnimationTrack::GetAllSect
 }
 
 
+bool UMovieSceneSkeletalAnimationTrack::SupportsMultipleRows() const
+{
+	return true;
+}
+
+
 UMovieSceneSection* UMovieSceneSkeletalAnimationTrack::CreateNewSection()
 {
 	return NewObject<UMovieSceneSkeletalAnimationSection>(this);
@@ -119,5 +127,37 @@ FText UMovieSceneSkeletalAnimationTrack::GetDefaultDisplayName() const
 
 #endif
 
+TInlineValue<FMovieSceneSegmentCompilerRules> UMovieSceneSkeletalAnimationTrack::GetRowCompilerRules() const
+{
+	// Apply an upper bound exclusive blend
+	struct FSkeletalAnimationRowCompilerRules : FMovieSceneSegmentCompilerRules
+	{
+		FSkeletalAnimationRowCompilerRules(TInlineValue<FMovieSceneSegmentCompilerRules>&& InParentCompilerRules) { ParentCompilerRules = MoveTemp(InParentCompilerRules); }
+			
+		virtual void BlendSegment(FMovieSceneSegment& Segment, const TArrayView<const FMovieSceneSectionData>& SourceData) const
+		{
+			// Make the skeletal animation sections upper bound exclusive so that when you place animation 
+			// sections back to back, they don't blend at the transition times. This might be an option that 
+			// is exposed to the user in the future.
+
+			MovieSceneSegmentCompiler::BlendSegmentUpperBoundExclusive(Segment, SourceData);
+
+			ParentCompilerRules.GetValue().BlendSegment(Segment, SourceData);
+		}
+
+		TInlineValue<FMovieSceneSegmentCompilerRules> ParentCompilerRules;
+	};
+	return FSkeletalAnimationRowCompilerRules(Super::GetRowCompilerRules());
+}
+
+void UMovieSceneSkeletalAnimationTrack::PostCompile(FMovieSceneEvaluationTrack& OutTrack, const FMovieSceneTrackCompilerArgs& Args) const
+{
+	FMovieSceneSharedDataId UniqueId = FMovieSceneSkeletalAnimationSharedTrack::GetSharedDataKey().UniqueId;
+
+	FMovieSceneEvaluationTrack ActuatorTemplate(Args.ObjectBindingId);
+	ActuatorTemplate.DefineAsSingleTemplate(FMovieSceneSkeletalAnimationSharedTrack());
+	ActuatorTemplate.SetEvaluationPriority(ActuatorTemplate.GetEvaluationPriority() - 1);
+	Args.Generator.AddSharedTrack(MoveTemp(ActuatorTemplate), UniqueId, *this);
+}
 
 #undef LOCTEXT_NAMESPACE
diff --git a/Engine/Source/Runtime/MovieSceneTracks/Public/Evaluation/MovieSceneParameterTemplate.h b/Engine/Source/Runtime/MovieSceneTracks/Public/Evaluation/MovieSceneParameterTemplate.h
--- a/Engine/Source/Runtime/MovieSceneTracks/Public/Evaluation/MovieSceneParameterTemplate.h
+++ b/Engine/Source/Runtime/MovieSceneTracks/Public/Evaluation/MovieSceneParameterTemplate.h
@@ -6,11 +6,12 @@
 #include "Stats/Stats.h"
 #include "UObject/ObjectMacros.h"
 #include "MovieSceneFwd.h"
-
 #include "MovieSceneExecutionToken.h"
 #include "Evaluation/MovieSceneEvalTemplate.h"
 #include "Sections/MovieSceneParameterSection.h"
 #include "Materials/MaterialInstanceDynamic.h"
+#include "IMovieScenePlayer.h"
+
 #include "MovieSceneParameterTemplate.generated.h"
 
 class UMovieSceneComponentMaterialTrack;
diff --git a/Engine/Source/Runtime/MovieSceneTracks/Public/Sections/MovieSceneAudioSection.h b/Engine/Source/Runtime/MovieSceneTracks/Public/Sections/MovieSceneAudioSection.h
--- a/Engine/Source/Runtime/MovieSceneTracks/Public/Sections/MovieSceneAudioSection.h
+++ b/Engine/Source/Runtime/MovieSceneTracks/Public/Sections/MovieSceneAudioSection.h
@@ -27,41 +27,67 @@ public:
 	/** Gets the sound for this section */
 	class USoundBase* GetSound() const {return Sound;}
 	
-	/** Sets the time that the sound is supposed to be played at */
-	void SetAudioStartTime(float InAudioStartTime) {AudioStartTime = InAudioStartTime;}
+	/** Set the offset into the beginning of the audio clip */
+	void SetStartOffset(float InStartOffset) {StartOffset = InStartOffset;}
 	
-	/** Gets the (absolute) time that the sound is supposed to be played at */
-	float GetAudioStartTime() const {return AudioStartTime;}
+	/** Get the offset into the beginning of the audio clip */
+	float GetStartOffset() const {return StartOffset;}
 	
 	/**
 	 * @return The range of times that the sound plays, truncated by the section limits
 	 */
 	TRange<float> MOVIESCENETRACKS_API GetAudioRange() const;
 	
+	DEPRECATED(4.15, "Audio true range no longer supported.")
+	TRange<float> GetAudioTrueRange() const;
+	
 	/**
-	 * @return The range of times that the sound plays, not truncated by the section limits
+	 * Gets the sound volume curve
+	 *
+	 * @return The rich curve for this sound volume
 	 */
-	TRange<float> MOVIESCENETRACKS_API GetAudioTrueRange() const;
-	
+	FRichCurve& GetSoundVolumeCurve() { return SoundVolume; }
+	const FRichCurve& GetSoundVolumeCurve() const { return SoundVolume; }
+
 	/**
-	 * @return The dilation multiplier of the sound
+	 * Gets the sound pitch curve
+	 *
+	 * @return The rich curve for this sound pitch
 	 */
-	float GetAudioDilationFactor() const {return AudioDilationFactor;}
+	FRichCurve& GetPitchMultiplierCurve() { return PitchMultiplier; }
+	const FRichCurve& GetPitchMultiplierCurve() const { return PitchMultiplier; }
 
 	/**
-	* Sets the dilation multiplier of the sound
-	*/
-	void SetAudioDilationFactor( float InAudioDilationFactor ) { AudioDilationFactor = InAudioDilationFactor; }
+	 * Return the sound volume
+	 *
+	 * @param InTime	The position in time within the movie scene
+	 * @return The volume the sound will be played with.
+	 */
+	float GetSoundVolume(float InTime) const { return SoundVolume.Eval(InTime); }
 
 	/**
-	* @return The volume the sound will be played with.
-	*/
-	float GetAudioVolume() const { return AudioVolume; }
+	 * Sets the sound volume
+	 *
+	 * @param InTime	The position in time within the movie scene
+	 * @param InVolume	The volume to set
+	 */
+	void SetSoundVolume(float InTime, float InVolume) { SoundVolume.AddKey(InTime, InVolume); }
 
 	/**
-	* Sets the volume the sound will be played with.
-	*/
-	void SetAudioVolume( float InAudioVolume ) { AudioVolume = InAudioVolume; }
+	 * Return the pitch multiplier
+	 *
+	 * @param Position	The position in time within the movie scene
+	 * @return The pitch multiplier the sound will be played with.
+	 */
+	float GetPitchMultiplier(float InTime) const { return PitchMultiplier.Eval(InTime); }
+
+	/**
+	 * Set the pitch multiplier
+	 *
+	 * @param Position	The position in time within the movie scene
+	 * @param InPitch The pitch multiplier to set
+	 */
+	void SetPitchMultiplier(float InTime, float InPitchMultiplier) { PitchMultiplier.AddKey(InTime, InPitchMultiplier); }
 
 	/**
 	 * Returns whether or not a provided position in time is within the timespan of the audio range
@@ -92,6 +118,8 @@ public:
 	{
 		return bSuppressSubtitles;
 	}
+	/** ~UObject interface */
+	virtual void PostLoad() override;
 
 public:
 
@@ -99,6 +127,7 @@ public:
 
 	virtual void MoveSection( float DeltaPosition, TSet<FKeyHandle>& KeyHandles ) override;
 	virtual void DilateSection( float DilationFactor, float Origin, TSet<FKeyHandle>& KeyHandles ) override;
+	virtual UMovieSceneSection* SplitSection(float SplitTime) override;
 	virtual void GetKeyHandles(TSet<FKeyHandle>& OutKeyHandles, TRange<float> TimeRange) const override;
 	virtual void GetSnapTimes(TArray<float>& OutSnapTimes, bool bGetSectionBorders) const override;
 	virtual TOptional<float> GetKeyTime( FKeyHandle KeyHandle ) const override { return TOptional<float>(); }
@@ -111,17 +140,29 @@ private:
 	UPROPERTY(EditAnywhere, Category="Audio")
 	USoundBase* Sound;
 
-	/** The absolute time that the sound starts playing at */
+	/** The offset into the beginning of the audio clip */
 	UPROPERTY(EditAnywhere, Category="Audio")
-	float AudioStartTime;
+	float StartOffset;
+
+	/** The absolute time that the sound starts playing at */
+	UPROPERTY( )
+	float AudioStartTime_DEPRECATED;
 	
 	/** The amount which this audio is time dilated by */
-	UPROPERTY(EditAnywhere, Category="Audio")
-	float AudioDilationFactor;
+	UPROPERTY( )
+	float AudioDilationFactor_DEPRECATED;
+
+	/** The volume the sound will be played with. */
+	UPROPERTY( )
+	float AudioVolume_DEPRECATED;
 
 	/** The volume the sound will be played with. */
 	UPROPERTY( EditAnywhere, Category = "Audio" )
-	float AudioVolume;
+	FRichCurve SoundVolume;
+
+	/** The pitch multiplier the sound will be played with. */
+	UPROPERTY( EditAnywhere, Category = "Audio" )
+	FRichCurve PitchMultiplier;
 
 #if WITH_EDITORONLY_DATA
 	/** Whether to show the actual intensity of the wave on the thumbnail, as well as the smoothed RMS */
diff --git a/Engine/Source/Runtime/MovieSceneTracks/Public/Sections/MovieSceneSkeletalAnimationSection.h b/Engine/Source/Runtime/MovieSceneTracks/Public/Sections/MovieSceneSkeletalAnimationSection.h
--- a/Engine/Source/Runtime/MovieSceneTracks/Public/Sections/MovieSceneSkeletalAnimationSection.h
+++ b/Engine/Source/Runtime/MovieSceneTracks/Public/Sections/MovieSceneSkeletalAnimationSection.h
@@ -45,6 +45,10 @@ struct FMovieSceneSkeletalAnimationParams
 	/** The slot name to use for the animation */
 	UPROPERTY( EditAnywhere, Category = "Animation" )
 	FName SlotName;
+
+	/** The weight curve for this animation section */
+	UPROPERTY( EditAnywhere, Category = "Animation" )
+	FRichCurve Weight;	
 };
 
 /**
@@ -73,10 +77,11 @@ public:
 	virtual TOptional<float> GetKeyTime( FKeyHandle KeyHandle ) const override { return TOptional<float>(); }
 	virtual void SetKeyTime( FKeyHandle KeyHandle, float Time ) override { }
 	virtual FMovieSceneEvalTemplatePtr GenerateTemplate() const override;
+	virtual void PostLoadUpgradeTrackRow(const TRange<float>& InEvaluationRange) override;
 
 	/** ~UObject interface */
 	virtual void PostLoad() override;
-	
+	virtual void Serialize(FArchive& Ar) override;
 private:
 
 	//~ UObject interface
diff --git a/Engine/Source/Runtime/MovieSceneTracks/Public/Tracks/MovieSceneSkeletalAnimationTrack.h b/Engine/Source/Runtime/MovieSceneTracks/Public/Tracks/MovieSceneSkeletalAnimationTrack.h
--- a/Engine/Source/Runtime/MovieSceneTracks/Public/Tracks/MovieSceneSkeletalAnimationTrack.h
+++ b/Engine/Source/Runtime/MovieSceneTracks/Public/Tracks/MovieSceneSkeletalAnimationTrack.h
@@ -34,9 +34,11 @@ public:
 	virtual void RemoveSection(UMovieSceneSection& Section) override;
 	virtual bool IsEmpty() const override;
 	virtual TRange<float> GetSectionBoundaries() const override;
-	virtual bool SupportsMultipleRows() const override { return true; }
 	virtual const TArray<UMovieSceneSection*>& GetAllSections() const override;
 	virtual UMovieSceneSection* CreateNewSection() override;
+	virtual TInlineValue<FMovieSceneSegmentCompilerRules> GetRowCompilerRules() const override;
+	virtual void PostCompile(FMovieSceneEvaluationTrack& Track, const FMovieSceneTrackCompilerArgs& Args) const override;
+	virtual bool SupportsMultipleRows() const override;
 
 #if WITH_EDITORONLY_DATA
 	virtual FText GetDefaultDisplayName() const override;
diff --git a/Engine/Source/Runtime/Projects/Private/PluginManager.cpp b/Engine/Source/Runtime/Projects/Private/PluginManager.cpp
--- a/Engine/Source/Runtime/Projects/Private/PluginManager.cpp
+++ b/Engine/Source/Runtime/Projects/Private/PluginManager.cpp
@@ -627,7 +627,7 @@ bool FPluginManager::LoadModulesForEnabledPlugins( const ELoadingPhase::Type Loa
 					}
 					else if ( FailureReason == EModuleLoadResult::FailedToInitialize )
 					{
-						FailureMessage = FText::Format( LOCTEXT("PluginModuleFailedToInitialize", "Plugin '{0}' failed to load because module '{1}' could be initialized successfully after it was loaded."), PluginNameText, TextModuleName );
+						FailureMessage = FText::Format( LOCTEXT("PluginModuleFailedToInitialize", "Plugin '{0}' failed to load because module '{1}' could not be initialized successfully after it was loaded."), PluginNameText, TextModuleName );
 					}
 					else 
 					{
diff --git a/Engine/Source/Runtime/Slate/Private/Framework/Docking/TabManager.cpp b/Engine/Source/Runtime/Slate/Private/Framework/Docking/TabManager.cpp
--- a/Engine/Source/Runtime/Slate/Private/Framework/Docking/TabManager.cpp
+++ b/Engine/Source/Runtime/Slate/Private/Framework/Docking/TabManager.cpp
@@ -17,6 +17,7 @@
 #include "Widgets/Docking/SDockTab.h"
 #include "Framework/Docking/SDockingTabStack.h"
 #include "Framework/Docking/SDockingTabWell.h"
+#include "LayoutExtender.h"
 #if PLATFORM_MAC
 #include "../MultiBox/Mac/MacMenu.h"
 #endif
@@ -387,6 +388,96 @@ TSharedRef<FJsonObject> FTabManager::FLayout::PersistToString_Helper(const TShar
 	return JsonObj;
 }
 
+void FTabManager::FLayout::ProcessExtensions(const FLayoutExtender& Extender)
+{
+	struct FTabInformation
+	{
+		FTabInformation(FTabManager::FLayout& Layout)
+		{
+			Gather(Layout);
+		}
+
+		void Gather(FTabManager::FLayout& Layout)
+		{
+			for (TSharedRef<FTabManager::FArea>& Area : Layout.Areas)
+			{
+				Gather(*Area);
+			}
+		}
+
+		void Gather(FTabManager::FSplitter& Splitter)
+		{
+			for (TSharedRef<FTabManager::FLayoutNode>& Child : Splitter.ChildNodes)
+			{
+				TSharedPtr<FTabManager::FStack> Stack = Child->AsStack();
+				if (Stack.IsValid())
+				{
+					AllStacks.Add(Stack.Get());
+
+					for (FTabManager::FTab& Tab : Stack->Tabs)
+					{
+						AllDefinedTabs.Add(Tab.TabId);
+					}
+
+					continue;
+				}
+
+				TSharedPtr<FTabManager::FSplitter> ChildSplitter = Child->AsSplitter();
+				if (ChildSplitter.IsValid())
+				{
+					Gather(*ChildSplitter);
+					continue;
+				}
+
+				TSharedPtr<FTabManager::FArea> Area = Child->AsArea();
+				if (Area.IsValid())
+				{
+					Gather(*Area);
+					continue;
+				}
+			}
+		}
+
+		bool Contains(FTabId TabId) const
+		{
+			return AllDefinedTabs.Contains(TabId);
+		}
+
+		TArray<FTabManager::FStack*> AllStacks;
+		TSet<FTabId> AllDefinedTabs;
+	};
+	FTabInformation AllTabs(*this);
+
+	TArray<FTab, TInlineAllocator<1>> ExtendedTabs;
+
+	for (FTabManager::FStack* Stack : AllTabs.AllStacks)
+	{
+		for (int32 TabIndex = 0; TabIndex < Stack->Tabs.Num();)
+		{
+			FTabId TabId = Stack->Tabs[TabIndex].TabId;
+
+			Extender.FindExtensions(TabId, ELayoutExtensionPosition::Before, ExtendedTabs);
+			for (FTab& NewTab : ExtendedTabs)
+			{
+				if (!AllTabs.Contains(NewTab.TabId))
+				{
+					Stack->Tabs.Insert(NewTab, TabIndex++);
+				}
+			}
+
+			++TabIndex;
+
+			Extender.FindExtensions(TabId, ELayoutExtensionPosition::After, ExtendedTabs);
+			for (FTab& NewTab : ExtendedTabs)
+			{
+				if (!AllTabs.Contains(NewTab.TabId))
+				{
+					Stack->Tabs.Insert(NewTab, TabIndex++);
+				}
+			}
+		}
+	}
+}
 
 //////////////////////////////////////////////////////////////////////////
 // FTabManager::PrivateApi
diff --git a/Engine/Source/Runtime/Slate/Public/Framework/Docking/TabManager.h b/Engine/Source/Runtime/Slate/Public/Framework/Docking/TabManager.h
--- a/Engine/Source/Runtime/Slate/Public/Framework/Docking/TabManager.h
+++ b/Engine/Source/Runtime/Slate/Public/Framework/Docking/TabManager.h
@@ -15,6 +15,7 @@ class FMultiBox;
 class FProxyTabmanager;
 class SDockingArea;
 class SDockingTabStack;
+class FLayoutExtender;
 struct FTabMatcher;
 
 DECLARE_MULTICAST_DELEGATE_TwoParams(
@@ -70,6 +71,11 @@ struct FTabId
 			: FText::Format( NSLOCTEXT("TabManager", "TabIdFormat", "{TabType} : {InstanceIdNumber}"), Args );
 	}
 
+	friend uint32 GetTypeHash(const FTabId& In)
+	{
+		return GetTypeHash(In.TabType) ^ In.InstanceId;
+	}
+
 	FName TabType;
 	int32 InstanceId;
 };
@@ -453,7 +459,9 @@ class SLATE_API FTabManager : public TSharedFromThis<FTabManager>
 				FName GetLayoutName() const;
 				FString ToString() const;
 
-			protected:				
+				void ProcessExtensions(const FLayoutExtender& Extender);
+
+			protected:
 				static TSharedRef<class FJsonObject> PersistToString_Helper(const TSharedRef<FLayoutNode>& NodeToPersist);
 				static TSharedRef<FLayoutNode> NewFromString_Helper( TSharedPtr<FJsonObject> JsonObject );
 
diff --git a/Engine/Source/Runtime/Slate/Public/Framework/MultiBox/MultiBoxBuilder.h b/Engine/Source/Runtime/Slate/Public/Framework/MultiBox/MultiBoxBuilder.h
--- a/Engine/Source/Runtime/Slate/Public/Framework/MultiBox/MultiBoxBuilder.h
+++ b/Engine/Source/Runtime/Slate/Public/Framework/MultiBox/MultiBoxBuilder.h
@@ -169,7 +169,7 @@ public:
 	 * @param	InExtensionHook		The section hook. Can be NAME_None
 	 * @param	InLabelOverride		Optional label override.  If omitted, then the action's label will be used instead.
 	 * @param	InToolTipOverride	Optional tool tip override.	 If omitted, then the action's label will be used instead.
-	 * @param	InIconOverride		Optional name of the slate brush to use for the tool bar image.  If omitted, then the action's icon will be used instead.
+	 * @param	InIconOverride		Optional name of the slate brush to use for the tool bar image.  If omitted, then the command's icon will be used instead.
 	 * @param	InTutorialHighlightName	Optional name to identify this widget and highlight during tutorials
 	 */
 	void AddMenuEntry( const TSharedPtr< const FUICommandInfo > InCommand, FName InExtensionHook = NAME_None, const TAttribute<FText>& InLabelOverride = TAttribute<FText>(), const TAttribute<FText>& InToolTipOverride = TAttribute<FText>(), const FSlateIcon& InIconOverride = FSlateIcon(), FName InTutorialHighlightName = NAME_None );
diff --git a/Engine/Source/Runtime/UMG/Private/Animation/MovieScene2DTransformTemplate.cpp b/Engine/Source/Runtime/UMG/Private/Animation/MovieScene2DTransformTemplate.cpp
--- a/Engine/Source/Runtime/UMG/Private/Animation/MovieScene2DTransformTemplate.cpp
+++ b/Engine/Source/Runtime/UMG/Private/Animation/MovieScene2DTransformTemplate.cpp
@@ -1,9 +1,10 @@
 // Copyright 1998-2017 Epic Games, Inc. All Rights Reserved.
 
 #include "Animation/MovieScene2DTransformTemplate.h"
-
 #include "Animation/MovieScene2DTransformSection.h"
 #include "Animation/MovieScene2DTransformTrack.h"
+#include "MovieSceneEvaluation.h"
+
 
 FMovieScene2DTransformSectionTemplate::FMovieScene2DTransformSectionTemplate(const UMovieScene2DTransformSection& Section, const UMovieScenePropertyTrack& Track)
 	: PropertyData(Track.GetPropertyName(), Track.GetPropertyPath())
diff --git a/Engine/Source/Runtime/UMG/Private/Animation/MovieSceneMarginTemplate.cpp b/Engine/Source/Runtime/UMG/Private/Animation/MovieSceneMarginTemplate.cpp
--- a/Engine/Source/Runtime/UMG/Private/Animation/MovieSceneMarginTemplate.cpp
+++ b/Engine/Source/Runtime/UMG/Private/Animation/MovieSceneMarginTemplate.cpp
@@ -4,6 +4,7 @@
 
 #include "Animation/MovieSceneMarginSection.h"
 #include "Tracks/MovieScenePropertyTrack.h"
+#include "MovieSceneEvaluation.h"
 
 FMovieSceneMarginSectionTemplate::FMovieSceneMarginSectionTemplate(const UMovieSceneMarginSection& Section, const UMovieScenePropertyTrack& Track)
 	: PropertyData(Track.GetPropertyName(), Track.GetPropertyPath())
diff --git a/Engine/Source/Runtime/UMG/Private/Animation/MovieSceneWidgetMaterialTemplate.cpp b/Engine/Source/Runtime/UMG/Private/Animation/MovieSceneWidgetMaterialTemplate.cpp
--- a/Engine/Source/Runtime/UMG/Private/Animation/MovieSceneWidgetMaterialTemplate.cpp
+++ b/Engine/Source/Runtime/UMG/Private/Animation/MovieSceneWidgetMaterialTemplate.cpp
@@ -4,6 +4,7 @@
 #include "Animation/WidgetMaterialTrackUtilities.h"
 #include "Animation/MovieSceneWidgetMaterialTrack.h"
 #include "Components/Widget.h"
+#include "MovieSceneEvaluation.h"
 
 // Container to ensure unique IDs per property path
 TMovieSceneAnimTypeIDContainer<TArray<FName>> BrushPropertyIDs;
diff --git a/Engine/Source/Runtime/UtilityShaders/Private/MediaShaders.cpp b/Engine/Source/Runtime/UtilityShaders/Private/MediaShaders.cpp
--- a/Engine/Source/Runtime/UtilityShaders/Private/MediaShaders.cpp
+++ b/Engine/Source/Runtime/UtilityShaders/Private/MediaShaders.cpp
@@ -7,6 +7,48 @@
 TGlobalResource<FMediaVertexDeclaration> GMediaVertexDeclaration;
 
 
+/* MediaShaders namespace
+ *****************************************************************************/
+
+namespace MediaShaders
+{
+	const FMatrix YuvToSrgbJpeg = FMatrix(
+		FPlane(1.000000f,  0.000000f,  1.402000f, 0.000000f),
+		FPlane(1.000000f, -0.344140f, -0.714140f, 0.000000f),
+		FPlane(1.000000f,  1.772000f,  0.000000f, 0.000000f),
+		FPlane(0.000000f,  0.000000f,  0.000000f, 0.000000f)
+	);
+
+	const FMatrix YuvToSrgbDefault = FMatrix(
+		FPlane(1.164383f,  0.000000f,  1.596027f, 0.000000f),
+		FPlane(1.164383f, -0.391762f, -0.812968f, 0.000000f),
+		FPlane(1.164383f,  2.017232f,  0.000000f, 0.000000f),
+		FPlane(0.000000f,  0.000000f,  0.000000f, 0.000000f)
+	);
+
+	const FMatrix YuvToSrgbPs4 = FMatrix(
+		FPlane(1.164400f,  0.000000f,  1.792700f, 0.000000f),
+		FPlane(1.164400f, -0.213300f, -0.532900f, 0.000000f),
+		FPlane(1.164400f,  2.112400f,  0.000000f, 0.000000f),
+		FPlane(0.000000f,  0.000000f,  0.000000f, 0.000000f)
+	);
+
+	const FMatrix YuvToSrgbRec601 = FMatrix(
+		FPlane(1.000000f, 0.000000f, 1.139830f, 0.000000f),
+		FPlane(1.000000f, -0.394650f, -0.580600f, 0.000000f),
+		FPlane(1.000000f, 2.032110, 0.000000, 0.000000f),
+		FPlane(0.000000f, 0.000000f, 0.000000f, 0.000000f)
+	);
+
+	const FMatrix YuvToRgbRec709 = FMatrix(
+		FPlane(1.000000f, 0.000000f, 1.280330f, 0.000000f),
+		FPlane(1.000000f, -0.214820f, -0.380590f, 0.000000f),
+		FPlane(1.000000f, 2.127980f, 0.000000f, 0.000000f),
+		FPlane(0.000000f, 0.000000f, 0.000000f, 0.000000f)
+	);
+}
+
+
 /* FMediaShadersVS shader
  *****************************************************************************/
 
@@ -17,24 +59,28 @@ IMPLEMENT_SHADER_TYPE(, FMediaShadersVS, TEXT("MediaShaders"), TEXT("MainVertexS
  *****************************************************************************/
 
 BEGIN_UNIFORM_BUFFER_STRUCT(FAYUVConvertUB, )
-DECLARE_UNIFORM_BUFFER_STRUCT_MEMBER_TEXTURE(Texture2D, AYUVTexture)
-DECLARE_UNIFORM_BUFFER_STRUCT_MEMBER_SAMPLER(SamplerState, AYUVSampler)
+DECLARE_UNIFORM_BUFFER_STRUCT_MEMBER(FMatrix, ColorTransform)
+DECLARE_UNIFORM_BUFFER_STRUCT_MEMBER(uint32, SrgbToLinear)
+DECLARE_UNIFORM_BUFFER_STRUCT_MEMBER_TEXTURE(Texture2D, Texture)
+DECLARE_UNIFORM_BUFFER_STRUCT_MEMBER_SAMPLER(SamplerState, Sampler)
 END_UNIFORM_BUFFER_STRUCT(FAYUVConvertUB)
 
 IMPLEMENT_UNIFORM_BUFFER_STRUCT(FAYUVConvertUB, TEXT("AYUVConvertUB"));
 IMPLEMENT_SHADER_TYPE(, FAYUVConvertPS, TEXT("MediaShaders"), TEXT("AYUVConvertPS"), SF_Pixel);
 
 
-void FAYUVConvertPS::SetParameters(FRHICommandList& RHICmdList, TRefCountPtr<FRHITexture2D> AYUVTexture)
+void FAYUVConvertPS::SetParameters(FRHICommandList& CommandList, TRefCountPtr<FRHITexture2D> AYUVTexture, const FMatrix& ColorTransform, bool SrgbToLinear)
 {
 	FAYUVConvertUB UB;
 	{
-		UB.AYUVTexture = AYUVTexture;
-		UB.AYUVSampler = TStaticSamplerState<SF_Bilinear>::GetRHI();
+		UB.ColorTransform = ColorTransform;
+		UB.Sampler = TStaticSamplerState<SF_Bilinear>::GetRHI();
+		UB.SrgbToLinear = SrgbToLinear;
+		UB.Texture = AYUVTexture;
 	}
 
 	TUniformBufferRef<FAYUVConvertUB> Data = TUniformBufferRef<FAYUVConvertUB>::CreateUniformBufferImmediate(UB, UniformBuffer_SingleFrame);
-	SetUniformBufferParameter(RHICmdList, GetPixelShader(), GetUniformBufferParameter<FAYUVConvertUB>(), Data);
+	SetUniformBufferParameter(CommandList, GetPixelShader(), GetUniformBufferParameter<FAYUVConvertUB>(), Data);
 }
 
 
@@ -43,25 +89,25 @@ void FAYUVConvertPS::SetParameters(FRHICommandList& RHICmdList, TRefCountPtr<FRH
 
 BEGIN_UNIFORM_BUFFER_STRUCT(FBMPConvertUB, )
 DECLARE_UNIFORM_BUFFER_STRUCT_MEMBER(FVector2D, UVScale)
-DECLARE_UNIFORM_BUFFER_STRUCT_MEMBER_TEXTURE(Texture2D, BMPTexture)
-DECLARE_UNIFORM_BUFFER_STRUCT_MEMBER_SAMPLER(SamplerState, BMPSampler)
+DECLARE_UNIFORM_BUFFER_STRUCT_MEMBER_TEXTURE(Texture2D, Texture)
+DECLARE_UNIFORM_BUFFER_STRUCT_MEMBER_SAMPLER(SamplerState, Sampler)
 END_UNIFORM_BUFFER_STRUCT(FBMPConvertUB)
 
 IMPLEMENT_UNIFORM_BUFFER_STRUCT(FBMPConvertUB, TEXT("BMPConvertUB"));
 IMPLEMENT_SHADER_TYPE(, FBMPConvertPS, TEXT("MediaShaders"), TEXT("BMPConvertPS"), SF_Pixel);
 
 
-void FBMPConvertPS::SetParameters(FRHICommandList& RHICmdList, TRefCountPtr<FRHITexture2D> BMPTexture, const FIntPoint& OutputDimensions)
+void FBMPConvertPS::SetParameters(FRHICommandList& CommandList, TRefCountPtr<FRHITexture2D> BMPTexture, const FIntPoint& OutputDimensions)
 {
 	FBMPConvertUB UB;
 	{
 		UB.UVScale = FVector2D((float)OutputDimensions.X / (float)BMPTexture->GetSizeX(), (float)OutputDimensions.Y / (float)BMPTexture->GetSizeY());
-		UB.BMPTexture = BMPTexture;
-		UB.BMPSampler = TStaticSamplerState<SF_Bilinear>::GetRHI();
+		UB.Texture = BMPTexture;
+		UB.Sampler = TStaticSamplerState<SF_Bilinear>::GetRHI();
 	}
 
 	TUniformBufferRef<FBMPConvertUB> Data = TUniformBufferRef<FBMPConvertUB>::CreateUniformBufferImmediate(UB, UniformBuffer_SingleFrame);
-	SetUniformBufferParameter(RHICmdList, GetPixelShader(), GetUniformBufferParameter<FBMPConvertUB>(), Data);
+	SetUniformBufferParameter(CommandList, GetPixelShader(), GetUniformBufferParameter<FBMPConvertUB>(), Data);
 }
 
 
@@ -69,30 +115,34 @@ void FBMPConvertPS::SetParameters(FRHICommandList& RHICmdList, TRefCountPtr<FRHI
  *****************************************************************************/
 
 BEGIN_UNIFORM_BUFFER_STRUCT(FNV12ConvertUB, )
+DECLARE_UNIFORM_BUFFER_STRUCT_MEMBER(FMatrix, ColorTransform)
 DECLARE_UNIFORM_BUFFER_STRUCT_MEMBER(uint32, OutputWidth)
+DECLARE_UNIFORM_BUFFER_STRUCT_MEMBER(uint32, SrgbToLinear)
 DECLARE_UNIFORM_BUFFER_STRUCT_MEMBER(FVector2D, UVScale)
-DECLARE_UNIFORM_BUFFER_STRUCT_MEMBER_TEXTURE(Texture2D, NV12Texture)
-DECLARE_UNIFORM_BUFFER_STRUCT_MEMBER_SAMPLER(SamplerState, NV12SamplerB)
-DECLARE_UNIFORM_BUFFER_STRUCT_MEMBER_SAMPLER(SamplerState, NV12SamplerP)
+DECLARE_UNIFORM_BUFFER_STRUCT_MEMBER_TEXTURE(Texture2D, Texture)
+DECLARE_UNIFORM_BUFFER_STRUCT_MEMBER_SAMPLER(SamplerState, SamplerB)
+DECLARE_UNIFORM_BUFFER_STRUCT_MEMBER_SAMPLER(SamplerState, SamplerP)
 END_UNIFORM_BUFFER_STRUCT(FNV12ConvertUB)
 
 IMPLEMENT_UNIFORM_BUFFER_STRUCT(FNV12ConvertUB, TEXT("NV12ConvertUB"));
 IMPLEMENT_SHADER_TYPE(, FNV12ConvertPS, TEXT("MediaShaders"), TEXT("NV12ConvertPS"), SF_Pixel);
 
 
-void FNV12ConvertPS::SetParameters(FRHICommandList& RHICmdList, TRefCountPtr<FRHITexture2D> NV12Texture, const FIntPoint& OutputDimensions)
+void FNV12ConvertPS::SetParameters(FRHICommandList& CommandList, TRefCountPtr<FRHITexture2D> NV12Texture, const FIntPoint& OutputDimensions, const FMatrix& ColorTransform, bool SrgbToLinear)
 {
 	FNV12ConvertUB UB;
 	{
+		UB.ColorTransform = ColorTransform;
 		UB.OutputWidth = OutputDimensions.X;
+		UB.SamplerB = TStaticSamplerState<SF_Bilinear>::GetRHI();
+		UB.SamplerP = TStaticSamplerState<SF_Point>::GetRHI();
+		UB.SrgbToLinear = SrgbToLinear;
+		UB.Texture = NV12Texture;
 		UB.UVScale = FVector2D((float)OutputDimensions.X / (float)NV12Texture->GetSizeX(), (float)OutputDimensions.Y / (float)NV12Texture->GetSizeY());
-		UB.NV12Texture = NV12Texture;
-		UB.NV12SamplerB = TStaticSamplerState<SF_Bilinear>::GetRHI();
-		UB.NV12SamplerP = TStaticSamplerState<SF_Point>::GetRHI();
 	}
 
 	TUniformBufferRef<FNV12ConvertUB> Data = TUniformBufferRef<FNV12ConvertUB>::CreateUniformBufferImmediate(UB, UniformBuffer_SingleFrame);
-	SetUniformBufferParameter(RHICmdList, GetPixelShader(), GetUniformBufferParameter<FNV12ConvertUB>(), Data);
+	SetUniformBufferParameter(CommandList, GetPixelShader(), GetUniformBufferParameter<FNV12ConvertUB>(), Data);
 }
 
 
@@ -100,30 +150,34 @@ void FNV12ConvertPS::SetParameters(FRHICommandList& RHICmdList, TRefCountPtr<FRH
  *****************************************************************************/
 
 BEGIN_UNIFORM_BUFFER_STRUCT(FNV21ConvertUB, )
+DECLARE_UNIFORM_BUFFER_STRUCT_MEMBER(FMatrix, ColorTransform)
 DECLARE_UNIFORM_BUFFER_STRUCT_MEMBER(uint32, OutputWidth)
+DECLARE_UNIFORM_BUFFER_STRUCT_MEMBER(uint32, SrgbToLinear)
 DECLARE_UNIFORM_BUFFER_STRUCT_MEMBER(FVector2D, UVScale)
-DECLARE_UNIFORM_BUFFER_STRUCT_MEMBER_TEXTURE(Texture2D, NV21Texture)
-DECLARE_UNIFORM_BUFFER_STRUCT_MEMBER_SAMPLER(SamplerState, NV21SamplerB)
-DECLARE_UNIFORM_BUFFER_STRUCT_MEMBER_SAMPLER(SamplerState, NV21SamplerP)
+DECLARE_UNIFORM_BUFFER_STRUCT_MEMBER_TEXTURE(Texture2D, Texture)
+DECLARE_UNIFORM_BUFFER_STRUCT_MEMBER_SAMPLER(SamplerState, SamplerB)
+DECLARE_UNIFORM_BUFFER_STRUCT_MEMBER_SAMPLER(SamplerState, SamplerP)
 END_UNIFORM_BUFFER_STRUCT(FNV21ConvertUB)
 
 IMPLEMENT_UNIFORM_BUFFER_STRUCT(FNV21ConvertUB, TEXT("NV21ConvertUB"));
 IMPLEMENT_SHADER_TYPE(, FNV21ConvertPS, TEXT("MediaShaders"), TEXT("NV21ConvertPS"), SF_Pixel);
 
 
-void FNV21ConvertPS::SetParameters(FRHICommandList& RHICmdList, TRefCountPtr<FRHITexture2D> NV21Texture, const FIntPoint& OutputDimensions)
+void FNV21ConvertPS::SetParameters(FRHICommandList& CommandList, TRefCountPtr<FRHITexture2D> NV21Texture, const FIntPoint& OutputDimensions, const FMatrix& ColorTransform, bool SrgbToLinear)
 {
 	FNV21ConvertUB UB;
 	{
+		UB.ColorTransform = ColorTransform;
 		UB.OutputWidth = OutputDimensions.Y;
+		UB.SamplerB = TStaticSamplerState<SF_Bilinear>::GetRHI();
+		UB.SamplerP = TStaticSamplerState<SF_Point>::GetRHI();
+		UB.SrgbToLinear = SrgbToLinear;
+		UB.Texture = NV21Texture;
 		UB.UVScale = FVector2D((float)OutputDimensions.X / (float)NV21Texture->GetSizeX(), (float)OutputDimensions.Y / (float)NV21Texture->GetSizeY());
-		UB.NV21Texture = NV21Texture;
-		UB.NV21SamplerB = TStaticSamplerState<SF_Bilinear>::GetRHI();
-		UB.NV21SamplerP = TStaticSamplerState<SF_Point>::GetRHI();
 	}
 
 	TUniformBufferRef<FNV21ConvertUB> Data = TUniformBufferRef<FNV21ConvertUB>::CreateUniformBufferImmediate(UB, UniformBuffer_SingleFrame);
-	SetUniformBufferParameter(RHICmdList, GetPixelShader(), GetUniformBufferParameter<FNV21ConvertUB>(), Data);
+	SetUniformBufferParameter(CommandList, GetPixelShader(), GetUniformBufferParameter<FNV21ConvertUB>(), Data);
 }
 
 
@@ -132,25 +186,26 @@ void FNV21ConvertPS::SetParameters(FRHICommandList& RHICmdList, TRefCountPtr<FRH
 
 BEGIN_UNIFORM_BUFFER_STRUCT(FRGBConvertUB, )
 DECLARE_UNIFORM_BUFFER_STRUCT_MEMBER(FVector2D, UVScale)
-DECLARE_UNIFORM_BUFFER_STRUCT_MEMBER_TEXTURE(Texture2D, RGBTexture)
-DECLARE_UNIFORM_BUFFER_STRUCT_MEMBER_SAMPLER(SamplerState, RGBSampler)
+DECLARE_UNIFORM_BUFFER_STRUCT_MEMBER(uint32, SrgbToLinear)
+DECLARE_UNIFORM_BUFFER_STRUCT_MEMBER_TEXTURE(Texture2D, Texture)
+DECLARE_UNIFORM_BUFFER_STRUCT_MEMBER_SAMPLER(SamplerState, Sampler)
 END_UNIFORM_BUFFER_STRUCT(FRGBConvertUB)
 
 IMPLEMENT_UNIFORM_BUFFER_STRUCT(FRGBConvertUB, TEXT("RGBConvertUB"));
 IMPLEMENT_SHADER_TYPE(, FRGBConvertPS, TEXT("MediaShaders"), TEXT("RGBConvertPS"), SF_Pixel);
 
 
-void FRGBConvertPS::SetParameters(FRHICommandList& RHICmdList, TRefCountPtr<FRHITexture2D> RGBTexture, const FIntPoint& OutputDimensions)
+void FRGBConvertPS::SetParameters(FRHICommandList& CommandList, TRefCountPtr<FRHITexture2D> RGBTexture, const FIntPoint& OutputDimensions, bool SrgbToLinear)
 {
 	FRGBConvertUB UB;
 	{
+		UB.Sampler = TStaticSamplerState<SF_Bilinear>::GetRHI();
+		UB.Texture = RGBTexture;
 		UB.UVScale = FVector2D((float)OutputDimensions.X / (float)RGBTexture->GetSizeX(), (float)OutputDimensions.Y / (float)RGBTexture->GetSizeY());
-		UB.RGBTexture = RGBTexture;
-		UB.RGBSampler = TStaticSamplerState<SF_Bilinear>::GetRHI();
 	}
 
 	TUniformBufferRef<FRGBConvertUB> Data = TUniformBufferRef<FRGBConvertUB>::CreateUniformBufferImmediate(UB, UniformBuffer_SingleFrame);
-	SetUniformBufferParameter(RHICmdList, GetPixelShader(), GetUniformBufferParameter<FRGBConvertUB>(), Data);
+	SetUniformBufferParameter(CommandList, GetPixelShader(), GetUniformBufferParameter<FRGBConvertUB>(), Data);
 }
 
 
@@ -158,6 +213,8 @@ void FRGBConvertPS::SetParameters(FRHICommandList& RHICmdList, TRefCountPtr<FRHI
  *****************************************************************************/
 
 BEGIN_UNIFORM_BUFFER_STRUCT(FYCbCrConvertUB, )
+DECLARE_UNIFORM_BUFFER_STRUCT_MEMBER(FMatrix, ColorTransform)
+DECLARE_UNIFORM_BUFFER_STRUCT_MEMBER(uint32, SrgbToLinear)
 DECLARE_UNIFORM_BUFFER_STRUCT_MEMBER_TEXTURE(Texture2D, LumaTexture)
 DECLARE_UNIFORM_BUFFER_STRUCT_MEMBER_TEXTURE(Texture2D, CbCrTexture)
 DECLARE_UNIFORM_BUFFER_STRUCT_MEMBER_SAMPLER(SamplerState, LumaSampler)
@@ -168,18 +225,20 @@ IMPLEMENT_UNIFORM_BUFFER_STRUCT(FYCbCrConvertUB, TEXT("YCbCrConvertUB"));
 IMPLEMENT_SHADER_TYPE(, FYCbCrConvertPS, TEXT("MediaShaders"), TEXT("YCbCrConvertPS"), SF_Pixel);
 
 
-void FYCbCrConvertPS::SetParameters(FRHICommandList& RHICmdList, TRefCountPtr<FRHITexture2D> LumaTexture, TRefCountPtr<FRHITexture2D> CbCrTexture)
+void FYCbCrConvertPS::SetParameters(FRHICommandList& CommandList, TRefCountPtr<FRHITexture2D> LumaTexture, TRefCountPtr<FRHITexture2D> CbCrTexture, const FMatrix& ColorTransform, bool SrgbToLinear)
 {
 	FYCbCrConvertUB UB;
 	{
-		UB.LumaTexture = LumaTexture;
+		UB.CbCrSampler = TStaticSamplerState<SF_Bilinear>::GetRHI();
 		UB.CbCrTexture = CbCrTexture;
+		UB.ColorTransform = ColorTransform;
 		UB.LumaSampler = TStaticSamplerState<SF_Bilinear>::GetRHI();
-		UB.CbCrSampler = TStaticSamplerState<SF_Bilinear>::GetRHI();
+		UB.LumaTexture = LumaTexture;
+		UB.SrgbToLinear = SrgbToLinear;
 	}
 
 	TUniformBufferRef<FYCbCrConvertUB> Data = TUniformBufferRef<FYCbCrConvertUB>::CreateUniformBufferImmediate(UB, UniformBuffer_SingleFrame);
-	SetUniformBufferParameter(RHICmdList, GetPixelShader(), GetUniformBufferParameter<FYCbCrConvertUB>(), Data);	
+	SetUniformBufferParameter(CommandList, GetPixelShader(), GetUniformBufferParameter<FYCbCrConvertUB>(), Data);	
 }
 
 
@@ -187,28 +246,32 @@ void FYCbCrConvertPS::SetParameters(FRHICommandList& RHICmdList, TRefCountPtr<FR
  *****************************************************************************/
 
 BEGIN_UNIFORM_BUFFER_STRUCT(FUYVYConvertUB, )
+DECLARE_UNIFORM_BUFFER_STRUCT_MEMBER(FMatrix, ColorTransform)
+DECLARE_UNIFORM_BUFFER_STRUCT_MEMBER(uint32, SrgbToLinear)
 DECLARE_UNIFORM_BUFFER_STRUCT_MEMBER(uint32, Width)
-DECLARE_UNIFORM_BUFFER_STRUCT_MEMBER_TEXTURE(Texture2D, UYVYTexture)
-DECLARE_UNIFORM_BUFFER_STRUCT_MEMBER_SAMPLER(SamplerState, UYVYSamplerB)
-DECLARE_UNIFORM_BUFFER_STRUCT_MEMBER_SAMPLER(SamplerState, UYVYSamplerP)
+DECLARE_UNIFORM_BUFFER_STRUCT_MEMBER_TEXTURE(Texture2D, Texture)
+DECLARE_UNIFORM_BUFFER_STRUCT_MEMBER_SAMPLER(SamplerState, SamplerB)
+DECLARE_UNIFORM_BUFFER_STRUCT_MEMBER_SAMPLER(SamplerState, SamplerP)
 END_UNIFORM_BUFFER_STRUCT(FUYVYConvertUB)
 
 IMPLEMENT_UNIFORM_BUFFER_STRUCT(FUYVYConvertUB, TEXT("UYVYConvertUB"));
 IMPLEMENT_SHADER_TYPE(, FUYVYConvertPS, TEXT("MediaShaders"), TEXT("UYVYConvertPS"), SF_Pixel);
 
 
-void FUYVYConvertPS::SetParameters(FRHICommandList& RHICmdList, TRefCountPtr<FRHITexture2D> UYVYTexture)
+void FUYVYConvertPS::SetParameters(FRHICommandList& CommandList, TRefCountPtr<FRHITexture2D> UYVYTexture, const FMatrix& ColorTransform, bool SrgbToLinear)
 {
 	FUYVYConvertUB UB;
 	{
+		UB.ColorTransform = ColorTransform;
+		UB.SamplerB = TStaticSamplerState<SF_Bilinear>::GetRHI();
+		UB.SamplerP = TStaticSamplerState<SF_Point>::GetRHI();
+		UB.SrgbToLinear = SrgbToLinear;
+		UB.Texture = UYVYTexture;
 		UB.Width = UYVYTexture->GetSizeX();
-		UB.UYVYTexture = UYVYTexture;
-		UB.UYVYSamplerB = TStaticSamplerState<SF_Bilinear>::GetRHI();
-		UB.UYVYSamplerP = TStaticSamplerState<SF_Point>::GetRHI();
 	}
 
 	TUniformBufferRef<FUYVYConvertUB> Data = TUniformBufferRef<FUYVYConvertUB>::CreateUniformBufferImmediate(UB, UniformBuffer_SingleFrame);
-	SetUniformBufferParameter(RHICmdList, GetPixelShader(), GetUniformBufferParameter<FUYVYConvertUB>(), Data);
+	SetUniformBufferParameter(CommandList, GetPixelShader(), GetUniformBufferParameter<FUYVYConvertUB>(), Data);
 }
 
 
@@ -216,6 +279,8 @@ void FUYVYConvertPS::SetParameters(FRHICommandList& RHICmdList, TRefCountPtr<FRH
  *****************************************************************************/
 
 BEGIN_UNIFORM_BUFFER_STRUCT(FYUVConvertUB, )
+DECLARE_UNIFORM_BUFFER_STRUCT_MEMBER(FMatrix, ColorTransform)
+DECLARE_UNIFORM_BUFFER_STRUCT_MEMBER(uint32, SrgbToLinear)
 DECLARE_UNIFORM_BUFFER_STRUCT_MEMBER_TEXTURE(Texture2D, YTexture)
 DECLARE_UNIFORM_BUFFER_STRUCT_MEMBER_TEXTURE(Texture2D, UTexture)
 DECLARE_UNIFORM_BUFFER_STRUCT_MEMBER_TEXTURE(Texture2D, VTexture)
@@ -228,10 +293,12 @@ IMPLEMENT_UNIFORM_BUFFER_STRUCT(FYUVConvertUB, TEXT("YUVConvertUB"));
 IMPLEMENT_SHADER_TYPE(, FYUVConvertPS, TEXT("MediaShaders"), TEXT("YUVConvertPS"), SF_Pixel);
 
 
-void FYUVConvertPS::SetParameters(FRHICommandList& RHICmdList, TRefCountPtr<FRHITexture2D> YTexture, TRefCountPtr<FRHITexture2D> UTexture, TRefCountPtr<FRHITexture2D> VTexture)
+void FYUVConvertPS::SetParameters(FRHICommandList& CommandList, TRefCountPtr<FRHITexture2D> YTexture, TRefCountPtr<FRHITexture2D> UTexture, TRefCountPtr<FRHITexture2D> VTexture, const FMatrix& ColorTransform, bool SrgbToLinear)
 {
 	FYUVConvertUB UB;
 	{
+		UB.ColorTransform = ColorTransform;
+		UB.SrgbToLinear = SrgbToLinear;
 		UB.YTexture = YTexture;
 		UB.UTexture = UTexture;
 		UB.VTexture = VTexture;
@@ -241,7 +308,7 @@ void FYUVConvertPS::SetParameters(FRHICommandList& RHICmdList, TRefCountPtr<FRHI
 	}
 
 	TUniformBufferRef<FYUVConvertUB> Data = TUniformBufferRef<FYUVConvertUB>::CreateUniformBufferImmediate(UB, UniformBuffer_SingleFrame);
-	SetUniformBufferParameter(RHICmdList, GetPixelShader(), GetUniformBufferParameter<FYUVConvertUB>(), Data);
+	SetUniformBufferParameter(CommandList, GetPixelShader(), GetUniformBufferParameter<FYUVConvertUB>(), Data);
 }
 
 
@@ -249,57 +316,65 @@ void FYUVConvertPS::SetParameters(FRHICommandList& RHICmdList, TRefCountPtr<FRHI
  *****************************************************************************/
 
 BEGIN_UNIFORM_BUFFER_STRUCT(FYUY2ConvertUB, )
+DECLARE_UNIFORM_BUFFER_STRUCT_MEMBER(FMatrix, ColorTransform)
 DECLARE_UNIFORM_BUFFER_STRUCT_MEMBER(uint32, OutputWidth)
+DECLARE_UNIFORM_BUFFER_STRUCT_MEMBER(uint32, SrgbToLinear)
 DECLARE_UNIFORM_BUFFER_STRUCT_MEMBER(FVector2D, UVScale)
-DECLARE_UNIFORM_BUFFER_STRUCT_MEMBER_TEXTURE(Texture2D, YUY2Texture)
-DECLARE_UNIFORM_BUFFER_STRUCT_MEMBER_SAMPLER(SamplerState, YUY2SamplerB)
-DECLARE_UNIFORM_BUFFER_STRUCT_MEMBER_SAMPLER(SamplerState, YUY2SamplerP)
+DECLARE_UNIFORM_BUFFER_STRUCT_MEMBER_TEXTURE(Texture2D, Texture)
+DECLARE_UNIFORM_BUFFER_STRUCT_MEMBER_SAMPLER(SamplerState, SamplerB)
+DECLARE_UNIFORM_BUFFER_STRUCT_MEMBER_SAMPLER(SamplerState, SamplerP)
 END_UNIFORM_BUFFER_STRUCT(FYUY2ConvertUB)
 
 IMPLEMENT_UNIFORM_BUFFER_STRUCT(FYUY2ConvertUB, TEXT("YUY2ConvertUB"));
 IMPLEMENT_SHADER_TYPE(, FYUY2ConvertPS, TEXT("MediaShaders"), TEXT("YUY2ConvertPS"), SF_Pixel);
 
 
-void FYUY2ConvertPS::SetParameters(FRHICommandList& RHICmdList, TRefCountPtr<FRHITexture2D> YUY2Texture, const FIntPoint& OutputDimensions)
+void FYUY2ConvertPS::SetParameters(FRHICommandList& CommandList, TRefCountPtr<FRHITexture2D> YUY2Texture, const FIntPoint& OutputDimensions, const FMatrix& ColorTransform, bool SrgbToLinear)
 {
 	FYUY2ConvertUB UB;
 	{
+		UB.ColorTransform = ColorTransform;
 		UB.OutputWidth = OutputDimensions.X;
+		UB.SamplerB = TStaticSamplerState<SF_Bilinear>::GetRHI();
+		UB.SamplerP = TStaticSamplerState<SF_Point>::GetRHI();
+		UB.SrgbToLinear = SrgbToLinear;
+		UB.Texture = YUY2Texture;
 		UB.UVScale = FVector2D((float)OutputDimensions.X / (2.0f * YUY2Texture->GetSizeX()), (float)OutputDimensions.Y / (float)YUY2Texture->GetSizeY());
-		UB.YUY2Texture = YUY2Texture;
-		UB.YUY2SamplerB = TStaticSamplerState<SF_Bilinear>::GetRHI();
-		UB.YUY2SamplerP = TStaticSamplerState<SF_Point>::GetRHI();
 	}
 
 	TUniformBufferRef<FYUY2ConvertUB> Data = TUniformBufferRef<FYUY2ConvertUB>::CreateUniformBufferImmediate(UB, UniformBuffer_SingleFrame);
-	SetUniformBufferParameter(RHICmdList, GetPixelShader(), GetUniformBufferParameter<FYUY2ConvertUB>(), Data);
+	SetUniformBufferParameter(CommandList, GetPixelShader(), GetUniformBufferParameter<FYUY2ConvertUB>(), Data);
 }
 
 
-/* FUYVYConvertPS shader
+/* FYVYUConvertPS shader
  *****************************************************************************/
 
 BEGIN_UNIFORM_BUFFER_STRUCT(FYVYUConvertUB, )
+DECLARE_UNIFORM_BUFFER_STRUCT_MEMBER(FMatrix, ColorTransform)
+DECLARE_UNIFORM_BUFFER_STRUCT_MEMBER(uint32, SrgbToLinear)
 DECLARE_UNIFORM_BUFFER_STRUCT_MEMBER(uint32, Width)
-DECLARE_UNIFORM_BUFFER_STRUCT_MEMBER_TEXTURE(Texture2D, YVYUTexture)
-DECLARE_UNIFORM_BUFFER_STRUCT_MEMBER_SAMPLER(SamplerState, YVYUSamplerB)
-DECLARE_UNIFORM_BUFFER_STRUCT_MEMBER_SAMPLER(SamplerState, YVYUSamplerP)
+DECLARE_UNIFORM_BUFFER_STRUCT_MEMBER_TEXTURE(Texture2D, Texture)
+DECLARE_UNIFORM_BUFFER_STRUCT_MEMBER_SAMPLER(SamplerState, SamplerB)
+DECLARE_UNIFORM_BUFFER_STRUCT_MEMBER_SAMPLER(SamplerState, SamplerP)
 END_UNIFORM_BUFFER_STRUCT(FYVYUConvertUB)
 
 IMPLEMENT_UNIFORM_BUFFER_STRUCT(FYVYUConvertUB, TEXT("YVYUConvertUB"));
 IMPLEMENT_SHADER_TYPE(, FYVYUConvertPS, TEXT("MediaShaders"), TEXT("YVYUConvertPS"), SF_Pixel);
 
 
-void FYVYUConvertPS::SetParameters(FRHICommandList& RHICmdList, TRefCountPtr<FRHITexture2D> YVYUTexture)
+void FYVYUConvertPS::SetParameters(FRHICommandList& CommandList, TRefCountPtr<FRHITexture2D> YVYUTexture, const FMatrix& ColorTransform, bool SrgbToLinear)
 {
 	FYVYUConvertUB UB;
 	{
+		UB.ColorTransform = ColorTransform;
+		UB.SamplerB = TStaticSamplerState<SF_Bilinear>::GetRHI();
+		UB.SamplerP = TStaticSamplerState<SF_Point>::GetRHI();
+		UB.SrgbToLinear = SrgbToLinear;
+		UB.Texture = YVYUTexture;
 		UB.Width = YVYUTexture->GetSizeX();
-		UB.YVYUTexture = YVYUTexture;
-		UB.YVYUSamplerB = TStaticSamplerState<SF_Bilinear>::GetRHI();
-		UB.YVYUSamplerP = TStaticSamplerState<SF_Point>::GetRHI();
 	}
 
 	TUniformBufferRef<FYVYUConvertUB> Data = TUniformBufferRef<FYVYUConvertUB>::CreateUniformBufferImmediate(UB, UniformBuffer_SingleFrame);
-	SetUniformBufferParameter(RHICmdList, GetPixelShader(), GetUniformBufferParameter<FYVYUConvertUB>(), Data);
+	SetUniformBufferParameter(CommandList, GetPixelShader(), GetUniformBufferParameter<FYVYUConvertUB>(), Data);
 }
diff --git a/Engine/Source/Runtime/UtilityShaders/Public/MediaShaders.h b/Engine/Source/Runtime/UtilityShaders/Public/MediaShaders.h
--- a/Engine/Source/Runtime/UtilityShaders/Public/MediaShaders.h
+++ b/Engine/Source/Runtime/UtilityShaders/Public/MediaShaders.h
@@ -9,6 +9,26 @@
 #include "Engine/Texture.h"
 #include "GlobalShader.h"
 
+
+namespace MediaShaders
+{
+	/** Color transform from YUV to sRGB (using values from MSDN). */
+	UTILITYSHADERS_API extern const FMatrix YuvToSrgbDefault;
+
+	/** Color transform from YUV to sRGB (in JPEG color space). */
+	UTILITYSHADERS_API extern const FMatrix YuvToSrgbJpeg;
+
+	/** Color transform from YUV to sRGB (using values from PS4 AvPlayer codec). */
+	UTILITYSHADERS_API extern const FMatrix YuvToSrgbPs4;
+
+	/** Color transform from YUV to sRGB (in Rec. 601 color space). */
+	UTILITYSHADERS_API extern const FMatrix YuvToSrgbRec601;
+
+	/** Color transform from YUV to sRGB (in Rec. 709 color space). */
+	UTILITYSHADERS_API extern const FMatrix YuvToRgbRec709;
+}
+
+
 /**
  * Stores media drawing vertices.
  */
@@ -119,7 +139,7 @@ public:
 		return bShaderHasOutdatedParameters;
 	}
 
-	UTILITYSHADERS_API void SetParameters(FRHICommandList& RHICmdList, TRefCountPtr<FRHITexture2D> AYUVTexture);
+	UTILITYSHADERS_API void SetParameters(FRHICommandList& RHICmdList, TRefCountPtr<FRHITexture2D> AYUVTexture, const FMatrix& ColorTransform, bool SrgbToLinear);
 };
 
 
@@ -187,7 +207,7 @@ public:
 		return bShaderHasOutdatedParameters;
 	}
 
-	UTILITYSHADERS_API void SetParameters(FRHICommandList& RHICmdList, TRefCountPtr<FRHITexture2D> NV12Texture, const FIntPoint& OutputDimensions);
+	UTILITYSHADERS_API void SetParameters(FRHICommandList& RHICmdList, TRefCountPtr<FRHITexture2D> NV12Texture, const FIntPoint& OutputDimensions, const FMatrix& ColorTransform, bool SrgbToLinear);
 };
 
 
@@ -222,7 +242,7 @@ public:
 		return bShaderHasOutdatedParameters;
 	}
 
-	UTILITYSHADERS_API void SetParameters(FRHICommandList& RHICmdList, TRefCountPtr<FRHITexture2D> NV21Texture, const FIntPoint& OutputDimensions);
+	UTILITYSHADERS_API void SetParameters(FRHICommandList& RHICmdList, TRefCountPtr<FRHITexture2D> NV21Texture, const FIntPoint& OutputDimensions, const FMatrix& ColorTransform, bool SrgbToLinear);
 };
 
 
@@ -256,7 +276,7 @@ public:
 		return bShaderHasOutdatedParameters;
 	}
 
-	UTILITYSHADERS_API void SetParameters(FRHICommandList& RHICmdList, TRefCountPtr<FRHITexture2D> RGBTexture, const FIntPoint& OutputDimensions);
+	UTILITYSHADERS_API void SetParameters(FRHICommandList& RHICmdList, TRefCountPtr<FRHITexture2D> RGBTexture, const FIntPoint& OutputDimensions, bool SrgbToLinear);
 };
 
 
@@ -293,7 +313,7 @@ public:
 		return bShaderHasOutdatedParameters;
 	}
 
-	UTILITYSHADERS_API void SetParameters(FRHICommandList& RHICmdList, TRefCountPtr<FRHITexture2D> LumaTexture, TRefCountPtr<FRHITexture2D> CbCrTexture);
+	UTILITYSHADERS_API void SetParameters(FRHICommandList& RHICmdList, TRefCountPtr<FRHITexture2D> LumaTexture, TRefCountPtr<FRHITexture2D> CbCrTexture, const FMatrix& ColorTransform, bool SrgbToLinear);
 };
 
 
@@ -329,7 +349,7 @@ public:
 		return bShaderHasOutdatedParameters;
 	}
 
-	UTILITYSHADERS_API void SetParameters(FRHICommandList& RHICmdList, TRefCountPtr<FRHITexture2D> UYVYTexture);
+	UTILITYSHADERS_API void SetParameters(FRHICommandList& RHICmdList, TRefCountPtr<FRHITexture2D> UYVYTexture, const FMatrix& ColorTransform, bool SrgbToLinear);
 };
 
 
@@ -363,7 +383,7 @@ public:
 		return bShaderHasOutdatedParameters;
 	}
 
-	UTILITYSHADERS_API void SetParameters(FRHICommandList& RHICmdList, TRefCountPtr<FRHITexture2D> YTexture, TRefCountPtr<FRHITexture2D> UTexture, TRefCountPtr<FRHITexture2D> VTexture);
+	UTILITYSHADERS_API void SetParameters(FRHICommandList& RHICmdList, TRefCountPtr<FRHITexture2D> YTexture, TRefCountPtr<FRHITexture2D> UTexture, TRefCountPtr<FRHITexture2D> VTexture, const FMatrix& ColorTransform, bool SrgbToLinear);
 };
 
 
@@ -399,7 +419,7 @@ public:
 		return bShaderHasOutdatedParameters;
 	}
 
-	UTILITYSHADERS_API void SetParameters(FRHICommandList& RHICmdList, TRefCountPtr<FRHITexture2D> YUY2Texture, const FIntPoint& OutputDimensions);
+	UTILITYSHADERS_API void SetParameters(FRHICommandList& RHICmdList, TRefCountPtr<FRHITexture2D> YUY2Texture, const FIntPoint& OutputDimensions, const FMatrix& ColorTransform, bool SrgbToLinear);
 };
 
 
@@ -435,5 +455,5 @@ public:
 		return bShaderHasOutdatedParameters;
 	}
 
-	UTILITYSHADERS_API void SetParameters(FRHICommandList& RHICmdList, TRefCountPtr<FRHITexture2D> YVYUTexture);
+	UTILITYSHADERS_API void SetParameters(FRHICommandList& RHICmdList, TRefCountPtr<FRHITexture2D> YVYUTexture, const FMatrix& ColorTransform, bool SrgbToLinear);
 };
