Add manual move constructor/assignment operator for VS2013 compatability

diff --git a/Engine/Source/Runtime/MovieSceneTracks/Private/Tracks/MovieSceneSkeletalAnimationTrack.cpp b/Engine/Source/Runtime/MovieSceneTracks/Private/Tracks/MovieSceneSkeletalAnimationTrack.cpp
--- a/Engine/Source/Runtime/MovieSceneTracks/Private/Tracks/MovieSceneSkeletalAnimationTrack.cpp
+++ b/Engine/Source/Runtime/MovieSceneTracks/Private/Tracks/MovieSceneSkeletalAnimationTrack.cpp
@@ -133,7 +133,22 @@ TInlineValue<FMovieSceneSegmentCompilerRules> UMovieSceneSkeletalAnimationTrack:
 	struct FSkeletalAnimationRowCompilerRules : FMovieSceneSegmentCompilerRules
 	{
 		FSkeletalAnimationRowCompilerRules(TInlineValue<FMovieSceneSegmentCompilerRules>&& InParentCompilerRules) { ParentCompilerRules = MoveTemp(InParentCompilerRules); }
-			
+
+#if PLATFORM_COMPILER_HAS_DEFAULTED_FUNCTIONS_FUNCTIONS
+		FSkeletalAnimationRowCompilerRules(FSkeletalAnimationRowCompilerRules&&) = default;
+		FSkeletalAnimationRowCompilerRules& operator=(FSkeletalAnimationRowCompilerRules&&) = default;
+#else
+		FSkeletalAnimationRowCompilerRules(FSkeletalAnimationRowCompilerRules&& RHS)
+			: FMovieSceneSegmentCompilerRules(MoveTemp(RHS))
+			, ParentCompilerRules(MoveTemp(RHS.ParentCompilerRules))
+		{}
+		FSkeletalAnimationRowCompilerRules& operator=(FSkeletalAnimationRowCompilerRules&& RHS)
+		{
+			FMovieSceneSegmentCompilerRules::operator=(MoveTemp(RHS));
+			ParentCompilerRules = MoveTemp(RHS.ParentCompilerRules);
+			return *this;
+		}
+#endif
 		virtual void BlendSegment(FMovieSceneSegment& Segment, const TArrayView<const FMovieSceneSectionData>& SourceData) const
 		{
 			// Make the skeletal animation sections upper bound exclusive so that when you place animation 
