Fix warning due to incorrect macro name.

diff --git a/Engine/Source/Runtime/MovieSceneTracks/Private/Tracks/MovieSceneSkeletalAnimationTrack.cpp b/Engine/Source/Runtime/MovieSceneTracks/Private/Tracks/MovieSceneSkeletalAnimationTrack.cpp
--- a/Engine/Source/Runtime/MovieSceneTracks/Private/Tracks/MovieSceneSkeletalAnimationTrack.cpp
+++ b/Engine/Source/Runtime/MovieSceneTracks/Private/Tracks/MovieSceneSkeletalAnimationTrack.cpp
@@ -134,7 +134,7 @@ TInlineValue<FMovieSceneSegmentCompilerRules> UMovieSceneSkeletalAnimationTrack:
 	{
 		FSkeletalAnimationRowCompilerRules(TInlineValue<FMovieSceneSegmentCompilerRules>&& InParentCompilerRules) { ParentCompilerRules = MoveTemp(InParentCompilerRules); }
 
-#if PLATFORM_COMPILER_HAS_DEFAULTED_FUNCTIONS_FUNCTIONS
+#if PLATFORM_COMPILER_HAS_DEFAULTED_FUNCTIONS
 		FSkeletalAnimationRowCompilerRules(FSkeletalAnimationRowCompilerRules&&) = default;
 		FSkeletalAnimationRowCompilerRules& operator=(FSkeletalAnimationRowCompilerRules&&) = default;
 #else
