MovieSceneTracks: Attempt to fix build

VS2013 does not properly support C++11 style initializer lists for arrays

diff --git a/Engine/Source/Runtime/MovieSceneTracks/Private/Evaluation/MovieScenePropertyTemplates.cpp b/Engine/Source/Runtime/MovieSceneTracks/Private/Evaluation/MovieScenePropertyTemplates.cpp
--- a/Engine/Source/Runtime/MovieSceneTracks/Private/Evaluation/MovieScenePropertyTemplates.cpp
+++ b/Engine/Source/Runtime/MovieSceneTracks/Private/Evaluation/MovieScenePropertyTemplates.cpp
@@ -205,10 +205,22 @@ void FMovieSceneVectorPropertySectionTemplate::Evaluate(const FMovieSceneEvaluat
 
 FMovieSceneTransformPropertySectionTemplate::FMovieSceneTransformPropertySectionTemplate(const UMovieScene3DTransformSection& Section, const UMovieScenePropertyTrack& Track)
 	: PropertyData(Track.GetPropertyName(), Track.GetPropertyPath())
-	, TranslationCurve{ Section.GetTranslationCurve(EAxis::X), Section.GetTranslationCurve(EAxis::Y), Section.GetTranslationCurve(EAxis::Z) }
-	, RotationCurve{ Section.GetRotationCurve(EAxis::X), Section.GetRotationCurve(EAxis::Y), Section.GetRotationCurve(EAxis::Z) }
-	, ScaleCurve{ Section.GetScaleCurve(EAxis::X), Section.GetScaleCurve(EAxis::Y), Section.GetScaleCurve(EAxis::Z) }
+//	, TranslationCurve{ Section.GetTranslationCurve(EAxis::X), Section.GetTranslationCurve(EAxis::Y), Section.GetTranslationCurve(EAxis::Z) }
+//	, RotationCurve{ Section.GetRotationCurve(EAxis::X), Section.GetRotationCurve(EAxis::Y), Section.GetRotationCurve(EAxis::Z) }
+//	, ScaleCurve{ Section.GetScaleCurve(EAxis::X), Section.GetScaleCurve(EAxis::Y), Section.GetScaleCurve(EAxis::Z) }
 {
+	// VS2013 does not properly support C++11 style initializer lists for arrays
+	TranslationCurve[0] = Section.GetTranslationCurve(EAxis::X);
+	TranslationCurve[1] = Section.GetTranslationCurve(EAxis::Y);
+	TranslationCurve[2] = Section.GetTranslationCurve(EAxis::Z);
+
+	RotationCurve[0] = Section.GetRotationCurve(EAxis::X);
+	RotationCurve[1] = Section.GetRotationCurve(EAxis::Y);
+	RotationCurve[2] = Section.GetRotationCurve(EAxis::Z);
+
+	ScaleCurve[0] = Section.GetScaleCurve(EAxis::X);
+	ScaleCurve[1] = Section.GetScaleCurve(EAxis::Y);
+	ScaleCurve[2] = Section.GetScaleCurve(EAxis::Z);
 }
 
 void FMovieSceneTransformPropertySectionTemplate::Evaluate(const FMovieSceneEvaluationOperand& Operand, const FMovieSceneContext& Context, const FPersistentEvaluationData& PersistentData, FMovieSceneExecutionTokens& ExecutionTokens) const
