Copying //UE4/Dev-Sequencer to //UE4/Main (Source: //UE4/Dev-Sequencer @ 3631098)

diff --git a/Engine/Source/Runtime/MovieSceneTracks/Private/Tracks/MovieScene3DTransformTrack.cpp b/Engine/Source/Runtime/MovieSceneTracks/Private/Tracks/MovieScene3DTransformTrack.cpp
--- a/Engine/Source/Runtime/MovieSceneTracks/Private/Tracks/MovieScene3DTransformTrack.cpp
+++ b/Engine/Source/Runtime/MovieSceneTracks/Private/Tracks/MovieScene3DTransformTrack.cpp
@@ -6,6 +6,7 @@
 #include "Compilation/MovieSceneTemplateInterrogation.h"
 #include "Compilation/MovieSceneSegmentCompiler.h"
 #include "MovieSceneEvaluationTrack.h"
+#include "Algo/BinarySearch.h"
 
 UMovieScene3DTransformTrack::UMovieScene3DTransformTrack( const FObjectInitializer& ObjectInitializer )
 	: Super( ObjectInitializer )
@@ -22,43 +23,103 @@ UMovieScene3DTransformTrack::UMovieScene3DTransformTrack( const FObjectInitializ
 	EvalOptions.bEvaluateNearestSection_DEPRECATED = EvalOptions.bCanEvaluateNearestSection = true;
 }
 
-TArray<FTrajectoryKey> UMovieScene3DTransformTrack::GetTrajectoryData() const
+TArray<FTrajectoryKey> UMovieScene3DTransformTrack::GetTrajectoryData(float Time, int32 MaxNumDataPoints) const
 {
-	TArray<FTrajectoryKey> Data;
-
-	struct FSectionAndCurve
+	struct FCurveKeyIterator
 	{
-		FSectionAndCurve(UMovieScene3DTransformSection* InSection, EMovieSceneTransformChannel InType, const FRichCurve& InCurve)
-			: Section(InSection), Type(InType), KeyIndex(0)
+		FCurveKeyIterator(UMovieScene3DTransformSection* InSection, EMovieSceneTransformChannel InType, const FRichCurve& InCurve, float StartTime)
+			: Section(InSection), Type(InType), SortedKeys(InCurve.GetConstRefOfKeys()), SectionRange(InSection->IsInfinite() ? TRange<float>::All() : InSection->GetRange())
 		{
-			TRange<float> SectionRange = InSection->IsInfinite() ? TRange<float>::All() : InSection->GetRange();
-			
-			for (const FRichCurveKey& Key : InCurve.GetConstRefOfKeys())
+			CurrentIndex = Algo::LowerBoundBy(SortedKeys, StartTime, &FRichCurveKey::Time, TLess<>());
+
+			bIsLowerBound = false;
+			bIsUpperBound = !SortedKeys.IsValidIndex(CurrentIndex) && SectionRange.GetUpperBound().IsClosed();
+		}
+
+		FCurveKeyIterator& operator--()
+		{
+			if (bIsLowerBound)
+			{
+				bIsLowerBound = false;
+				CurrentIndex = -1;
+			}
+			else
 			{
-				if (SectionRange.Contains(Key.Time))
+				if (bIsUpperBound)
 				{
-					Keys.Emplace(Key.Time, Key.InterpMode);
+					bIsUpperBound = false;
+					CurrentIndex = Algo::LowerBoundBy(SortedKeys, SectionRange.GetUpperBoundValue(), &FRichCurveKey::Time, TLess<>()) - 1;
 				}
+				else
+				{
+					--CurrentIndex;
+				}
+
+				bIsLowerBound = SectionRange.GetLowerBound().IsClosed() && (!SortedKeys.IsValidIndex(CurrentIndex) || !SectionRange.Contains(SortedKeys[CurrentIndex].Time) );
 			}
+			return *this;
+		}
 
-			// Inser the start and end ranges
-			if (!SectionRange.GetLowerBound().IsOpen() && (!Keys.Num() || !FMath::IsNearlyEqual(SectionRange.GetLowerBoundValue(), Keys[0].Get<0>())))
+		FCurveKeyIterator& operator++()
+		{
+			if (bIsUpperBound)
 			{
-				Keys.Insert(MakeTuple(SectionRange.GetLowerBoundValue(), RCIM_None), 0);
+				bIsUpperBound = false;
+				CurrentIndex = -1;
 			}
-			if (!SectionRange.GetUpperBound().IsOpen() && (!Keys.Num() || !FMath::IsNearlyEqual(SectionRange.GetUpperBoundValue(), Keys.Last().Get<0>())))
+			else
 			{
-				Keys.Add(MakeTuple(SectionRange.GetUpperBoundValue(), RCIM_None));
+				if (bIsLowerBound)
+				{
+					bIsLowerBound = false;
+					CurrentIndex = Algo::UpperBoundBy(SortedKeys, SectionRange.GetLowerBoundValue(), &FRichCurveKey::Time, TLess<>());
+				}
+				else 
+				{
+					++CurrentIndex;
+				}
+
+				bIsUpperBound = SectionRange.GetUpperBound().IsClosed() && (!SortedKeys.IsValidIndex(CurrentIndex) || !SectionRange.Contains(SortedKeys[CurrentIndex].Time) );
 			}
+			return *this;
+		}
+
+		explicit operator bool() const
+		{
+			return ( bIsLowerBound || bIsUpperBound ) || ( SortedKeys.IsValidIndex(CurrentIndex) && SectionRange.Contains(SortedKeys[CurrentIndex].Time) );
+		}
+
+		float GetTime() const
+		{
+			return bIsLowerBound ? SectionRange.GetLowerBoundValue() : bIsUpperBound ? SectionRange.GetUpperBoundValue() : SortedKeys[CurrentIndex].Time;
+		}
+
+		ERichCurveInterpMode GetInterpMode() const
+		{
+			return ( bIsLowerBound || bIsUpperBound ) ? RCIM_None : SortedKeys[CurrentIndex].InterpMode.GetValue();
+		}
+
+		EMovieSceneTransformChannel GetType() const
+		{
+			return Type;
 		}
 
+		UMovieScene3DTransformSection* GetSection() const
+		{
+			return Section;
+		}
+
+	private:
 		UMovieScene3DTransformSection* Section;
 		EMovieSceneTransformChannel Type;
-		TArray<TTuple<float, ERichCurveInterpMode>> Keys;
-		/** Running key index into the above array of tuples */
-		int32 KeyIndex;
+		const TArray<FRichCurveKey>& SortedKeys;
+		TRange<float> SectionRange;
+		int32 CurrentIndex;
+		bool bIsUpperBound, bIsLowerBound;
 	};
-	TArray<FSectionAndCurve> SectionsAndCurves;
+
+	TArray<FCurveKeyIterator> ForwardIters;
+	TArray<FCurveKeyIterator> BackwardIters;
 
 	for (UMovieSceneSection* Section : Sections)
 	{
@@ -68,68 +129,105 @@ TArray<FTrajectoryKey> UMovieScene3DTransformTrack::GetTrajectoryData() const
 			EMovieSceneTransformChannel Mask = TransformSection->GetMask().GetChannels();
 			if (EnumHasAnyFlags(Mask, EMovieSceneTransformChannel::TranslationX))
 			{
-				SectionsAndCurves.Emplace(TransformSection, EMovieSceneTransformChannel::TranslationX, TransformSection->GetTranslationCurve(EAxis::X));
+				ForwardIters.Emplace(TransformSection, EMovieSceneTransformChannel::TranslationX, TransformSection->GetTranslationCurve(EAxis::X), Time);
+				BackwardIters.Emplace(TransformSection, EMovieSceneTransformChannel::TranslationX, TransformSection->GetTranslationCurve(EAxis::X), Time);
 			}
 			if (EnumHasAnyFlags(Mask, EMovieSceneTransformChannel::TranslationY))
 			{
-				SectionsAndCurves.Emplace(TransformSection, EMovieSceneTransformChannel::TranslationY, TransformSection->GetTranslationCurve(EAxis::Y));
+				ForwardIters.Emplace(TransformSection, EMovieSceneTransformChannel::TranslationY, TransformSection->GetTranslationCurve(EAxis::Y), Time);
+				BackwardIters.Emplace(TransformSection, EMovieSceneTransformChannel::TranslationY, TransformSection->GetTranslationCurve(EAxis::Y), Time);
 			}
 			if (EnumHasAnyFlags(Mask, EMovieSceneTransformChannel::TranslationZ))
 			{
-				SectionsAndCurves.Emplace(TransformSection, EMovieSceneTransformChannel::TranslationZ, TransformSection->GetTranslationCurve(EAxis::Z));
+				ForwardIters.Emplace(TransformSection, EMovieSceneTransformChannel::TranslationZ, TransformSection->GetTranslationCurve(EAxis::Z), Time);
+				BackwardIters.Emplace(TransformSection, EMovieSceneTransformChannel::TranslationZ, TransformSection->GetTranslationCurve(EAxis::Z), Time);
 			}
 			if (EnumHasAnyFlags(Mask, EMovieSceneTransformChannel::RotationX))
 			{
-				SectionsAndCurves.Emplace(TransformSection, EMovieSceneTransformChannel::RotationX, TransformSection->GetRotationCurve(EAxis::X));
+				ForwardIters.Emplace(TransformSection, EMovieSceneTransformChannel::RotationX, TransformSection->GetRotationCurve(EAxis::X), Time);
+				BackwardIters.Emplace(TransformSection, EMovieSceneTransformChannel::RotationX, TransformSection->GetRotationCurve(EAxis::X), Time);
 			}
 			if (EnumHasAnyFlags(Mask, EMovieSceneTransformChannel::RotationY))
 			{
-				SectionsAndCurves.Emplace(TransformSection, EMovieSceneTransformChannel::RotationY, TransformSection->GetRotationCurve(EAxis::Y));
+				ForwardIters.Emplace(TransformSection, EMovieSceneTransformChannel::RotationY, TransformSection->GetRotationCurve(EAxis::Y), Time);
+				BackwardIters.Emplace(TransformSection, EMovieSceneTransformChannel::RotationY, TransformSection->GetRotationCurve(EAxis::Y), Time);
 			}
 			if (EnumHasAnyFlags(Mask, EMovieSceneTransformChannel::RotationZ))
 			{
-				SectionsAndCurves.Emplace(TransformSection, EMovieSceneTransformChannel::RotationZ, TransformSection->GetRotationCurve(EAxis::Z));
+				ForwardIters.Emplace(TransformSection, EMovieSceneTransformChannel::RotationZ, TransformSection->GetRotationCurve(EAxis::Z), Time);
+				BackwardIters.Emplace(TransformSection, EMovieSceneTransformChannel::RotationZ, TransformSection->GetRotationCurve(EAxis::Z), Time);
 			}
 		}
 	}
 
-	int32 TotalNumKeys = 0;
-	for (const FSectionAndCurve& SectionAndCurve : SectionsAndCurves)
+	auto HasAnyValidIterators = [](const FCurveKeyIterator& It)
+	{
+		return It;
+	};
+
+	for (FCurveKeyIterator& Bck : BackwardIters)
 	{
-		TotalNumKeys += SectionAndCurve.Keys.Num();
+		--Bck;
 	}
 
-	int32 NumVisitedKeys = 0;
-	while (NumVisitedKeys < TotalNumKeys)
+	TArray<FTrajectoryKey> Result;
+	while (ForwardIters.ContainsByPredicate(HasAnyValidIterators) || BackwardIters.ContainsByPredicate(HasAnyValidIterators))
 	{
-		FTrajectoryKey& NewKey = Data[Data.Emplace(TNumericLimits<float>::Max())];
+		if (MaxNumDataPoints != 0 && Result.Num() >= MaxNumDataPoints)
+		{
+			break;
+		}
 
-		// Find the earliest key time
-		for (const FSectionAndCurve& SectionAndCurve : SectionsAndCurves)
+		FTrajectoryKey& NewKey = Result[Result.Emplace(TNumericLimits<float>::Max())];
+
+		// Find the closest key time
+		for (const FCurveKeyIterator& Fwd : ForwardIters)
 		{
-			if (SectionAndCurve.Keys.IsValidIndex(SectionAndCurve.KeyIndex))
+			if (Fwd && ( FMath::Abs(Time - Fwd.GetTime()) < FMath::Abs(Time - NewKey.Time) ) )
 			{
-				NewKey.Time = FMath::Min(NewKey.Time, SectionAndCurve.Keys[SectionAndCurve.KeyIndex].Get<0>());
+				NewKey.Time = Fwd.GetTime();
+			}
+		}
+		for (const FCurveKeyIterator& Bck : BackwardIters)
+		{
+			if (Bck && ( FMath::Abs(Time - Bck.GetTime()) < FMath::Abs(Time - NewKey.Time) ) )
+			{
+				NewKey.Time = Bck.GetTime();
 			}
 		}
 
-		// Add any keys at the current time to this trajectory key
-		for (FSectionAndCurve& SectionAndCurve : SectionsAndCurves)
+		for (FCurveKeyIterator& Fwd : ForwardIters)
 		{
-			if (SectionAndCurve.Keys.IsValidIndex(SectionAndCurve.KeyIndex) && FMath::IsNearlyEqual(NewKey.Time, SectionAndCurve.Keys[SectionAndCurve.KeyIndex].Get<0>()))
+			if (Fwd && FMath::IsNearlyEqual(Fwd.GetTime(), NewKey.Time))
 			{
 				// Add this key to the trajectory key
-				NewKey.KeyData.Emplace(SectionAndCurve.Section, SectionAndCurve.Keys[SectionAndCurve.KeyIndex].Get<1>(), SectionAndCurve.Type);
+				NewKey.KeyData.Emplace(Fwd.GetSection(), Fwd.GetInterpMode(), Fwd.GetType());
+				// Move onto the next key in this curve
+				++Fwd;
+			}
+		}
 
+		for (FCurveKeyIterator& Bck : BackwardIters)
+		{
+			if (Bck && FMath::IsNearlyEqual(Bck.GetTime(), NewKey.Time))
+			{
+				// Add this key to the trajectory key
+				NewKey.KeyData.Emplace(Bck.GetSection(), Bck.GetInterpMode(), Bck.GetType());
 				// Move onto the next key in this curve
-				++SectionAndCurve.KeyIndex;
-				++NumVisitedKeys;
+				--Bck;
 			}
 		}
 	}
-	return Data;
-}
 
+	Result.Sort(
+		[](const FTrajectoryKey& A, const FTrajectoryKey& B)
+		{
+			return A.Time < B.Time;
+		}
+	);
+
+	return Result;
+}
 
 UMovieSceneSection* UMovieScene3DTransformTrack::CreateNewSection()
 {
diff --git a/Engine/Source/Runtime/MovieSceneTracks/Public/Sections/MovieSceneAudioSection.h b/Engine/Source/Runtime/MovieSceneTracks/Public/Sections/MovieSceneAudioSection.h
--- a/Engine/Source/Runtime/MovieSceneTracks/Public/Sections/MovieSceneAudioSection.h
+++ b/Engine/Source/Runtime/MovieSceneTracks/Public/Sections/MovieSceneAudioSection.h
@@ -173,6 +173,7 @@ public:
 	virtual UMovieSceneSection* SplitSection(float SplitTime) override;
 	virtual void GetKeyHandles(TSet<FKeyHandle>& OutKeyHandles, TRange<float> TimeRange) const override;
 	virtual void GetSnapTimes(TArray<float>& OutSnapTimes, bool bGetSectionBorders) const override;
+	virtual TOptional<float> GetOffsetTime() const override { return TOptional<float>(StartOffset); }
 	virtual TOptional<float> GetKeyTime( FKeyHandle KeyHandle ) const override { return TOptional<float>(); }
 	virtual void SetKeyTime( FKeyHandle KeyHandle, float Time ) override { }
 	virtual FMovieSceneEvalTemplatePtr GenerateTemplate() const override;
diff --git a/Engine/Source/Runtime/MovieSceneTracks/Public/Sections/MovieSceneSkeletalAnimationSection.h b/Engine/Source/Runtime/MovieSceneTracks/Public/Sections/MovieSceneSkeletalAnimationSection.h
--- a/Engine/Source/Runtime/MovieSceneTracks/Public/Sections/MovieSceneSkeletalAnimationSection.h
+++ b/Engine/Source/Runtime/MovieSceneTracks/Public/Sections/MovieSceneSkeletalAnimationSection.h
@@ -74,6 +74,7 @@ public:
 	virtual UMovieSceneSection* SplitSection(float SplitTime) override;
 	virtual void GetKeyHandles(TSet<FKeyHandle>& OutKeyHandles, TRange<float> TimeRange) const override;
 	virtual void GetSnapTimes(TArray<float>& OutSnapTimes, bool bGetSectionBorders) const override;
+	virtual TOptional<float> GetOffsetTime() const override { return TOptional<float>(Params.StartOffset); }
 	virtual TOptional<float> GetKeyTime( FKeyHandle KeyHandle ) const override { return TOptional<float>(); }
 	virtual void SetKeyTime( FKeyHandle KeyHandle, float Time ) override { }
 	virtual FMovieSceneEvalTemplatePtr GenerateTemplate() const override;
diff --git a/Engine/Source/Runtime/MovieSceneTracks/Public/Sections/MovieSceneSubSection.h b/Engine/Source/Runtime/MovieSceneTracks/Public/Sections/MovieSceneSubSection.h
--- a/Engine/Source/Runtime/MovieSceneTracks/Public/Sections/MovieSceneSubSection.h
+++ b/Engine/Source/Runtime/MovieSceneTracks/Public/Sections/MovieSceneSubSection.h
@@ -125,6 +125,7 @@ public:
 
 	virtual UMovieSceneSection* SplitSection( float SplitTime ) override;
 	virtual void TrimSection( float TrimTime, bool bTrimLeft ) override;
+	virtual TOptional<float> GetOffsetTime() const override { return TOptional<float>(Parameters.StartOffset); }
 	virtual TOptional<float> GetKeyTime( FKeyHandle KeyHandle ) const override { return TOptional<float>(); }
 	virtual void SetKeyTime( FKeyHandle KeyHandle, float Time ) override { }
 
diff --git a/Engine/Source/Runtime/MovieSceneTracks/Public/Tracks/MovieScene3DTransformTrack.h b/Engine/Source/Runtime/MovieSceneTracks/Public/Tracks/MovieScene3DTransformTrack.h
--- a/Engine/Source/Runtime/MovieSceneTracks/Public/Tracks/MovieScene3DTransformTrack.h
+++ b/Engine/Source/Runtime/MovieSceneTracks/Public/Tracks/MovieScene3DTransformTrack.h
@@ -59,7 +59,7 @@ public:
 
 	virtual UMovieSceneSection* CreateNewSection() override;
 
-	MOVIESCENETRACKS_API TArray<FTrajectoryKey> GetTrajectoryData() const;
+	MOVIESCENETRACKS_API TArray<FTrajectoryKey> GetTrajectoryData(float Time, int32 MaxNumDataPoints) const;
 
 	/**
 	 * Access the interrogation key for transform data - any interrgation data stored with this key is guaranteed to be of type 'FTransform'
