Copying //UE4/Dev-Sequencer to //UE4/Dev-Main (Source: //UE4/Dev-Sequencer @ 3800792)

diff --git a/Engine/Source/Runtime/MovieSceneTracks/Private/Evaluation/MovieScene3DAttachTemplate.cpp b/Engine/Source/Runtime/MovieSceneTracks/Private/Evaluation/MovieScene3DAttachTemplate.cpp
--- a/Engine/Source/Runtime/MovieSceneTracks/Private/Evaluation/MovieScene3DAttachTemplate.cpp
+++ b/Engine/Source/Runtime/MovieSceneTracks/Private/Evaluation/MovieScene3DAttachTemplate.cpp
@@ -117,9 +117,9 @@ struct F3DAttachExecutionToken
 	: IMovieSceneExecutionToken
 	, F3DAttachTrackToken
 {
-	F3DAttachExecutionToken(FGuid InAttachGuid, FName InAttachSocketName, FName InAttachComponentName, bool bInShouldBeAttached, EAttachmentRule InAttachmentLocationRule, EAttachmentRule InAttachmentRotationRule, EAttachmentRule InAttachmentScaleRule, EDetachmentRule InDetachmentLocationRule, EDetachmentRule InDetachmentRotationRule, EDetachmentRule InDetachmentScaleRule) 
+	F3DAttachExecutionToken(FMovieSceneObjectBindingID InAttachBindingID, FName InAttachSocketName, FName InAttachComponentName, bool bInShouldBeAttached, EAttachmentRule InAttachmentLocationRule, EAttachmentRule InAttachmentRotationRule, EAttachmentRule InAttachmentScaleRule, EDetachmentRule InDetachmentLocationRule, EDetachmentRule InDetachmentRotationRule, EDetachmentRule InDetachmentScaleRule) 
 		: F3DAttachTrackToken(nullptr, InAttachSocketName, bInShouldBeAttached, InAttachmentLocationRule, InAttachmentRotationRule, InAttachmentScaleRule, InDetachmentLocationRule, InDetachmentRotationRule, InDetachmentScaleRule)
-		, AttachGuid(InAttachGuid)
+		, AttachBindingID(InAttachBindingID)
 		, AttachComponentName(InAttachComponentName)
 	{}
 
@@ -170,7 +170,16 @@ struct F3DAttachExecutionToken
 	{
 		MOVIESCENE_DETAILED_SCOPE_CYCLE_COUNTER(MovieSceneEval_AttachTrack_TokenExecute)
 
-		FMovieSceneEvaluationOperand AttachOperand(Operand.SequenceID, AttachGuid);
+		FMovieSceneSequenceID SequenceID = Operand.SequenceID;
+		if (AttachBindingID.GetSequenceID().IsValid())
+		{
+			// Ensure that this ID is resolvable from the root, based on the current local sequence ID
+			FMovieSceneObjectBindingID RootBindingID = AttachBindingID.ResolveLocalToRoot(SequenceID, Player.GetEvaluationTemplate().GetHierarchy());
+			SequenceID = RootBindingID.GetSequenceID();
+		}
+
+		// If the transform is set, otherwise use the bound actor's transform
+		FMovieSceneEvaluationOperand AttachOperand(SequenceID, AttachBindingID.GetGuid());
 		
 		TArrayView<TWeakObjectPtr<>> Objects = Player.FindBoundObjects(AttachOperand);
 		if (!Objects.Num())
@@ -210,12 +219,12 @@ struct F3DAttachExecutionToken
 		}
 	}
 	
-	FGuid AttachGuid;
+	FMovieSceneObjectBindingID AttachBindingID;
 	FName AttachComponentName;
 };
 
 FMovieScene3DAttachSectionTemplate::FMovieScene3DAttachSectionTemplate(const UMovieScene3DAttachSection& Section)
-	: AttachGuid(Section.GetConstraintId())
+	: AttachBindingID(Section.GetConstraintBindingID())
 	, AttachSocketName(Section.AttachSocketName)
 	, AttachComponentName(Section.AttachComponentName)
 	, AttachmentLocationRule(Section.AttachmentLocationRule)
@@ -233,5 +242,5 @@ void FMovieScene3DAttachSectionTemplate::Evaluate(const FMovieSceneEvaluationOpe
 
 	const bool bShouldBeAttached = true;
 
-	ExecutionTokens.Add(F3DAttachExecutionToken(AttachGuid, AttachSocketName, AttachComponentName, bShouldBeAttached, AttachmentLocationRule, AttachmentRotationRule, AttachmentScaleRule, DetachmentLocationRule, DetachmentRotationRule, DetachmentScaleRule));
+	ExecutionTokens.Add(F3DAttachExecutionToken(AttachBindingID, AttachSocketName, AttachComponentName, bShouldBeAttached, AttachmentLocationRule, AttachmentRotationRule, AttachmentScaleRule, DetachmentLocationRule, DetachmentRotationRule, DetachmentScaleRule));
 }
diff --git a/Engine/Source/Runtime/MovieSceneTracks/Private/Evaluation/MovieScene3DAttachTemplate.h b/Engine/Source/Runtime/MovieSceneTracks/Private/Evaluation/MovieScene3DAttachTemplate.h
--- a/Engine/Source/Runtime/MovieSceneTracks/Private/Evaluation/MovieScene3DAttachTemplate.h
+++ b/Engine/Source/Runtime/MovieSceneTracks/Private/Evaluation/MovieScene3DAttachTemplate.h
@@ -7,6 +7,7 @@
 #include "Engine/EngineTypes.h"
 #include "Misc/Guid.h"
 #include "Evaluation/MovieSceneEvalTemplate.h"
+#include "MovieSceneObjectBindingID.h"
 #include "MovieScene3DAttachTemplate.generated.h"
 
 class UMovieScene3DAttachSection;
@@ -19,9 +20,9 @@ struct FMovieScene3DAttachSectionTemplate : public FMovieSceneEvalTemplate
 	FMovieScene3DAttachSectionTemplate() {}
 	FMovieScene3DAttachSectionTemplate(const UMovieScene3DAttachSection& Section);
 
-	/** GUID of the parent we should attach to */
+	/** The object binding ID we should attach to */
 	UPROPERTY()
-	FGuid AttachGuid;
+	FMovieSceneObjectBindingID AttachBindingID;
 
 	/** Name of the socket to attach to */
 	UPROPERTY()
diff --git a/Engine/Source/Runtime/MovieSceneTracks/Private/Evaluation/MovieScene3DPathTemplate.cpp b/Engine/Source/Runtime/MovieSceneTracks/Private/Evaluation/MovieScene3DPathTemplate.cpp
--- a/Engine/Source/Runtime/MovieSceneTracks/Private/Evaluation/MovieScene3DPathTemplate.cpp
+++ b/Engine/Source/Runtime/MovieSceneTracks/Private/Evaluation/MovieScene3DPathTemplate.cpp
@@ -17,8 +17,8 @@ DECLARE_CYCLE_STAT(TEXT("Path Track Token Execute"), MovieSceneEval_PathTrack_To
 struct F3DPathExecutionToken
 	: IMovieSceneExecutionToken
 {
-	F3DPathExecutionToken(FGuid InPathGuid, float InTiming, MovieScene3DPathSection_Axis InFrontAxisEnum, MovieScene3DPathSection_Axis InUpAxisEnum, bool bInFollow, bool bInReverse, bool bInForceUpright) 
-		: PathGuid(InPathGuid)
+	F3DPathExecutionToken(FMovieSceneObjectBindingID InPathBindingID, float InTiming, MovieScene3DPathSection_Axis InFrontAxisEnum, MovieScene3DPathSection_Axis InUpAxisEnum, bool bInFollow, bool bInReverse, bool bInForceUpright) 
+		: PathBindingID(InPathBindingID)
 		, Timing(InTiming)
 		, FrontAxisEnum(InFrontAxisEnum)
 		, UpAxisEnum(InUpAxisEnum)
@@ -125,7 +125,16 @@ struct F3DPathExecutionToken
 	{
 		MOVIESCENE_DETAILED_SCOPE_CYCLE_COUNTER(MovieSceneEval_PathTrack_TokenExecute)
 
-		FMovieSceneEvaluationOperand PathOperand(Operand.SequenceID, PathGuid);
+		FMovieSceneSequenceID SequenceID = Operand.SequenceID;
+		if (PathBindingID.GetSequenceID().IsValid())
+		{
+			// Ensure that this ID is resolvable from the root, based on the current local sequence ID
+			FMovieSceneObjectBindingID RootBindingID = PathBindingID.ResolveLocalToRoot(SequenceID, Player.GetEvaluationTemplate().GetHierarchy());
+			SequenceID = RootBindingID.GetSequenceID();
+		}
+
+		// If the transform is set, otherwise use the bound actor's transform
+		FMovieSceneEvaluationOperand PathOperand(SequenceID, PathBindingID.GetGuid());
 		
 		TArrayView<TWeakObjectPtr<>> Objects = Player.FindBoundObjects(PathOperand);
 		if (!Objects.Num())
@@ -167,7 +176,7 @@ struct F3DPathExecutionToken
 		}
 	}
 	
-	FGuid PathGuid;
+	FMovieSceneObjectBindingID PathBindingID;
 	float Timing;
 	MovieScene3DPathSection_Axis FrontAxisEnum;
 	MovieScene3DPathSection_Axis UpAxisEnum;
@@ -177,7 +186,7 @@ struct F3DPathExecutionToken
 };
 
 FMovieScene3DPathSectionTemplate::FMovieScene3DPathSectionTemplate(const UMovieScene3DPathSection& Section)
-	: PathGuid(Section.GetConstraintId())
+	: PathBindingID(Section.GetConstraintBindingID())
 	, TimingCurve(Section.GetTimingCurve())
 	, FrontAxisEnum(Section.GetFrontAxisEnum())
 	, UpAxisEnum(Section.GetUpAxisEnum())
@@ -193,5 +202,5 @@ void FMovieScene3DPathSectionTemplate::Evaluate(const FMovieSceneEvaluationOpera
 		
 	float Timing = TimingCurve.Eval(Context.GetTime());
 
-	ExecutionTokens.Add(F3DPathExecutionToken(PathGuid, Timing, FrontAxisEnum, UpAxisEnum, bFollow, bReverse, bForceUpright));
+	ExecutionTokens.Add(F3DPathExecutionToken(PathBindingID, Timing, FrontAxisEnum, UpAxisEnum, bFollow, bReverse, bForceUpright));
 }
diff --git a/Engine/Source/Runtime/MovieSceneTracks/Private/Evaluation/MovieScene3DPathTemplate.h b/Engine/Source/Runtime/MovieSceneTracks/Private/Evaluation/MovieScene3DPathTemplate.h
--- a/Engine/Source/Runtime/MovieSceneTracks/Private/Evaluation/MovieScene3DPathTemplate.h
+++ b/Engine/Source/Runtime/MovieSceneTracks/Private/Evaluation/MovieScene3DPathTemplate.h
@@ -8,6 +8,7 @@
 #include "Curves/RichCurve.h"
 #include "Evaluation/MovieSceneEvalTemplate.h"
 #include "Sections/MovieScene3DPathSection.h"
+#include "MovieSceneObjectBindingID.h"
 #include "MovieScene3DPathTemplate.generated.h"
 
 USTRUCT()
@@ -18,9 +19,9 @@ struct FMovieScene3DPathSectionTemplate : public FMovieSceneEvalTemplate
 	FMovieScene3DPathSectionTemplate() {}
 	FMovieScene3DPathSectionTemplate(const UMovieScene3DPathSection& Section);
 
-	/** GUID of the path we should attach to */
+	/** The object binding ID of the path we should attach to */
 	UPROPERTY()
-	FGuid PathGuid;
+	FMovieSceneObjectBindingID PathBindingID;
 
 	/** The timing curve */
 	UPROPERTY()
diff --git a/Engine/Source/Runtime/MovieSceneTracks/Private/Evaluation/MovieScene3DTransformTemplate.cpp b/Engine/Source/Runtime/MovieSceneTracks/Private/Evaluation/MovieScene3DTransformTemplate.cpp
--- a/Engine/Source/Runtime/MovieSceneTracks/Private/Evaluation/MovieScene3DTransformTemplate.cpp
+++ b/Engine/Source/Runtime/MovieSceneTracks/Private/Evaluation/MovieScene3DTransformTemplate.cpp
@@ -9,6 +9,7 @@
 #include "IMovieScenePlayer.h"
 #include "Evaluation/Blending/BlendableTokenStack.h"
 #include "Evaluation/Blending/MovieSceneBlendingActuatorID.h"
+#include "IMovieSceneTransformOrigin.h"
 
 DECLARE_CYCLE_STAT(TEXT("Transform Track Evaluate"), MovieSceneEval_TransformTrack_Evaluate, STATGROUP_MovieSceneEval);
 DECLARE_CYCLE_STAT(TEXT("Transform Track Token Execute"), MovieSceneEval_TransformTrack_TokenExecute, STATGROUP_MovieSceneEval);
@@ -98,6 +99,28 @@ FMovieSceneComponentTransformSectionTemplate::FMovieSceneComponentTransformSecti
 {
 }
 
+struct FComponentTransformPersistentData : IPersistentEvaluationData
+{
+	FTransform Origin;
+};
+
+void FMovieSceneComponentTransformSectionTemplate::Initialize(const FMovieSceneEvaluationOperand& Operand, const FMovieSceneContext& Context, FPersistentEvaluationData& PersistentData, IMovieScenePlayer& Player) const
+{
+	// If the global instance data implements a transform origin interface, use its transform as an origin for this transform
+	const UObject* InstanceData = Player.GetInstanceData();
+	const IMovieSceneTransformOrigin* RawInterface = Cast<IMovieSceneTransformOrigin>(InstanceData);
+	const bool bHasInterface = RawInterface || (InstanceData && InstanceData->GetClass()->ImplementsInterface(UMovieSceneTransformOrigin::StaticClass()));
+
+	if (bHasInterface && TemplateData.BlendType == EMovieSceneBlendType::Absolute)
+	{
+		// Retrieve the current origin
+		FTransform TransformOrigin = RawInterface ? RawInterface->GetTransformOrigin() : IMovieSceneTransformOrigin::Execute_BP_GetTransformOrigin(InstanceData);
+
+		// Assign the transform origin to the peristent data so it can be queried in Evaluate
+		PersistentData.GetOrAddSectionData<FComponentTransformPersistentData>().Origin = TransformOrigin;
+	}
+}
+
 void FMovieSceneComponentTransformSectionTemplate::Evaluate(const FMovieSceneEvaluationOperand& Operand, const FMovieSceneContext& Context, const FPersistentEvaluationData& PersistentData, FMovieSceneExecutionTokens& ExecutionTokens) const
 {
 	using namespace MovieScene;
@@ -108,6 +131,37 @@ void FMovieSceneComponentTransformSectionTemplate::Evaluate(const FMovieSceneEva
 		return;
 	}
 
+	// Apply origin transformation if necessary 
+	const FComponentTransformPersistentData* Data = PersistentData.FindSectionData<FComponentTransformPersistentData>();
+	if (TemplateData.BlendType == EMovieSceneBlendType::Absolute && Data)
+	{
+		float Components[6] = {
+			TransformValue.Get(0, 0.f), TransformValue.Get(1, 0.f), TransformValue.Get(2, 0.f),
+			TransformValue.Get(3, 0.f), TransformValue.Get(4, 0.f), TransformValue.Get(5, 0.f),
+		};
+
+		FTransform AnimatedTransform(FRotator(Components[4], Components[5], Components[3]), FVector(Components[0], Components[1], Components[2]));
+		AnimatedTransform = AnimatedTransform * Data->Origin;
+
+		FVector Location = AnimatedTransform.GetTranslation();
+		Components[0] = Location.X;
+		Components[1] = Location.Y;
+		Components[2] = Location.Z;
+
+		FVector Rotation = AnimatedTransform.GetRotation().Euler();
+		Components[3] = Rotation.X;
+		Components[4] = Rotation.Y;
+		Components[5] = Rotation.Z;
+
+		for (int32 Index = 0; Index < ARRAY_COUNT(Components); ++Index)
+		{
+			if (TransformValue.IsSet(Index))
+			{
+				TransformValue.Set(Index, Components[Index]);
+			}
+		}
+	}
+
 	// Ensure the accumulator knows how to actually apply component transforms
 	FMovieSceneBlendingActuatorID ActuatorTypeID = FComponentTransformActuator::GetActuatorTypeID();
 	if (!ExecutionTokens.GetBlendingAccumulator().FindActuator<F3DTransformTrackToken>(ActuatorTypeID))
diff --git a/Engine/Source/Runtime/MovieSceneTracks/Private/Evaluation/MovieSceneCameraAnimTemplate.cpp b/Engine/Source/Runtime/MovieSceneTracks/Private/Evaluation/MovieSceneCameraAnimTemplate.cpp
--- a/Engine/Source/Runtime/MovieSceneTracks/Private/Evaluation/MovieSceneCameraAnimTemplate.cpp
+++ b/Engine/Source/Runtime/MovieSceneTracks/Private/Evaluation/MovieSceneCameraAnimTemplate.cpp
@@ -26,12 +26,25 @@ struct FBlendedPostProcessSettings : FPostProcessSettings
 };
 
 
+FMovieSceneSharedDataId GetSharedExecutionTokenID()
+{
+	static FMovieSceneSharedDataId SharedDataId = FMovieSceneSharedDataId::Allocate();
+	return SharedDataId;
+}
+
 /** Persistent data that exists as long as a given camera track is being evaluated */
 struct FMovieSceneAdditiveCameraData : IPersistentEvaluationData
 {
+	FMovieSceneAdditiveCameraData()
+		: bApplyTransform(false)
+		, bApplyPostProcessing(false)
+		, TotalTransform(FTransform::Identity)
+		, TotalFOVOffset(0.f)
+	{}
+
 	static FMovieSceneAdditiveCameraData& Get(const FMovieSceneEvaluationOperand& Operand, FPersistentEvaluationData& PersistentData)
 	{
-		return PersistentData.GetOrAdd<FMovieSceneAdditiveCameraData>(FSharedPersistentDataKey(FMovieSceneAdditiveCameraAnimationTrackTemplate::SharedDataId, Operand));
+		return PersistentData.GetOrAdd<FMovieSceneAdditiveCameraData>(FSharedPersistentDataKey(GetSharedExecutionTokenID(), Operand));
 	}
 
 	/** Reset the additive camera values */
@@ -188,45 +201,36 @@ struct FPreAnimatedPostProcessingBlendsTokenProducer : IMovieScenePreAnimatedTok
 };
 
 /** A movie scene execution token that applies camera cuts */
-struct FApplyCameraAnimExecutionToken : IMovieSceneExecutionToken
+struct FApplyCameraAnimExecutionToken : IMovieSceneSharedExecutionToken
 {
-	static FMovieSceneAnimTypeID GetAnimTypeID()
+	FApplyCameraAnimExecutionToken(const FMovieSceneEvaluationOperand& InOperand)
 	{
-		return TMovieSceneAnimTypeID<FApplyCameraAnimExecutionToken>();
+		Operands.Add(InOperand);
+		// Evaluate after everything else in the group
+		Order = 1000;
 	}
 
-	/** Execute this token, operating on all objects referenced by 'Operand' */
-	virtual void Execute(const FMovieSceneContext& Context, const FMovieSceneEvaluationOperand& Operand, FPersistentEvaluationData& PersistentData, IMovieScenePlayer& Player) override
+	virtual void Execute(FPersistentEvaluationData& PersistentData, IMovieScenePlayer& Player) override
 	{
-		FMovieSceneAdditiveCameraData& SharedData = FMovieSceneAdditiveCameraData::Get(Operand, PersistentData);
-
-		for (TWeakObjectPtr<> ObjectWP : Player.FindBoundObjects(Operand))
+		for (FMovieSceneEvaluationOperand& Operand : Operands)
 		{
-			UCameraComponent* CameraComponent = MovieSceneHelpers::CameraComponentFromRuntimeObject(ObjectWP.Get());
-			if (CameraComponent)
+			FMovieSceneAdditiveCameraData& SharedData = FMovieSceneAdditiveCameraData::Get(Operand, PersistentData);
+
+			for (TWeakObjectPtr<> ObjectWP : Player.FindBoundObjects(Operand))
 			{
-				SharedData.ApplyCumulativeAnimation(*CameraComponent);
+				UCameraComponent* CameraComponent = MovieSceneHelpers::CameraComponentFromRuntimeObject(ObjectWP.Get());
+				if (CameraComponent)
+				{
+					SharedData.ApplyCumulativeAnimation(*CameraComponent);
+				}
 			}
+
+			SharedData.Reset();
 		}
 	}
-};
-
-const FMovieSceneSharedDataId FMovieSceneAdditiveCameraAnimationTrackTemplate::SharedDataId = FMovieSceneSharedDataId::Allocate();
 
-FMovieSceneAdditiveCameraAnimationTrackTemplate::FMovieSceneAdditiveCameraAnimationTrackTemplate()
-{
-}
-
-void FMovieSceneAdditiveCameraAnimationTrackTemplate::Initialize(const FMovieSceneEvaluationOperand& Operand, const FMovieSceneContext& Context, FPersistentEvaluationData& PersistentData, IMovieScenePlayer& Player) const
-{
-	FMovieSceneAdditiveCameraData::Get(Operand, PersistentData).Reset();
-}
-
-void FMovieSceneAdditiveCameraAnimationTrackTemplate::Evaluate(const FMovieSceneEvaluationOperand& Operand, const FMovieSceneContext& Context, const FPersistentEvaluationData& PersistentData, FMovieSceneExecutionTokens& ExecutionTokens) const
-{
-	// Add a final execution token that will apply the blended result of anything added to our shared data
-	ExecutionTokens.Add(FApplyCameraAnimExecutionToken());
-}
+	TSet<FMovieSceneEvaluationOperand> Operands;
+};
 
 /** A movie scene execution token that applies camera animations */
 struct FAccumulateCameraAnimExecutionToken : IMovieSceneExecutionToken
@@ -303,6 +307,19 @@ struct FAccumulateCameraAnimExecutionToken : IMovieSceneExecutionToken
 void FMovieSceneAdditiveCameraAnimationTemplate::Evaluate(const FMovieSceneEvaluationOperand& Operand, const FMovieSceneContext& Context, const FPersistentEvaluationData& PersistentData, FMovieSceneExecutionTokens& ExecutionTokens) const
 {
 	ExecutionTokens.Add(FAccumulateCameraAnimExecutionToken(this));
+
+	// Add a shared token that will apply the blended camera anims
+
+	FMovieSceneSharedDataId TokenID = GetSharedExecutionTokenID();
+	// It's safe to cast here as only FApplyCameraAnimExecutionTokens can have this token ID
+	if (FApplyCameraAnimExecutionToken* ExistingToken = static_cast<FApplyCameraAnimExecutionToken*>(ExecutionTokens.FindShared(TokenID)))
+	{
+		ExistingToken->Operands.Add(Operand);
+	}
+	else
+	{
+		ExecutionTokens.AddShared(GetSharedExecutionTokenID(), FApplyCameraAnimExecutionToken(Operand));
+	}
 }
 
 /** Persistent data that exists as long as a given camera anim section is being evaluated */
diff --git a/Engine/Source/Runtime/MovieSceneTracks/Private/Evaluation/MovieSceneCameraAnimTemplate.h b/Engine/Source/Runtime/MovieSceneTracks/Private/Evaluation/MovieSceneCameraAnimTemplate.h
--- a/Engine/Source/Runtime/MovieSceneTracks/Private/Evaluation/MovieSceneCameraAnimTemplate.h
+++ b/Engine/Source/Runtime/MovieSceneTracks/Private/Evaluation/MovieSceneCameraAnimTemplate.h
@@ -77,23 +77,3 @@ private:
 	UPROPERTY()
 	float SectionStartTime;
 };
-
-
-USTRUCT()
-struct FMovieSceneAdditiveCameraAnimationTrackTemplate : public FMovieSceneEvalTemplate
-{
-	GENERATED_BODY()
-
-	FMovieSceneAdditiveCameraAnimationTrackTemplate();
-
-	static const FMovieSceneSharedDataId SharedDataId;
-
-	virtual UScriptStruct& GetScriptStructImpl() const override { return *StaticStruct(); }
-	virtual void SetupOverrides() override
-	{
-		EnableOverrides(RequiresInitializeFlag);
-	}
-	
-	virtual void Initialize(const FMovieSceneEvaluationOperand& Operand, const FMovieSceneContext& Context, FPersistentEvaluationData& PersistentData, IMovieScenePlayer& Player) const override;
-	virtual void Evaluate(const FMovieSceneEvaluationOperand& Operand, const FMovieSceneContext& Context, const FPersistentEvaluationData& PersistentData, FMovieSceneExecutionTokens& ExecutionTokens) const override;
-};
diff --git a/Engine/Source/Runtime/MovieSceneTracks/Private/Evaluation/MovieSceneCameraCutTemplate.cpp b/Engine/Source/Runtime/MovieSceneTracks/Private/Evaluation/MovieSceneCameraCutTemplate.cpp
--- a/Engine/Source/Runtime/MovieSceneTracks/Private/Evaluation/MovieSceneCameraCutTemplate.cpp
+++ b/Engine/Source/Runtime/MovieSceneTracks/Private/Evaluation/MovieSceneCameraCutTemplate.cpp
@@ -12,12 +12,12 @@ DECLARE_CYCLE_STAT(TEXT("Camera Cut Track Token Execute"), MovieSceneEval_Camera
 /** A movie scene execution token that sets up the streaming system with the camera cut location */
 struct FCameraCutPreRollExecutionToken : IMovieSceneExecutionToken
 {
-	FGuid CameraGuid;
+	FMovieSceneObjectBindingID CameraBindingID;
 	FTransform CutTransform;
 	bool bHasCutTransform;
 
-	FCameraCutPreRollExecutionToken(const FGuid& InCameraGuid, const FTransform& InCutTransform, bool bInHasCutTransform)
-		: CameraGuid(InCameraGuid)
+	FCameraCutPreRollExecutionToken(const FMovieSceneObjectBindingID& InCameraBindingID, const FTransform& InCutTransform, bool bInHasCutTransform)
+		: CameraBindingID(InCameraBindingID)
 		, CutTransform(InCutTransform)
 		, bHasCutTransform(bInHasCutTransform)
 	{}
@@ -38,8 +38,16 @@ struct FCameraCutPreRollExecutionToken : IMovieSceneExecutionToken
 		}
 		else
 		{
+			FMovieSceneSequenceID SequenceID = Operand.SequenceID;
+			if (CameraBindingID.GetSequenceID().IsValid())
+			{
+				// Ensure that this ID is resolvable from the root, based on the current local sequence ID
+				FMovieSceneObjectBindingID RootBindingID = CameraBindingID.ResolveLocalToRoot(SequenceID, Player.GetEvaluationTemplate().GetHierarchy());
+				SequenceID = RootBindingID.GetSequenceID();
+			}
+
 			// If the transform is set, otherwise use the bound actor's transform
-			FMovieSceneEvaluationOperand CameraOperand(Operand.SequenceID, CameraGuid);
+			FMovieSceneEvaluationOperand CameraOperand(SequenceID, CameraBindingID.GetGuid());
 		
 			TArrayView<TWeakObjectPtr<>> Objects = Player.FindBoundObjects(CameraOperand);
 			if (!Objects.Num())
@@ -78,10 +86,10 @@ struct FCameraCutPreAnimatedToken : IMovieScenePreAnimatedGlobalToken
 /** A movie scene execution token that applies camera cuts */
 struct FCameraCutExecutionToken : IMovieSceneExecutionToken
 {
-	FGuid CameraGuid;
+	FMovieSceneObjectBindingID CameraBindingID;
 
-	FCameraCutExecutionToken(const FGuid& InCameraGuid)
-		: CameraGuid(InCameraGuid)
+	FCameraCutExecutionToken(const FMovieSceneObjectBindingID& InCameraBindingID)
+		: CameraBindingID(InCameraBindingID)
 	{}
 
 	static FMovieSceneAnimTypeID GetAnimTypeID()
@@ -92,9 +100,23 @@ struct FCameraCutExecutionToken : IMovieSceneExecutionToken
 	/** Execute this token, operating on all objects referenced by 'Operand' */
 	virtual void Execute(const FMovieSceneContext& Context, const FMovieSceneEvaluationOperand& Operand, FPersistentEvaluationData& PersistentData, IMovieScenePlayer& Player) override
 	{
+		if (!Player.CanUpdateCameraCut())
+		{
+			Player.UpdateCameraCut(nullptr, nullptr);
+			return;
+		}
+
 		MOVIESCENE_DETAILED_SCOPE_CYCLE_COUNTER(MovieSceneEval_CameraCutTrack_TokenExecute)
 
-		FMovieSceneEvaluationOperand CameraOperand(Operand.SequenceID, CameraGuid);
+		FMovieSceneSequenceID SequenceID = Operand.SequenceID;
+		if (CameraBindingID.GetSequenceID().IsValid())
+		{
+			// Ensure that this ID is resolvable from the root, based on the current local sequence ID
+			FMovieSceneObjectBindingID RootBindingID = CameraBindingID.ResolveLocalToRoot(SequenceID, Player.GetEvaluationTemplate().GetHierarchy());
+			SequenceID = RootBindingID.GetSequenceID();
+		}
+
+		FMovieSceneEvaluationOperand CameraOperand(SequenceID, CameraBindingID.GetGuid());
 		
 		TArrayView<TWeakObjectPtr<>> Objects = Player.FindBoundObjects(CameraOperand);
 		if (!Objects.Num())
@@ -130,7 +152,7 @@ struct FCameraCutExecutionToken : IMovieSceneExecutionToken
 
 
 FMovieSceneCameraCutSectionTemplate::FMovieSceneCameraCutSectionTemplate(const UMovieSceneCameraCutSection& Section, TOptional<FTransform> InCutTransform)
-	: CameraGuid(Section.GetCameraGuid())
+	: CameraBindingID(Section.GetCameraBindingID())
 	, CutTransform(InCutTransform.Get(FTransform()))
 	, bHasCutTransform(InCutTransform.IsSet())
 {
@@ -140,10 +162,10 @@ void FMovieSceneCameraCutSectionTemplate::Evaluate(const FMovieSceneEvaluationOp
 {
 	if (Context.IsPreRoll())
 	{
-		ExecutionTokens.Add(FCameraCutPreRollExecutionToken(CameraGuid, CutTransform, bHasCutTransform));
+		ExecutionTokens.Add(FCameraCutPreRollExecutionToken(CameraBindingID, CutTransform, bHasCutTransform));
 	}
 	else
 	{
-		ExecutionTokens.Add(FCameraCutExecutionToken(CameraGuid));
+		ExecutionTokens.Add(FCameraCutExecutionToken(CameraBindingID));
 	}
 }
diff --git a/Engine/Source/Runtime/MovieSceneTracks/Private/Evaluation/MovieSceneCameraCutTemplate.h b/Engine/Source/Runtime/MovieSceneTracks/Private/Evaluation/MovieSceneCameraCutTemplate.h
--- a/Engine/Source/Runtime/MovieSceneTracks/Private/Evaluation/MovieSceneCameraCutTemplate.h
+++ b/Engine/Source/Runtime/MovieSceneTracks/Private/Evaluation/MovieSceneCameraCutTemplate.h
@@ -7,6 +7,7 @@
 #include "Misc/Guid.h"
 #include "Evaluation/MovieSceneTrackImplementation.h"
 #include "Evaluation/MovieSceneEvalTemplate.h"
+#include "MovieSceneObjectBindingID.h"
 #include "MovieSceneCameraCutTemplate.generated.h"
 
 class UMovieSceneCameraCutSection;
@@ -22,7 +23,7 @@ struct FMovieSceneCameraCutSectionTemplate : public FMovieSceneEvalTemplate
 
 	/** GUID of the camera we should cut to in this sequence */
 	UPROPERTY()
-	FGuid CameraGuid;
+	FMovieSceneObjectBindingID CameraBindingID;
 
 	UPROPERTY()
 	FTransform CutTransform;
diff --git a/Engine/Source/Runtime/MovieSceneTracks/Private/Evaluation/MovieSceneMaterialParameterCollectionTemplate.cpp b/Engine/Source/Runtime/MovieSceneTracks/Private/Evaluation/MovieSceneMaterialParameterCollectionTemplate.cpp
--- a/Engine/Source/Runtime/MovieSceneTracks/Private/Evaluation/MovieSceneMaterialParameterCollectionTemplate.cpp
+++ b/Engine/Source/Runtime/MovieSceneTracks/Private/Evaluation/MovieSceneMaterialParameterCollectionTemplate.cpp
@@ -118,6 +118,12 @@ struct FMaterialParameterCollectionExecutionToken : IMovieSceneExecutionToken
 		static TMovieSceneAnimTypeIDContainer<FName> AnimTypeIDsByName;
 		UMaterialParameterCollectionInstance* Instance = World->GetParameterCollectionInstance(Collection);
 
+		if (!Instance)
+		{
+			UE_LOG(LogMovieScene, Warning, TEXT("Failed to get material parameter collection to animate."));
+			return;
+		}
+
 		TArray<FString> InvalidParameterNames;
 
 		for ( const FScalarParameterNameAndValue& ScalarNameAndValue : Values.ScalarValues )
diff --git a/Engine/Source/Runtime/MovieSceneTracks/Private/Evaluation/MovieSceneSkeletalAnimationTemplate.cpp b/Engine/Source/Runtime/MovieSceneTracks/Private/Evaluation/MovieSceneSkeletalAnimationTemplate.cpp
--- a/Engine/Source/Runtime/MovieSceneTracks/Private/Evaluation/MovieSceneSkeletalAnimationTemplate.cpp
+++ b/Engine/Source/Runtime/MovieSceneTracks/Private/Evaluation/MovieSceneSkeletalAnimationTemplate.cpp
@@ -32,22 +32,22 @@ void ResetAnimSequencerInstance(UObject& ObjectToRestore, IMovieScenePlayer& Pla
 
 struct FStopPlayingMontageTokenProducer : IMovieScenePreAnimatedTokenProducer
 {
-	TWeakObjectPtr<UAnimMontage> TempMontage;
+	TWeakObjectPtr<UAnimInstance> TempInstance;
 
-	FStopPlayingMontageTokenProducer(TWeakObjectPtr<UAnimMontage> InTempMontage) : TempMontage(InTempMontage) {}
+	FStopPlayingMontageTokenProducer(TWeakObjectPtr<UAnimInstance> InTempInstance) : TempInstance(InTempInstance) {}
 
 	virtual IMovieScenePreAnimatedTokenPtr CacheExistingState(UObject& Object) const
 	{
 		struct FToken : IMovieScenePreAnimatedToken
 		{
-			TWeakObjectPtr<UAnimMontage> WeakMontage;
+			TWeakObjectPtr<UAnimInstance> WeakInstance;
 
-			FToken(TWeakObjectPtr<UAnimMontage> InMontage) : WeakMontage(InMontage) {}
+			FToken(TWeakObjectPtr<UAnimInstance> InWeakInstance) : WeakInstance(InWeakInstance) {}
 
 			virtual void RestoreState(UObject& ObjectToRestore, IMovieScenePlayer& Player) override
 			{
-				UAnimInstance* AnimInstance = CastChecked<UAnimInstance>(&ObjectToRestore);
-				UAnimMontage* Montage = WeakMontage.Get();
+				UAnimMontage* Montage = CastChecked<UAnimMontage>(&ObjectToRestore);
+				UAnimInstance* AnimInstance = WeakInstance.Get();
 				if (AnimInstance && Montage)
 				{
 					AnimInstance->Montage_Stop(0.f, Montage);
@@ -55,7 +55,7 @@ struct FStopPlayingMontageTokenProducer : IMovieScenePreAnimatedTokenProducer
 			}
 		};
 
-		return FToken(TempMontage);
+		return FToken(TempInstance);
 	}
 };
 
@@ -145,10 +145,22 @@ namespace MovieScene
 
 		static FMovieSceneBlendingActuatorID GetActuatorTypeID()
 		{
-			static FMovieSceneAnimTypeID TypeID = TMovieSceneAnimTypeID<FComponentAnimationActuator>();
+			static FMovieSceneAnimTypeID TypeID = TMovieSceneAnimTypeID<FComponentAnimationActuator, 0>();
 			return FMovieSceneBlendingActuatorID(TypeID);
 		}
 
+		static FMovieSceneAnimTypeID GetMontageAnimationTypeID()
+		{
+			static FMovieSceneAnimTypeID TypeID = TMovieSceneAnimTypeID<FComponentAnimationActuator, 1>();
+			return TypeID;
+		}
+
+		static FMovieSceneAnimTypeID GetAnimControlTypeID()
+		{
+			static FMovieSceneAnimTypeID TypeID = TMovieSceneAnimTypeID<FComponentAnimationActuator, 2>();
+			return TypeID;
+		}
+
 		virtual FBlendedAnimation RetrieveCurrentValue(UObject* InObject, IMovieScenePlayer* Player) const
 		{
 			check(false);
@@ -165,8 +177,7 @@ namespace MovieScene
 				return;
 			}
 
-			static FMovieSceneAnimTypeID AnimTypeID = TMovieSceneAnimTypeID<FComponentAnimationActuator>();
-			OriginalStack.SavePreAnimatedState(Player, *SkeletalMeshComponent, AnimTypeID, FPreAnimatedAnimationTokenProducer());
+			OriginalStack.SavePreAnimatedState(Player, *SkeletalMeshComponent, GetAnimControlTypeID(), FPreAnimatedAnimationTokenProducer());
 
 			UAnimCustomInstance::BindToSkeletalMeshComponent<UAnimSequencerInstance>(SkeletalMeshComponent);
 
@@ -249,8 +260,8 @@ namespace MovieScene
 				UAnimInstance* AnimInst = SkeletalMeshComponent->GetAnimInstance();
 				if (AnimInst && Montage.IsValid())
 				{
-					FMovieSceneAnimTypeID SlotTypeID = MontageSlotAnimationIDs.GetAnimTypeID(SlotName);
-					Player.SavePreAnimatedState(*AnimInst, SlotTypeID, FStopPlayingMontageTokenProducer(Montage));
+					FMovieSceneAnimTypeID MontageAnimID = GetMontageAnimationTypeID();
+					Player.SavePreAnimatedState(*Montage.Get(), MontageAnimID, FStopPlayingMontageTokenProducer(AnimInst));
 
 					AnimInst->Montage_Resume(Montage.Get());
 				}
@@ -286,8 +297,8 @@ namespace MovieScene
 					if (Montage.IsValid())
 					{
 						// Unique anim type ID per slot
-						FMovieSceneAnimTypeID SlotTypeID = MontageSlotAnimationIDs.GetAnimTypeID(SlotName);
-						Player.SavePreAnimatedState(*AnimInst, SlotTypeID, FStopPlayingMontageTokenProducer(Montage));
+						FMovieSceneAnimTypeID MontageAnimID = GetMontageAnimationTypeID();
+						Player.SavePreAnimatedState(*Montage.Get(), MontageAnimID, FStopPlayingMontageTokenProducer(AnimInst));
 
 						if (bPlaying)
 						{
@@ -308,8 +319,6 @@ namespace MovieScene
 			}
 		}
 
-
-		TMovieSceneAnimTypeIDContainer<FName> MontageSlotAnimationIDs;
 		TMovieSceneAnimTypeIDContainer<FObjectKey> SectionToAnimationIDs;
 	};
 
diff --git a/Engine/Source/Runtime/MovieSceneTracks/Private/Sections/MovieScene3DAttachSection.cpp b/Engine/Source/Runtime/MovieSceneTracks/Private/Sections/MovieScene3DAttachSection.cpp
--- a/Engine/Source/Runtime/MovieSceneTracks/Private/Sections/MovieScene3DAttachSection.cpp
+++ b/Engine/Source/Runtime/MovieSceneTracks/Private/Sections/MovieScene3DAttachSection.cpp
@@ -23,11 +23,11 @@ UMovieScene3DAttachSection::UMovieScene3DAttachSection( const FObjectInitializer
 }
 
 
-void UMovieScene3DAttachSection::AddAttach( float Time, float SequenceEndTime, const FGuid& InAttachId )
+void UMovieScene3DAttachSection::AddAttach( float Time, float SequenceEndTime, const FMovieSceneObjectBindingID& InAttachBindingID )
 {
 	if (TryModify())
 	{
-		ConstraintId = InAttachId;
+		ConstraintBindingID = InAttachBindingID;
 	}
 }
 
diff --git a/Engine/Source/Runtime/MovieSceneTracks/Private/Sections/MovieScene3DConstraintSection.cpp b/Engine/Source/Runtime/MovieSceneTracks/Private/Sections/MovieScene3DConstraintSection.cpp
--- a/Engine/Source/Runtime/MovieSceneTracks/Private/Sections/MovieScene3DConstraintSection.cpp
+++ b/Engine/Source/Runtime/MovieSceneTracks/Private/Sections/MovieScene3DConstraintSection.cpp
@@ -1,6 +1,7 @@
 // Copyright 1998-2017 Epic Games, Inc. All Rights Reserved.
 
 #include "Sections/MovieScene3DConstraintSection.h"
+#include "MovieSceneObjectBindingID.h"
 
 
 UMovieScene3DConstraintSection::UMovieScene3DConstraintSection( const FObjectInitializer& ObjectInitializer )
@@ -12,14 +13,14 @@ void UMovieScene3DConstraintSection::SetConstraintId(const FGuid& InConstraintId
 {
 	if (TryModify())
 	{
-		ConstraintId = InConstraintId;
+		SetConstraintBindingID(FMovieSceneObjectBindingID(InConstraintId, MovieSceneSequenceID::Root, EMovieSceneObjectBindingSpace::Local));
 	}
 }
 
 
 FGuid UMovieScene3DConstraintSection::GetConstraintId() const
 {
-	return ConstraintId;
+	return FGuid();
 }
 
 
@@ -37,9 +38,22 @@ void UMovieScene3DConstraintSection::SetKeyTime(FKeyHandle KeyHandle, float Time
 
 void UMovieScene3DConstraintSection::OnBindingsUpdated(const TMap<FGuid, FGuid>& OldGuidToNewGuidMap)
 {
-	if (OldGuidToNewGuidMap.Contains(GetConstraintId()))
+	if (OldGuidToNewGuidMap.Contains(ConstraintBindingID.GetGuid()))
 	{
-		SetConstraintId(OldGuidToNewGuidMap[GetConstraintId()]);
+		ConstraintBindingID.SetGuid(OldGuidToNewGuidMap[ConstraintBindingID.GetGuid()]);
 	}
 }
 
+void UMovieScene3DConstraintSection::PostLoad()
+{
+	Super::PostLoad();
+
+	if (ConstraintId_DEPRECATED.IsValid())
+	{
+		if (!ConstraintBindingID.IsValid())
+		{
+			ConstraintBindingID = FMovieSceneObjectBindingID(ConstraintId_DEPRECATED, MovieSceneSequenceID::Root, EMovieSceneObjectBindingSpace::Local);
+		}
+		ConstraintId_DEPRECATED.Invalidate();
+	}
+}
diff --git a/Engine/Source/Runtime/MovieSceneTracks/Private/Sections/MovieScene3DPathSection.cpp b/Engine/Source/Runtime/MovieSceneTracks/Private/Sections/MovieScene3DPathSection.cpp
--- a/Engine/Source/Runtime/MovieSceneTracks/Private/Sections/MovieScene3DPathSection.cpp
+++ b/Engine/Source/Runtime/MovieSceneTracks/Private/Sections/MovieScene3DPathSection.cpp
@@ -145,11 +145,11 @@ void UMovieScene3DPathSection::GetKeyHandles(TSet<FKeyHandle>& OutKeyHandles, TR
 }
 
 
-void UMovieScene3DPathSection::AddPath( float Time, float SequenceEndTime, const FGuid& InPathId )
+void UMovieScene3DPathSection::AddPath( float Time, float SequenceEndTime, const FMovieSceneObjectBindingID& InPathBindingID )
 {
 	if (TryModify())
 	{
-		ConstraintId = InPathId;
+		ConstraintBindingID = InPathBindingID;
 
 		TimingCurve.UpdateOrAddKey(Time, 0);
 		TimingCurve.UpdateOrAddKey(SequenceEndTime, 1);
diff --git a/Engine/Source/Runtime/MovieSceneTracks/Private/Sections/MovieSceneCameraCutSection.cpp b/Engine/Source/Runtime/MovieSceneTracks/Private/Sections/MovieSceneCameraCutSection.cpp
--- a/Engine/Source/Runtime/MovieSceneTracks/Private/Sections/MovieSceneCameraCutSection.cpp
+++ b/Engine/Source/Runtime/MovieSceneTracks/Private/Sections/MovieSceneCameraCutSection.cpp
@@ -22,7 +22,7 @@ FMovieSceneEvalTemplatePtr UMovieSceneCameraCutSection::GenerateTemplate() const
 
 	for (const FMovieSceneBinding& Binding : MovieScene->GetBindings())
 	{
-		if (Binding.GetObjectGuid() == CameraGuid)
+		if (Binding.GetObjectGuid() == CameraBindingID.GetGuid())
 		{
 			for (UMovieSceneTrack* Track : Binding.GetTracks())
 			{
@@ -51,8 +51,22 @@ FMovieSceneEvalTemplatePtr UMovieSceneCameraCutSection::GenerateTemplate() const
 
 void UMovieSceneCameraCutSection::OnBindingsUpdated(const TMap<FGuid, FGuid>& OldGuidToNewGuidMap)
 {
-	if (OldGuidToNewGuidMap.Contains(CameraGuid))
+	if (OldGuidToNewGuidMap.Contains(CameraBindingID.GetGuid()))
 	{
-		CameraGuid = OldGuidToNewGuidMap[CameraGuid];
+		CameraBindingID.SetGuid(OldGuidToNewGuidMap[CameraBindingID.GetGuid()]);
+	}
+}
+
+void UMovieSceneCameraCutSection::PostLoad()
+{
+	Super::PostLoad();
+
+	if (CameraGuid_DEPRECATED.IsValid())
+	{
+		if (!CameraBindingID.IsValid())
+		{
+			CameraBindingID = FMovieSceneObjectBindingID(CameraGuid_DEPRECATED, MovieSceneSequenceID::Root, EMovieSceneObjectBindingSpace::Local);
+		}
+		CameraGuid_DEPRECATED.Invalidate();
 	}
 }
diff --git a/Engine/Source/Runtime/MovieSceneTracks/Private/Sections/MovieSceneCinematicShotSection.cpp b/Engine/Source/Runtime/MovieSceneTracks/Private/Sections/MovieSceneCinematicShotSection.cpp
--- a/Engine/Source/Runtime/MovieSceneTracks/Private/Sections/MovieSceneCinematicShotSection.cpp
+++ b/Engine/Source/Runtime/MovieSceneTracks/Private/Sections/MovieSceneCinematicShotSection.cpp
@@ -8,3 +8,14 @@
 
 UMovieSceneCinematicShotSection::UMovieSceneCinematicShotSection() : UMovieSceneSubSection()
 { }
+
+void UMovieSceneCinematicShotSection::PostLoad()
+{
+	Super::PostLoad();
+
+	if (!DisplayName_DEPRECATED.IsEmpty())
+	{
+		ShotDisplayName = DisplayName_DEPRECATED.ToString();
+		DisplayName_DEPRECATED = FText::GetEmpty();
+	}
+}
diff --git a/Engine/Source/Runtime/MovieSceneTracks/Private/Sections/MovieSceneEventSection.cpp b/Engine/Source/Runtime/MovieSceneTracks/Private/Sections/MovieSceneEventSection.cpp
--- a/Engine/Source/Runtime/MovieSceneTracks/Private/Sections/MovieSceneEventSection.cpp
+++ b/Engine/Source/Runtime/MovieSceneTracks/Private/Sections/MovieSceneEventSection.cpp
@@ -6,6 +6,7 @@
 #include "ReleaseObjectVersion.h"
 #include "LinkerLoad.h"
 #include "Serialization/MemoryArchive.h"
+#include "Engine/UserDefinedStruct.h"
 
 #include "Curves/KeyFrameAlgorithms.h"
 
@@ -72,6 +73,7 @@ public:
 	FEventParameterWriter(TArray<uint8>& InBytes)
 		: Bytes(InBytes)
 	{
+		ArNoDelta = true;
 		ArIsSaving = true;
 		ArIsPersistent = true;
 		UsingCustomVersion(EventParameterVersionGUID);
@@ -237,6 +239,14 @@ void FMovieSceneEventParameters::GetInstance(FStructOnScope& OutStruct) const
 {
 	UStruct* StructPtr = GetStructType();
 	OutStruct.Initialize(StructPtr);
+
+#if WITH_EDITOR
+	if (UUserDefinedStruct* UDS = Cast<UUserDefinedStruct>(StructPtr))
+	{
+		UDS->InitializeDefaultValue(OutStruct.GetStructMemory());
+	}
+#endif
+
 	uint8* Memory = OutStruct.GetStructMemory();
 	if (StructPtr && StructPtr->GetStructureSize() > 0 && StructBytes.Num())
 	{
diff --git a/Engine/Source/Runtime/MovieSceneTracks/Private/Tracks/MovieScene3DAttachTrack.cpp b/Engine/Source/Runtime/MovieSceneTracks/Private/Tracks/MovieScene3DAttachTrack.cpp
--- a/Engine/Source/Runtime/MovieSceneTracks/Private/Tracks/MovieScene3DAttachTrack.cpp
+++ b/Engine/Source/Runtime/MovieSceneTracks/Private/Tracks/MovieScene3DAttachTrack.cpp
@@ -19,12 +19,12 @@ FMovieSceneEvalTemplatePtr UMovieScene3DAttachTrack::CreateTemplateForSection(co
 	return FMovieScene3DAttachSectionTemplate(*CastChecked<UMovieScene3DAttachSection>(&InSection));
 }
 
-void UMovieScene3DAttachTrack::AddConstraint(float KeyTime, float ConstraintEndTime, const FName SocketName, const FName ComponentName, const FGuid& ConstraintId)
+void UMovieScene3DAttachTrack::AddConstraint(float KeyTime, float ConstraintEndTime, const FName SocketName, const FName ComponentName, const FMovieSceneObjectBindingID& ConstraintBindingID)
 {
 	// add the section
 	UMovieScene3DAttachSection* NewSection = NewObject<UMovieScene3DAttachSection>(this);
 	NewSection->SetFlags(RF_Transactional);
-	NewSection->AddAttach(KeyTime, ConstraintEndTime, ConstraintId);
+	NewSection->AddAttach(KeyTime, ConstraintEndTime, ConstraintBindingID);
 	NewSection->InitialPlacement( ConstraintSections, KeyTime, ConstraintEndTime, SupportsMultipleRows() );
 	NewSection->AttachSocketName = SocketName;
 	NewSection->AttachComponentName = ComponentName;
diff --git a/Engine/Source/Runtime/MovieSceneTracks/Private/Tracks/MovieScene3DConstraintTrack.cpp b/Engine/Source/Runtime/MovieSceneTracks/Private/Tracks/MovieScene3DConstraintTrack.cpp
--- a/Engine/Source/Runtime/MovieSceneTracks/Private/Tracks/MovieScene3DConstraintTrack.cpp
+++ b/Engine/Source/Runtime/MovieSceneTracks/Private/Tracks/MovieScene3DConstraintTrack.cpp
@@ -22,6 +22,11 @@ const TArray<UMovieSceneSection*>& UMovieScene3DConstraintTrack::GetAllSections(
 	return ConstraintSections;
 }
 
+void UMovieScene3DConstraintTrack::AddConstraint(float Time, float ConstraintEndTime, const FName SocketName, const FName ComponentName, const FGuid& ConstraintId)
+{
+	AddConstraint(Time, ConstraintEndTime, SocketName, ComponentName, FMovieSceneObjectBindingID(ConstraintId, MovieSceneSequenceID::Root));
+}
+
 
 void UMovieScene3DConstraintTrack::RemoveAllAnimationData()
 {
diff --git a/Engine/Source/Runtime/MovieSceneTracks/Private/Tracks/MovieScene3DPathTrack.cpp b/Engine/Source/Runtime/MovieSceneTracks/Private/Tracks/MovieScene3DPathTrack.cpp
--- a/Engine/Source/Runtime/MovieSceneTracks/Private/Tracks/MovieScene3DPathTrack.cpp
+++ b/Engine/Source/Runtime/MovieSceneTracks/Private/Tracks/MovieScene3DPathTrack.cpp
@@ -21,11 +21,11 @@ FMovieSceneEvalTemplatePtr UMovieScene3DPathTrack::CreateTemplateForSection(cons
 }
 
 
-void UMovieScene3DPathTrack::AddConstraint(float KeyTime, float ConstraintEndTime, const FName SocketName, const FName ComponentName, const FGuid& ConstraintId)
+void UMovieScene3DPathTrack::AddConstraint(float KeyTime, float ConstraintEndTime, const FName SocketName, const FName ComponentName, const FMovieSceneObjectBindingID& ConstraintBindingID)
 {
 	UMovieScene3DPathSection* NewSection = NewObject<UMovieScene3DPathSection>(this);
 	{
-		NewSection->AddPath(KeyTime, ConstraintEndTime, ConstraintId);
+		NewSection->AddPath(KeyTime, ConstraintEndTime, ConstraintBindingID);
 		NewSection->InitialPlacement( ConstraintSections, KeyTime, ConstraintEndTime, SupportsMultipleRows() );
 	}
 
diff --git a/Engine/Source/Runtime/MovieSceneTracks/Private/Tracks/MovieSceneAudioTrack.cpp b/Engine/Source/Runtime/MovieSceneTracks/Private/Tracks/MovieSceneAudioTrack.cpp
--- a/Engine/Source/Runtime/MovieSceneTracks/Private/Tracks/MovieSceneAudioTrack.cpp
+++ b/Engine/Source/Runtime/MovieSceneTracks/Private/Tracks/MovieSceneAudioTrack.cpp
@@ -138,20 +138,20 @@ bool UMovieSceneAudioTrack::IsAMasterTrack() const
 }
 
 
-TInlineValue<FMovieSceneSegmentCompilerRules> UMovieSceneAudioTrack::GetRowCompilerRules() const
+FMovieSceneTrackRowSegmentBlenderPtr UMovieSceneAudioTrack::GetRowSegmentBlender() const
 {
-	struct FCompilerRules : FMovieSceneSegmentCompilerRules
+	struct FBlender : FMovieSceneTrackRowSegmentBlender
 	{
-		virtual void BlendSegment(FMovieSceneSegment& Segment, const TArrayView<const FMovieSceneSectionData>& SourceData) const
+		virtual void Blend(FSegmentBlendData& BlendData) const override
 		{
 			// Run the default high pass filter for overlap priority
-			MovieSceneSegmentCompiler::BlendSegmentHighPass(Segment, SourceData);
+			MovieSceneSegmentCompiler::FilterOutUnderlappingSections(BlendData);
 
 			// Weed out based on array index (legacy behaviour)
-			MovieSceneSegmentCompiler::BlendSegmentLegacySectionOrder(Segment, SourceData);
+			MovieSceneSegmentCompiler::BlendSegmentLegacySectionOrder(BlendData);
 		}
 	};
-	return FCompilerRules();
+	return FBlender();
 }
 
 
diff --git a/Engine/Source/Runtime/MovieSceneTracks/Private/Tracks/MovieSceneBoolTrack.cpp b/Engine/Source/Runtime/MovieSceneTracks/Private/Tracks/MovieSceneBoolTrack.cpp
--- a/Engine/Source/Runtime/MovieSceneTracks/Private/Tracks/MovieSceneBoolTrack.cpp
+++ b/Engine/Source/Runtime/MovieSceneTracks/Private/Tracks/MovieSceneBoolTrack.cpp
@@ -14,21 +14,4 @@ UMovieSceneSection* UMovieSceneBoolTrack::CreateNewSection()
 FMovieSceneEvalTemplatePtr UMovieSceneBoolTrack::CreateTemplateForSection(const UMovieSceneSection& InSection) const
 {
 	return FMovieSceneBoolPropertySectionTemplate(*CastChecked<const UMovieSceneBoolSection>(&InSection), *this);
-}
-
-bool UMovieSceneBoolTrack::Eval( float Position, float LastPostion, bool& InOutBool ) const
-{	
-	const UMovieSceneSection* Section = MovieSceneHelpers::FindNearestSectionAtTime( Sections, Position );
-
-	if( Section )
-	{
-		if (!Section->IsInfinite())
-		{
-			Position = FMath::Clamp(Position, Section->GetStartTime(), Section->GetEndTime());
-		}
-
-		InOutBool = CastChecked<UMovieSceneBoolSection>( Section )->Eval( Position, InOutBool );
-	}
-
-	return (Section != nullptr);
-}
+}
\ No newline at end of file
diff --git a/Engine/Source/Runtime/MovieSceneTracks/Private/Tracks/MovieSceneCameraAnimTrack.cpp b/Engine/Source/Runtime/MovieSceneTracks/Private/Tracks/MovieSceneCameraAnimTrack.cpp
--- a/Engine/Source/Runtime/MovieSceneTracks/Private/Tracks/MovieSceneCameraAnimTrack.cpp
+++ b/Engine/Source/Runtime/MovieSceneTracks/Private/Tracks/MovieSceneCameraAnimTrack.cpp
@@ -29,22 +29,9 @@ void UMovieSceneCameraAnimTrack::AddNewCameraAnim(float KeyTime, UCameraAnim* Ca
 *****************************************************************************/
 
 
-void UMovieSceneCameraAnimTrack::PostCompile(FMovieSceneEvaluationTrack& OutTrack, const FMovieSceneTrackCompilerArgs& Args) const
+FMovieSceneTrackSegmentBlenderPtr UMovieSceneCameraAnimTrack::GetTrackSegmentBlender() const
 {
-	FMovieSceneSharedDataId UniqueId = FMovieSceneAdditiveCameraAnimationTrackTemplate::SharedDataId;
-
-	// Add a new shared track for the additive camera anim. There will only be one of these, and it will apply all the additive camera animations for this object.
-	FMovieSceneEvaluationTrack SharedTrackTemplate(Args.ObjectBindingId);
-	SharedTrackTemplate.DefineAsSingleTemplate(FMovieSceneAdditiveCameraAnimationTrackTemplate());
-	SharedTrackTemplate.SetEvaluationPriority(0xF);
-	
-	Args.Generator.AddSharedTrack(MoveTemp(SharedTrackTemplate), UniqueId, *this);
-}
-
-
-TInlineValue<FMovieSceneSegmentCompilerRules> UMovieSceneCameraAnimTrack::GetTrackCompilerRules() const
-{
-	return FMovieSceneAdditiveCameraRules(this);
+	return FMovieSceneAdditiveCameraTrackBlender();
 }
 
 
diff --git a/Engine/Source/Runtime/MovieSceneTracks/Private/Tracks/MovieSceneCameraCutTrack.cpp b/Engine/Source/Runtime/MovieSceneTracks/Private/Tracks/MovieSceneCameraCutTrack.cpp
--- a/Engine/Source/Runtime/MovieSceneTracks/Private/Tracks/MovieSceneCameraCutTrack.cpp
+++ b/Engine/Source/Runtime/MovieSceneTracks/Private/Tracks/MovieSceneCameraCutTrack.cpp
@@ -23,6 +23,11 @@ UMovieSceneCameraCutTrack::UMovieSceneCameraCutTrack( const FObjectInitializer&
 
 
 void UMovieSceneCameraCutTrack::AddNewCameraCut(FGuid CameraHandle, float StartTime)
+{
+	AddNewCameraCut(FMovieSceneObjectBindingID(CameraHandle, MovieSceneSequenceID::Root), StartTime);
+}
+
+UMovieSceneCameraCutSection* UMovieSceneCameraCutTrack::AddNewCameraCut(const FMovieSceneObjectBindingID& CameraBindingID, float StartTime)
 {
 	Modify();
 
@@ -43,7 +48,7 @@ void UMovieSceneCameraCutTrack::AddNewCameraCut(FGuid CameraHandle, float StartT
 	UMovieSceneCameraCutSection* NewSection = ExistingSection;
 	if (ExistingSection != nullptr)
 	{
-		ExistingSection->SetCameraGuid(CameraHandle);
+		ExistingSection->SetCameraBindingID(CameraBindingID);
 	}
 	else
 	{
@@ -51,7 +56,7 @@ void UMovieSceneCameraCutTrack::AddNewCameraCut(FGuid CameraHandle, float StartT
 		{
 			NewSection->SetStartTime(StartTime);
 			NewSection->SetEndTime(FindEndTimeForCameraCut(StartTime));
-			NewSection->SetCameraGuid(CameraHandle);
+			NewSection->SetCameraBindingID(CameraBindingID);
 		}
 
 		Sections.Add(NewSection);
@@ -62,6 +67,8 @@ void UMovieSceneCameraCutTrack::AddNewCameraCut(FGuid CameraHandle, float StartT
 
 	// Once CameraCuts are sorted fixup the surrounding CameraCuts to fix any gaps
 	MovieSceneHelpers::FixupConsecutiveSections(Sections, *NewSection, false);
+
+	return NewSection;
 }
 
 
diff --git a/Engine/Source/Runtime/MovieSceneTracks/Private/Tracks/MovieSceneCameraShakeTrack.cpp b/Engine/Source/Runtime/MovieSceneTracks/Private/Tracks/MovieSceneCameraShakeTrack.cpp
--- a/Engine/Source/Runtime/MovieSceneTracks/Private/Tracks/MovieSceneCameraShakeTrack.cpp
+++ b/Engine/Source/Runtime/MovieSceneTracks/Private/Tracks/MovieSceneCameraShakeTrack.cpp
@@ -24,21 +24,9 @@ void UMovieSceneCameraShakeTrack::AddNewCameraShake(float KeyTime, TSubclassOf<U
 	}
 }
 
-void UMovieSceneCameraShakeTrack::PostCompile(FMovieSceneEvaluationTrack& OutTrack, const FMovieSceneTrackCompilerArgs& Args) const
+FMovieSceneTrackSegmentBlenderPtr UMovieSceneCameraShakeTrack::GetTrackSegmentBlender() const
 {
-	FMovieSceneSharedDataId UniqueId = FMovieSceneAdditiveCameraAnimationTrackTemplate::SharedDataId;
-
-	// Add a new shared track for the additive camera anim. There will only be one of these, and it will apply all the additive camera animations for this object.
-	FMovieSceneEvaluationTrack SharedTrackTemplate(Args.ObjectBindingId);
-	SharedTrackTemplate.DefineAsSingleTemplate(FMovieSceneAdditiveCameraAnimationTrackTemplate());
-	SharedTrackTemplate.SetEvaluationPriority(0xF);
-	
-	Args.Generator.AddSharedTrack(MoveTemp(SharedTrackTemplate), UniqueId, *this);
-}
-
-TInlineValue<FMovieSceneSegmentCompilerRules> UMovieSceneCameraShakeTrack::GetTrackCompilerRules() const
-{
-	return FMovieSceneAdditiveCameraRules(this);
+	return FMovieSceneAdditiveCameraTrackBlender();
 }
 
 #if WITH_EDITORONLY_DATA
diff --git a/Engine/Source/Runtime/MovieSceneTracks/Private/Tracks/MovieSceneCinematicShotTrack.cpp b/Engine/Source/Runtime/MovieSceneTracks/Private/Tracks/MovieSceneCinematicShotTrack.cpp
--- a/Engine/Source/Runtime/MovieSceneTracks/Private/Tracks/MovieSceneCinematicShotTrack.cpp
+++ b/Engine/Source/Runtime/MovieSceneTracks/Private/Tracks/MovieSceneCinematicShotTrack.cpp
@@ -30,7 +30,7 @@ UMovieSceneSubSection* UMovieSceneCinematicShotTrack::AddSequence(UMovieSceneSeq
 
 	if (Sequence != nullptr)
 	{
-		NewShotSection->SetShotDisplayName(Sequence->GetDisplayName());	
+		NewShotSection->SetShotDisplayName(Sequence->GetDisplayName().ToString());	
 	}
 
 #endif
@@ -75,59 +75,36 @@ bool UMovieSceneCinematicShotTrack::SupportsMultipleRows() const
 	return true;
 }
 
-TInlineValue<FMovieSceneSegmentCompilerRules> UMovieSceneCinematicShotTrack::GetTrackCompilerRules() const
+FMovieSceneTrackSegmentBlenderPtr UMovieSceneCinematicShotTrack::GetTrackSegmentBlender() const
 {
 	// Apply a high pass filter to overlapping sections such that only the highest row in a track wins
-	struct FCinematicShotTrackCompilerRules : FMovieSceneSegmentCompilerRules
+	struct FCinematicShotTrackRowBlender : FMovieSceneTrackSegmentBlender
 	{
-		virtual void BlendSegment(FMovieSceneSegment& Segment, const TArrayView<const FMovieSceneSectionData>& SourceData) const
+		virtual void Blend(FSegmentBlendData& BlendData) const override
 		{
-			MovieSceneSegmentCompiler::BlendSegmentHighPass(Segment, SourceData);
+			MovieSceneSegmentCompiler::ChooseLowestRowIndex(BlendData);
 		}
 	};
-	return FCinematicShotTrackCompilerRules();
+	return FCinematicShotTrackRowBlender();
 }
 
-TInlineValue<FMovieSceneSegmentCompilerRules> UMovieSceneCinematicShotTrack::GetRowCompilerRules() const
+FMovieSceneTrackRowSegmentBlenderPtr UMovieSceneCinematicShotTrack::GetRowSegmentBlender() const
 {
-	class FCinematicRowRules : public FMovieSceneSegmentCompilerRules
+	class FCinematicRowRules : public FMovieSceneTrackRowSegmentBlender
 	{
-		virtual void BlendSegment(FMovieSceneSegment& Segment, const TArrayView<const FMovieSceneSectionData>& SourceData) const override
+		virtual void Blend(FSegmentBlendData& BlendData) const override
 		{
 			// Sort everything by priority, then latest start time wins
-			if (Segment.Impls.Num() <= 1)
+			if (BlendData.Num() <= 1)
 			{
 				return;
 			}
 
-			Segment.Impls.Sort(
-				[&](const FSectionEvaluationData& A, const FSectionEvaluationData& B)
-				{
-					const FMovieSceneSectionData& DataA = SourceData[A.ImplIndex];
-					const FMovieSceneSectionData& DataB = SourceData[B.ImplIndex];
-					
-					// Always sort pre/postroll to the front of the array
-					const bool PrePostRollA = DataA.EvalData.IsPreRoll() || DataA.EvalData.IsPostRoll();
-					const bool PrePostRollB = DataB.EvalData.IsPreRoll() || DataB.EvalData.IsPostRoll();
-					if (PrePostRollA != PrePostRollB)
-					{
-						return PrePostRollA;
-					}
-					else if (PrePostRollA)
-					{
-						return false;
-					}
-					else if (DataA.Priority == DataB.Priority)
-					{
-						return TRangeBound<float>::MaxLower(DataA.Bounds.GetLowerBound(), DataB.Bounds.GetLowerBound()) == DataA.Bounds.GetLowerBound();
-					}
-					return DataA.Priority > DataB.Priority;
-				}
-			);
+			BlendData.Sort(SortPredicate);
 
 			int32 RemoveAtIndex = 0;
 			// Skip over any pre/postroll sections
-			while (Segment.Impls.IsValidIndex(RemoveAtIndex) && (Segment.Impls[RemoveAtIndex].Flags & (ESectionEvaluationFlags::PreRoll | ESectionEvaluationFlags::PostRoll)) != ESectionEvaluationFlags::None)
+			while (BlendData.IsValidIndex(RemoveAtIndex) && EnumHasAnyFlags(BlendData[RemoveAtIndex].Flags, ESectionEvaluationFlags::PreRoll | ESectionEvaluationFlags::PostRoll))
 			{
 				++RemoveAtIndex;
 			}
@@ -135,12 +112,36 @@ TInlineValue<FMovieSceneSegmentCompilerRules> UMovieSceneCinematicShotTrack::Get
 			// Skip over the first genuine evaluation if it exists
 			++RemoveAtIndex;
 
-			int32 NumToRemove = Segment.Impls.Num() - RemoveAtIndex;
+			int32 NumToRemove = BlendData.Num() - RemoveAtIndex;
 			if (NumToRemove > 0)
 			{
-				Segment.Impls.RemoveAt(RemoveAtIndex, NumToRemove, true);
+				BlendData.RemoveAt(RemoveAtIndex, NumToRemove, true);
 			}
 		}
+
+		static bool SortPredicate(const FMovieSceneSectionData& A, const FMovieSceneSectionData& B)
+		{
+			// Always sort pre/postroll to the front of the array
+			const bool PrePostRollA = EnumHasAnyFlags(A.Flags, ESectionEvaluationFlags::PreRoll | ESectionEvaluationFlags::PostRoll);
+			const bool PrePostRollB = EnumHasAnyFlags(B.Flags, ESectionEvaluationFlags::PreRoll | ESectionEvaluationFlags::PostRoll);
+
+			if (PrePostRollA != PrePostRollB)
+			{
+				return PrePostRollA;
+			}
+			else if (PrePostRollA)
+			{
+				return false;
+			}
+			else if (A.Section->GetOverlapPriority() == B.Section->GetOverlapPriority())
+			{
+				TRangeBound<float> StartBoundA = A.Section->IsInfinite() ? TRangeBound<float>::Open() : TRangeBound<float>::Inclusive(A.Section->GetStartTime());
+				TRangeBound<float> StartBoundB = B.Section->IsInfinite() ? TRangeBound<float>::Open() : TRangeBound<float>::Inclusive(B.Section->GetStartTime());
+
+				return TRangeBound<float>::MaxLower(StartBoundA, StartBoundB) == StartBoundA;
+			}
+			return A.Section->GetOverlapPriority() > B.Section->GetOverlapPriority();
+		}
 	};
 
 	return FCinematicRowRules();
diff --git a/Engine/Source/Runtime/MovieSceneTracks/Private/Tracks/MovieSceneSkeletalAnimationTrack.cpp b/Engine/Source/Runtime/MovieSceneTracks/Private/Tracks/MovieSceneSkeletalAnimationTrack.cpp
--- a/Engine/Source/Runtime/MovieSceneTracks/Private/Tracks/MovieSceneSkeletalAnimationTrack.cpp
+++ b/Engine/Source/Runtime/MovieSceneTracks/Private/Tracks/MovieSceneSkeletalAnimationTrack.cpp
@@ -140,23 +140,23 @@ FText UMovieSceneSkeletalAnimationTrack::GetDefaultDisplayName() const
 
 #endif
 
-TInlineValue<FMovieSceneSegmentCompilerRules> UMovieSceneSkeletalAnimationTrack::GetRowCompilerRules() const
+FMovieSceneTrackRowSegmentBlenderPtr UMovieSceneSkeletalAnimationTrack::GetRowSegmentBlender() const
 {
 	// Apply an upper bound exclusive blend
-	struct FSkeletalAnimationRowCompilerRules : FMovieSceneSegmentCompilerRules
+	struct FSkeletalAnimationRowCompilerRules : FMovieSceneTrackRowSegmentBlender
 	{
 		bool bUseLegacySectionIndexBlend;
 		FSkeletalAnimationRowCompilerRules(bool bInUseLegacySectionIndexBlend) : bUseLegacySectionIndexBlend(bInUseLegacySectionIndexBlend) {}
 
-		virtual void BlendSegment(FMovieSceneSegment& Segment, const TArrayView<const FMovieSceneSectionData>& SourceData) const
+		virtual void Blend(FSegmentBlendData& BlendData) const override
 		{
 			// Run the default high pass filter for overlap priority
-			MovieSceneSegmentCompiler::BlendSegmentHighPass(Segment, SourceData);
+			MovieSceneSegmentCompiler::FilterOutUnderlappingSections(BlendData);
 
 			if (bUseLegacySectionIndexBlend)
 			{
 				// Weed out based on array index (legacy behaviour)
-				MovieSceneSegmentCompiler::BlendSegmentLegacySectionOrder(Segment, SourceData);
+				MovieSceneSegmentCompiler::BlendSegmentLegacySectionOrder(BlendData);
 			}
 		}
 	};
diff --git a/Engine/Source/Runtime/MovieSceneTracks/Private/Tracks/MovieSceneStringTrack.cpp b/Engine/Source/Runtime/MovieSceneTracks/Private/Tracks/MovieSceneStringTrack.cpp
--- a/Engine/Source/Runtime/MovieSceneTracks/Private/Tracks/MovieSceneStringTrack.cpp
+++ b/Engine/Source/Runtime/MovieSceneTracks/Private/Tracks/MovieSceneStringTrack.cpp
@@ -38,24 +38,6 @@ bool UMovieSceneStringTrack::AddKeyToSection(float Time, const FString& String)
 }
 
 
-bool UMovieSceneStringTrack::Eval(float Position, float LastPostion, FString& InOutString) const
-{
-	const UMovieSceneSection* Section = MovieSceneHelpers::FindNearestSectionAtTime(Sections, Position);
-
-	if (Section)
-	{
-		if (!Section->IsInfinite())
-		{
-			Position = FMath::Clamp(Position, Section->GetStartTime(), Section->GetEndTime());
-		}
-
-		InOutString = CastChecked<UMovieSceneStringSection>(Section)->Eval(Position, InOutString);
-	}
-
-	return Section != nullptr;
-}
-
-
 /* UMovieSceneTrack interface
  *****************************************************************************/
 
diff --git a/Engine/Source/Runtime/MovieSceneTracks/Private/Tracks/MovieSceneVectorTrack.cpp b/Engine/Source/Runtime/MovieSceneTracks/Private/Tracks/MovieSceneVectorTrack.cpp
--- a/Engine/Source/Runtime/MovieSceneTracks/Private/Tracks/MovieSceneVectorTrack.cpp
+++ b/Engine/Source/Runtime/MovieSceneTracks/Private/Tracks/MovieSceneVectorTrack.cpp
@@ -9,6 +9,7 @@ UMovieSceneVectorTrack::UMovieSceneVectorTrack( const FObjectInitializer& Object
 	: Super( ObjectInitializer )
 {
 	NumChannelsUsed = 0;
+	SupportedBlendTypes = FMovieSceneBlendTypeField::All();
 }
 
 
diff --git a/Engine/Source/Runtime/MovieSceneTracks/Public/Evaluation/MovieScene3DTransformTemplate.h b/Engine/Source/Runtime/MovieSceneTracks/Public/Evaluation/MovieScene3DTransformTemplate.h
--- a/Engine/Source/Runtime/MovieSceneTracks/Public/Evaluation/MovieScene3DTransformTemplate.h
+++ b/Engine/Source/Runtime/MovieSceneTracks/Public/Evaluation/MovieScene3DTransformTemplate.h
@@ -56,6 +56,8 @@ struct FMovieSceneComponentTransformSectionTemplate : public FMovieSceneEvalTemp
 protected:
 
 	virtual UScriptStruct& GetScriptStructImpl() const override { return *StaticStruct(); }
+	virtual void SetupOverrides() override { EnableOverrides(RequiresInitializeFlag); }
+	virtual void Initialize(const FMovieSceneEvaluationOperand& Operand, const FMovieSceneContext& Context, FPersistentEvaluationData& PersistentData, IMovieScenePlayer& Player) const override;
 	virtual void Evaluate(const FMovieSceneEvaluationOperand& Operand, const FMovieSceneContext& Context, const FPersistentEvaluationData& PersistentData, FMovieSceneExecutionTokens& ExecutionTokens) const override;
 	virtual void Interrogate(const FMovieSceneContext& Context, FMovieSceneInterrogationData& Container, UObject* BindingOverride) const override;
 };
diff --git a/Engine/Source/Runtime/MovieSceneTracks/Public/Sections/MovieScene3DAttachSection.h b/Engine/Source/Runtime/MovieSceneTracks/Public/Sections/MovieScene3DAttachSection.h
--- a/Engine/Source/Runtime/MovieSceneTracks/Public/Sections/MovieScene3DAttachSection.h
+++ b/Engine/Source/Runtime/MovieSceneTracks/Public/Sections/MovieScene3DAttachSection.h
@@ -27,9 +27,9 @@ public:
 	 *
 	 * @param Time	The location in time where the attach should be added
 	 * @param SequenceEndTime   The time at the end of the sequence, by default the attach is set to end at this time
-	 * @param InAttachId The id to the path
+	 * @param InAttachBindingId The object binding id to the path
 	 */
-	void AddAttach(float Time, float SequenceEndTime, const FGuid& InAttachId);
+	void AddAttach(float Time, float SequenceEndTime, const FMovieSceneObjectBindingID& InAttachBindingID);
 	
 public:
 
diff --git a/Engine/Source/Runtime/MovieSceneTracks/Public/Sections/MovieScene3DConstraintSection.h b/Engine/Source/Runtime/MovieSceneTracks/Public/Sections/MovieScene3DConstraintSection.h
--- a/Engine/Source/Runtime/MovieSceneTracks/Public/Sections/MovieScene3DConstraintSection.h
+++ b/Engine/Source/Runtime/MovieSceneTracks/Public/Sections/MovieScene3DConstraintSection.h
@@ -7,6 +7,7 @@
 #include "Misc/Guid.h"
 #include "Curves/KeyHandle.h"
 #include "MovieSceneSection.h"
+#include "MovieSceneObjectBindingID.h"
 #include "MovieScene3DConstraintSection.generated.h"
 
 
@@ -24,9 +25,21 @@ public:
 	/** Sets the constraint id for this section */
 	virtual void SetConstraintId(const FGuid& InId);
 
-	/** Gets the constraint id for this section */
+	DEPRECATED(4.18, "Constraint guid no longer supported, Use GetConstraintBindingID.")
 	virtual FGuid GetConstraintId() const;
 
+	/** Gets the constraint binding for this Constraint section */
+	const FMovieSceneObjectBindingID& GetConstraintBindingID() const
+	{
+		return ConstraintBindingID;
+	}
+
+	/** Sets the constraint binding for this Constraint section */
+	void SetConstraintBindingID(const FMovieSceneObjectBindingID& InConstraintBindingID)
+	{
+		ConstraintBindingID = InConstraintBindingID;
+	}
+
 public:
 
 	//~ UMovieSceneSection interface
@@ -35,9 +48,17 @@ public:
 	virtual void SetKeyTime(FKeyHandle KeyHandle, float Time) override;
 	virtual void OnBindingsUpdated(const TMap<FGuid, FGuid>& OldGuidToNewGuidMap) override;
 	
+	/** ~UObject interface */
+	virtual void PostLoad() override;
+
 protected:
 
 	/** The possessable guid that this constraint uses */
 	UPROPERTY()
-	FGuid ConstraintId;
+	FGuid ConstraintId_DEPRECATED;
+
+	/** The constraint binding that this movie Constraint uses */
+	UPROPERTY(EditAnywhere, Category="Section")
+	FMovieSceneObjectBindingID ConstraintBindingID;
+
 };
diff --git a/Engine/Source/Runtime/MovieSceneTracks/Public/Sections/MovieScene3DPathSection.h b/Engine/Source/Runtime/MovieSceneTracks/Public/Sections/MovieScene3DPathSection.h
--- a/Engine/Source/Runtime/MovieSceneTracks/Public/Sections/MovieScene3DPathSection.h
+++ b/Engine/Source/Runtime/MovieSceneTracks/Public/Sections/MovieScene3DPathSection.h
@@ -47,9 +47,9 @@ public:
 	 *
 	 * @param Time The location in time where the path should be added.
 	 * @param SequenceEndTime The time at the end of the sequence, by default the path is set to end at this time.
-	 * @param InPathId The id to the path.
+	 * @param InPathBindingID The object binding id to the path.
 	 */
-	void AddPath(float Time, float SequenceEndTime, const FGuid& InPathId);
+	void AddPath(float Time, float SequenceEndTime, const FMovieSceneObjectBindingID& InPathBindingID);
 
 	/** 
 	 * Returns the timing curve.
diff --git a/Engine/Source/Runtime/MovieSceneTracks/Public/Sections/MovieSceneCameraCutSection.h b/Engine/Source/Runtime/MovieSceneTracks/Public/Sections/MovieSceneCameraCutSection.h
--- a/Engine/Source/Runtime/MovieSceneTracks/Public/Sections/MovieSceneCameraCutSection.h
+++ b/Engine/Source/Runtime/MovieSceneTracks/Public/Sections/MovieSceneCameraCutSection.h
@@ -8,6 +8,7 @@
 #include "Curves/KeyHandle.h"
 #include "MovieSceneSection.h"
 #include "SequencerObjectVersion.h"
+#include "MovieSceneObjectBindingID.h"
 #include "MovieSceneCameraCutSection.generated.h"
 
 
@@ -30,16 +31,28 @@ public:
 				EMovieSceneCompletionMode::ProjectDefault);
 	}
 
-	/** Gets the camera guid for this CameraCut section */
+	DEPRECATED(4.18, "Camera guid no longer supported, Use GetCameraBindingID.")
 	FGuid GetCameraGuid() const
 	{
-		return CameraGuid;
+		return FGuid();
 	}
 
-	/** Sets the camera guid for this CameraCut section */
+	/** Sets the camera binding for this CameraCut section. Evaluates from the sequence binding ID */
 	void SetCameraGuid(const FGuid& InGuid)
 	{
-		CameraGuid = InGuid;
+		SetCameraBindingID(FMovieSceneObjectBindingID(InGuid, MovieSceneSequenceID::Root, EMovieSceneObjectBindingSpace::Local));
+	}
+
+	/** Gets the camera binding for this CameraCut section */
+	const FMovieSceneObjectBindingID& GetCameraBindingID() const
+	{
+		return CameraBindingID;
+	}
+
+	/** Sets the camera binding for this CameraCut section */
+	void SetCameraBindingID(const FMovieSceneObjectBindingID& InCameraBindingID)
+	{
+		CameraBindingID = InCameraBindingID;
 	}
 
 	virtual FMovieSceneEvalTemplatePtr GenerateTemplate() const override;
@@ -49,11 +62,18 @@ public:
 	virtual void SetKeyTime(FKeyHandle KeyHandle, float Time) override { }
 	virtual void OnBindingsUpdated(const TMap<FGuid, FGuid>& OldGuidToNewGuidMap) override;
 
+	/** ~UObject interface */
+	virtual void PostLoad() override;
+
 private:
 
 	/** The camera possessable or spawnable that this movie CameraCut uses */
 	UPROPERTY()
-	FGuid CameraGuid;
+	FGuid CameraGuid_DEPRECATED;
+
+	/** The camera binding that this movie CameraCut uses */
+	UPROPERTY(EditAnywhere, Category="Section")
+	FMovieSceneObjectBindingID CameraBindingID;
 
 #if WITH_EDITORONLY_DATA
 public:
diff --git a/Engine/Source/Runtime/MovieSceneTracks/Public/Sections/MovieSceneCinematicShotSection.h b/Engine/Source/Runtime/MovieSceneTracks/Public/Sections/MovieSceneCinematicShotSection.h
--- a/Engine/Source/Runtime/MovieSceneTracks/Public/Sections/MovieSceneCinematicShotSection.h
+++ b/Engine/Source/Runtime/MovieSceneTracks/Public/Sections/MovieSceneCinematicShotSection.h
@@ -19,20 +19,23 @@ class MOVIESCENETRACKS_API UMovieSceneCinematicShotSection
 	/** Default constructor. */
 	UMovieSceneCinematicShotSection();
 
+	/** ~UObject interface */
+	virtual void PostLoad() override;
+
 public:
 
 	/** @return The shot display name */
-	FText GetShotDisplayName() const
+	FString GetShotDisplayName() const
 	{
-		return DisplayName;
+		return ShotDisplayName;
 	}
 
 	/** Set the shot display name */
-	void SetShotDisplayName(const FText& InDisplayName)
+	void SetShotDisplayName(const FString& InShotDisplayName)
 	{
 		if (TryModify())
 		{
-			DisplayName = InDisplayName;
+			ShotDisplayName = InShotDisplayName;
 		}
 	}
 
@@ -40,7 +43,11 @@ private:
 
 	/** The Shot's display name */
 	UPROPERTY()
-	FText DisplayName;
+	FString ShotDisplayName;
+
+	/** The Shot's display name */
+	UPROPERTY()
+	FText DisplayName_DEPRECATED;
 
 #if WITH_EDITORONLY_DATA
 public:
diff --git a/Engine/Source/Runtime/MovieSceneTracks/Public/Tracks/IMovieSceneTransformOrigin.h b/Engine/Source/Runtime/MovieSceneTracks/Public/Tracks/IMovieSceneTransformOrigin.h
new file mode 100644
--- /dev/null
+++ b/Engine/Source/Runtime/MovieSceneTracks/Public/Tracks/IMovieSceneTransformOrigin.h
@@ -0,0 +1,36 @@
+// Copyright 1998-2017 Epic Games, Inc. All Rights Reserved.
+
+#pragma once
+
+#include "UObject/ObjectMacros.h"
+#include "UObject/Interface.h"
+
+#include "IMovieSceneTransformOrigin.generated.h"
+
+
+UINTERFACE(Category="Sequencer", Blueprintable, meta=(DisplayName = "Transform Origin"))
+class MOVIESCENETRACKS_API UMovieSceneTransformOrigin : public UInterface
+{
+public:
+	GENERATED_BODY()
+};
+
+/** Interface that should be implemented to provide transform tracks with an origin transform. Scale is ignored. */
+class MOVIESCENETRACKS_API IMovieSceneTransformOrigin
+{
+public:
+	GENERATED_BODY()
+
+	/** Get the transform origin for this interface */
+	FTransform GetTransformOrigin() const
+	{
+		return NativeGetTransformOrigin();
+	}
+
+protected:
+
+	/** Get the transform from which all absolute component transform sections should be relative. Scale is ignored. */
+	UFUNCTION(BlueprintCallable, BlueprintImplementableEvent, Category = "Sequencer", DisplayName = "GetTransformOrigin", meta=(CallInEditor="true"))
+	FTransform BP_GetTransformOrigin() const;
+	virtual FTransform NativeGetTransformOrigin() const { return BP_GetTransformOrigin(); }
+};
diff --git a/Engine/Source/Runtime/MovieSceneTracks/Public/Tracks/MovieScene3DAttachTrack.h b/Engine/Source/Runtime/MovieSceneTracks/Public/Tracks/MovieScene3DAttachTrack.h
--- a/Engine/Source/Runtime/MovieSceneTracks/Public/Tracks/MovieScene3DAttachTrack.h
+++ b/Engine/Source/Runtime/MovieSceneTracks/Public/Tracks/MovieScene3DAttachTrack.h
@@ -20,7 +20,7 @@ public:
 
 	// UMovieScene3DConstraintTrack interface
 
-	virtual void AddConstraint( float Time, float ConstraintEndTime, const FName SocketName, const FName ComponentName, const FGuid& ConstraintId ) override;
+	virtual void AddConstraint( float Time, float ConstraintEndTime, const FName SocketName, const FName ComponentName, const FMovieSceneObjectBindingID& ConstraintBindingID ) override;
 	virtual class UMovieSceneSection* CreateNewSection() override;
 
 public:
diff --git a/Engine/Source/Runtime/MovieSceneTracks/Public/Tracks/MovieScene3DConstraintTrack.h b/Engine/Source/Runtime/MovieSceneTracks/Public/Tracks/MovieScene3DConstraintTrack.h
--- a/Engine/Source/Runtime/MovieSceneTracks/Public/Tracks/MovieScene3DConstraintTrack.h
+++ b/Engine/Source/Runtime/MovieSceneTracks/Public/Tracks/MovieScene3DConstraintTrack.h
@@ -5,6 +5,7 @@
 #include "CoreMinimal.h"
 #include "UObject/ObjectMacros.h"
 #include "MovieSceneTrack.h"
+#include "MovieSceneObjectBindingID.h"
 #include "MovieScene3DConstraintTrack.generated.h"
 
 /**
@@ -18,6 +19,9 @@ class UMovieScene3DConstraintTrack
 
 public:
 
+	DEPRECATED(4.18, "Constraint guid no longer supported, Use ConstraintBindingID.")
+	virtual void AddConstraint(float Time, float ConstraintEndTime, const FName SocketName, const FName ComponentName, const FGuid& ConstraintId);
+
 	/**
 	 * Adds a constraint.
 	 *
@@ -25,9 +29,9 @@ public:
 	 * @param ConstraintEndTime Set the constraint to end at this time.
 	 * @param SocketName The socket name for the constraint.
 	 * @param ComponentName The name of the component the socket resides in.
-	 * @param ConstraintId The id to the constraint.
+	 * @param FMovieSceneObjectBindingID The object binding id to the constraint.
 	 */
-	virtual void AddConstraint(float Time, float ConstraintEndTime, const FName SocketName, const FName ComponentName, const FGuid& ConstraintId) { }
+	virtual void AddConstraint(float Time, float ConstraintEndTime, const FName SocketName, const FName ComponentName, const FMovieSceneObjectBindingID& ConstraintBindingID) { }
 
 public:
 
diff --git a/Engine/Source/Runtime/MovieSceneTracks/Public/Tracks/MovieScene3DPathTrack.h b/Engine/Source/Runtime/MovieSceneTracks/Public/Tracks/MovieScene3DPathTrack.h
--- a/Engine/Source/Runtime/MovieSceneTracks/Public/Tracks/MovieScene3DPathTrack.h
+++ b/Engine/Source/Runtime/MovieSceneTracks/Public/Tracks/MovieScene3DPathTrack.h
@@ -20,7 +20,7 @@ public:
 
 	// UMovieScene3DConstraintTrack interface
 
-	virtual void AddConstraint(float Time, float ConstraintEndTime, const FName SocketName, const FName ComponentName, const FGuid& ConstraintId) override;
+	virtual void AddConstraint(float Time, float ConstraintEndTime, const FName SocketName, const FName ComponentName, const FMovieSceneObjectBindingID& ConstraintBindingID) override;
 
 public:
 
diff --git a/Engine/Source/Runtime/MovieSceneTracks/Public/Tracks/MovieSceneAudioTrack.h b/Engine/Source/Runtime/MovieSceneTracks/Public/Tracks/MovieSceneAudioTrack.h
--- a/Engine/Source/Runtime/MovieSceneTracks/Public/Tracks/MovieSceneAudioTrack.h
+++ b/Engine/Source/Runtime/MovieSceneTracks/Public/Tracks/MovieSceneAudioTrack.h
@@ -52,7 +52,7 @@ public:
 	virtual TRange<float> GetSectionBoundaries() const override;
 	virtual const TArray<UMovieSceneSection*>& GetAllSections() const override;
 	virtual bool SupportsMultipleRows() const override;
-	virtual TInlineValue<FMovieSceneSegmentCompilerRules> GetRowCompilerRules() const override;
+	virtual FMovieSceneTrackRowSegmentBlenderPtr GetRowSegmentBlender() const override;
 
 private:
 
diff --git a/Engine/Source/Runtime/MovieSceneTracks/Public/Tracks/MovieSceneBoolTrack.h b/Engine/Source/Runtime/MovieSceneTracks/Public/Tracks/MovieSceneBoolTrack.h
--- a/Engine/Source/Runtime/MovieSceneTracks/Public/Tracks/MovieSceneBoolTrack.h
+++ b/Engine/Source/Runtime/MovieSceneTracks/Public/Tracks/MovieSceneBoolTrack.h
@@ -22,7 +22,4 @@ public:
 
 	virtual UMovieSceneSection* CreateNewSection() override;
 	virtual FMovieSceneEvalTemplatePtr CreateTemplateForSection(const UMovieSceneSection& InSection) const override;
-
-	DEPRECATED(4.15, "Direct evaluation of boolean tracks is no longer supported. Please create an evaluation template (see FMovieSceneBoolPropertySectionTemplate).")
-	virtual bool Eval( float Position, float LastPostion, bool& InOutBool ) const;
 };
diff --git a/Engine/Source/Runtime/MovieSceneTracks/Public/Tracks/MovieSceneCameraAnimTrack.h b/Engine/Source/Runtime/MovieSceneTracks/Public/Tracks/MovieSceneCameraAnimTrack.h
--- a/Engine/Source/Runtime/MovieSceneTracks/Public/Tracks/MovieSceneCameraAnimTrack.h
+++ b/Engine/Source/Runtime/MovieSceneTracks/Public/Tracks/MovieSceneCameraAnimTrack.h
@@ -36,8 +36,7 @@ public:
 	virtual const TArray<UMovieSceneSection*>& GetAllSections() const override;
 	virtual UMovieSceneSection* CreateNewSection() override;
 	virtual void RemoveAllAnimationData() override;
-	virtual TInlineValue<FMovieSceneSegmentCompilerRules> GetTrackCompilerRules() const override;
-	virtual void PostCompile(FMovieSceneEvaluationTrack& OutTrack, const FMovieSceneTrackCompilerArgs& Args) const override;
+	virtual FMovieSceneTrackSegmentBlenderPtr GetTrackSegmentBlender() const override;
 	
 #if WITH_EDITORONLY_DATA
 	virtual FText GetDisplayName() const override;
diff --git a/Engine/Source/Runtime/MovieSceneTracks/Public/Tracks/MovieSceneCameraCutTrack.h b/Engine/Source/Runtime/MovieSceneTracks/Public/Tracks/MovieSceneCameraCutTrack.h
--- a/Engine/Source/Runtime/MovieSceneTracks/Public/Tracks/MovieSceneCameraCutTrack.h
+++ b/Engine/Source/Runtime/MovieSceneTracks/Public/Tracks/MovieSceneCameraCutTrack.h
@@ -6,8 +6,11 @@
 #include "UObject/ObjectMacros.h"
 #include "Misc/Guid.h"
 #include "MovieSceneNameableTrack.h"
+#include "MovieSceneObjectBindingID.h"
 #include "MovieSceneCameraCutTrack.generated.h"
 
+class UMovieSceneCameraCutSection;
+
 /**
  * Handles manipulation of CameraCut properties in a movie scene.
  */
@@ -20,13 +23,17 @@ class UMovieSceneCameraCutTrack
 	
 public:
 
+	DEPRECATED(4.18, "Camera guid no longer supported, Use CameraBindingID.")
+	MOVIESCENETRACKS_API void AddNewCameraCut(FGuid CameraHandle, float Time);
+
 	/** 
 	 * Adds a new CameraCut at the specified time.
 	 *	
-	 * @param CameraHandle Handle to the camera that the CameraCut switches to when active.
+	 * @param CameraBindingID Handle to the camera that the CameraCut switches to when active.
 	 * @param TimeRange The range within this track's movie scene where the CameraCut is initially placed.
+	 * @return The newly created camera cut section
 	 */
-	MOVIESCENETRACKS_API void AddNewCameraCut(FGuid CameraHandle, float Time);
+	MOVIESCENETRACKS_API UMovieSceneCameraCutSection* AddNewCameraCut(const FMovieSceneObjectBindingID& CameraBindingID, float Time);
 
 public:
 
diff --git a/Engine/Source/Runtime/MovieSceneTracks/Public/Tracks/MovieSceneCameraShakeTrack.h b/Engine/Source/Runtime/MovieSceneTracks/Public/Tracks/MovieSceneCameraShakeTrack.h
--- a/Engine/Source/Runtime/MovieSceneTracks/Public/Tracks/MovieSceneCameraShakeTrack.h
+++ b/Engine/Source/Runtime/MovieSceneTracks/Public/Tracks/MovieSceneCameraShakeTrack.h
@@ -36,8 +36,7 @@ public:
 	virtual const TArray<UMovieSceneSection*>& GetAllSections() const override;
 	virtual UMovieSceneSection* CreateNewSection() override;
 	virtual void RemoveAllAnimationData() override;
-	virtual TInlineValue<FMovieSceneSegmentCompilerRules> GetTrackCompilerRules() const override;
-	virtual void PostCompile(FMovieSceneEvaluationTrack& OutTrack, const FMovieSceneTrackCompilerArgs& Args) const override;
+	virtual FMovieSceneTrackSegmentBlenderPtr GetTrackSegmentBlender() const override;
 
 #if WITH_EDITORONLY_DATA
 	virtual FText GetDisplayName() const override;
diff --git a/Engine/Source/Runtime/MovieSceneTracks/Public/Tracks/MovieSceneCinematicShotTrack.h b/Engine/Source/Runtime/MovieSceneTracks/Public/Tracks/MovieSceneCinematicShotTrack.h
--- a/Engine/Source/Runtime/MovieSceneTracks/Public/Tracks/MovieSceneCinematicShotTrack.h
+++ b/Engine/Source/Runtime/MovieSceneTracks/Public/Tracks/MovieSceneCinematicShotTrack.h
@@ -35,8 +35,8 @@ public:
 	virtual UMovieSceneSection* CreateNewSection() override;
 	virtual void RemoveSection(UMovieSceneSection& Section) override;
 	virtual bool SupportsMultipleRows() const override;
-	virtual TInlineValue<FMovieSceneSegmentCompilerRules> GetRowCompilerRules() const override;
-	virtual TInlineValue<FMovieSceneSegmentCompilerRules> GetTrackCompilerRules() const override;
+	virtual FMovieSceneTrackRowSegmentBlenderPtr GetRowSegmentBlender() const override;
+	virtual FMovieSceneTrackSegmentBlenderPtr GetTrackSegmentBlender() const override;
 	
 #if WITH_EDITOR
 	virtual void OnSectionMoved(UMovieSceneSection& Section) override;
diff --git a/Engine/Source/Runtime/MovieSceneTracks/Public/Tracks/MovieSceneSkeletalAnimationTrack.h b/Engine/Source/Runtime/MovieSceneTracks/Public/Tracks/MovieSceneSkeletalAnimationTrack.h
--- a/Engine/Source/Runtime/MovieSceneTracks/Public/Tracks/MovieSceneSkeletalAnimationTrack.h
+++ b/Engine/Source/Runtime/MovieSceneTracks/Public/Tracks/MovieSceneSkeletalAnimationTrack.h
@@ -37,7 +37,7 @@ public:
 	virtual TRange<float> GetSectionBoundaries() const override;
 	virtual const TArray<UMovieSceneSection*>& GetAllSections() const override;
 	virtual UMovieSceneSection* CreateNewSection() override;
-	virtual TInlineValue<FMovieSceneSegmentCompilerRules> GetRowCompilerRules() const override;
+	virtual FMovieSceneTrackRowSegmentBlenderPtr GetRowSegmentBlender() const override;
 	virtual bool SupportsMultipleRows() const override;
 
 #if WITH_EDITORONLY_DATA
diff --git a/Engine/Source/Runtime/MovieSceneTracks/Public/Tracks/MovieSceneStringTrack.h b/Engine/Source/Runtime/MovieSceneTracks/Public/Tracks/MovieSceneStringTrack.h
--- a/Engine/Source/Runtime/MovieSceneTracks/Public/Tracks/MovieSceneStringTrack.h
+++ b/Engine/Source/Runtime/MovieSceneTracks/Public/Tracks/MovieSceneStringTrack.h
@@ -39,17 +39,6 @@ public:
 	 */
 	bool AddKeyToSection(float Time, const FString& String);
 
-	/**
-	 * Evaluates the track at the playback position
-	 *
-	 * @param Position The position at which to evaluate this track.
-	 * @param LastPositionThe last playback position.
-	 * @param InOutString The string at the evaluation time.
-	 * @return true if anything was evaluated. Note: if false is returned InOutString remains unchanged
-	 */
-	DEPRECATED(4.15, "Direct evaluation of string tracks is no longer supported. Please create an evaluation template (see FMovieSceneStringPropertySectionTemplate).")
-	virtual bool Eval(float Position, float LastPostion, FString& InOutString) const;
-
 public:
 
 	//~ UMovieSceneTrack interface
diff --git a/Engine/Source/Runtime/UMG/Private/Animation/UMGSequencePlayer.cpp b/Engine/Source/Runtime/UMG/Private/Animation/UMGSequencePlayer.cpp
--- a/Engine/Source/Runtime/UMG/Private/Animation/UMGSequencePlayer.cpp
+++ b/Engine/Source/Runtime/UMG/Private/Animation/UMGSequencePlayer.cpp
@@ -203,7 +203,9 @@ void UUMGSequencePlayer::Pause()
 	// Purposely don't trigger any OnFinished events
 	PlayerStatus = EMovieScenePlayerStatus::Stopped;
 
-	RootTemplateInstance.Finish(*this);
+	// Evaluate the sequence at its current time, with a status of 'stopped' to ensure that animated state pauses correctly. (ie. audio sounds should stop/pause)
+	const FMovieSceneContext Context(FMovieSceneEvaluationRange(TimeCursorPosition, TimeCursorPosition), PlayerStatus);
+	RootTemplateInstance.Evaluate(Context, *this);
 
 	ApplyLatentActions();
 }
