Copying //UE4/Dev-Sequencer to //UE4/Dev-Main (Source: //UE4/Dev-Sequencer @ 3946692)

diff --git a/Engine/Source/Runtime/MovieSceneCapture/Private/MovieSceneCaptureModule.cpp b/Engine/Source/Runtime/MovieSceneCapture/Private/MovieSceneCaptureModule.cpp
--- a/Engine/Source/Runtime/MovieSceneCapture/Private/MovieSceneCaptureModule.cpp
+++ b/Engine/Source/Runtime/MovieSceneCapture/Private/MovieSceneCaptureModule.cpp
@@ -23,6 +23,8 @@
 
 #define LOCTEXT_NAMESPACE "MovieSceneCapture"
 
+DEFINE_LOG_CATEGORY(LogMovieSceneCapture);
+
 class FMovieSceneCaptureModule : public IMovieSceneCaptureModule
 {
 private:
@@ -30,15 +32,22 @@ private:
 	/** Handle to a movie capture implementation created from the command line, to be initialized once a world is loaded */
 	FMovieSceneCaptureHandle StartupMovieCaptureHandle;
 	FMovieSceneCaptureProtocolRegistry ProtocolRegistry;
-
+	bool bStereoAllowed;
+	
 	virtual FMovieSceneCaptureProtocolRegistry& GetProtocolRegistry()
 	{
 		return ProtocolRegistry;
 	}
 
+	virtual bool IsStereoAllowed() override
+	{
+		return bStereoAllowed;
+	}
 	
 	virtual void StartupModule() override
 	{
+		bStereoAllowed = false;
+		
 		FCoreDelegates::OnPreExit.AddRaw(this, &FMovieSceneCaptureModule::PreExit);
 		FCoreUObjectDelegates::PostLoadMapWithWorld.AddRaw(this, &FMovieSceneCaptureModule::OnPostLoadMap );
 
@@ -124,7 +133,12 @@ private:
 		{
 			return nullptr;
 		}
-
+		
+		if (FParse::Param(FCommandLine::Get(), TEXT("EmulateStereo")))
+		{
+			bStereoAllowed = true;
+		}
+		
 		FString TypeName;
 		if( FParse::Value( FCommandLine::Get(), TEXT( "-MovieSceneCaptureType=" ), TypeName ) )
 		{
diff --git a/Engine/Source/Runtime/MovieSceneCapture/Private/VideoCaptureProtocol.cpp b/Engine/Source/Runtime/MovieSceneCapture/Private/VideoCaptureProtocol.cpp
--- a/Engine/Source/Runtime/MovieSceneCapture/Private/VideoCaptureProtocol.cpp
+++ b/Engine/Source/Runtime/MovieSceneCapture/Private/VideoCaptureProtocol.cpp
@@ -4,6 +4,7 @@
 #include "HAL/FileManager.h"
 #include "Misc/CommandLine.h"
 #include "Templates/Casts.h"
+#include "FrameRate.h"
 
 bool FVideoCaptureProtocol::Initialize(const FCaptureProtocolInitSettings& InSettings, const ICaptureProtocolHost& Host)
 {
@@ -43,7 +44,8 @@ void FVideoCaptureProtocol::ConditionallyCreateWriter(const ICaptureProtocolHost
 
 	FAVIWriterOptions Options;
 	Options.OutputFilename = MoveTemp(VideoFilename);
-	Options.CaptureFPS = Host.GetCaptureFrequency();
+	Options.CaptureFramerateNumerator = Host.GetCaptureFrameRate().Numerator;
+	Options.CaptureFramerateDenominator = Host.GetCaptureFrameRate().Denominator;
 	Options.CodecName = CaptureSettings->VideoCodec;
 	Options.bSynchronizeFrames = Host.GetCaptureStrategy().ShouldSynchronizeFrames();
 	Options.Width = InitSettings->DesiredSize.X;
diff --git a/Engine/Source/Runtime/MovieSceneCapture/Public/AutomatedLevelSequenceCapture.h b/Engine/Source/Runtime/MovieSceneCapture/Public/AutomatedLevelSequenceCapture.h
--- a/Engine/Source/Runtime/MovieSceneCapture/Public/AutomatedLevelSequenceCapture.h
+++ b/Engine/Source/Runtime/MovieSceneCapture/Public/AutomatedLevelSequenceCapture.h
@@ -32,7 +32,7 @@ public:
 
 	/** Frame number to start capturing.  The frame number range depends on whether the bUseRelativeFrameNumbers option is enabled. */
 	UPROPERTY(config, EditAnywhere, Category=Animation, AdvancedDisplay, meta=(EditCondition="bUseCustomStartFrame"))
-	int32 StartFrame;
+	FFrameNumber StartFrame;
 
 	/** When enabled, the EndFrame setting will override the default ending frame number */
 	UPROPERTY(config, EditAnywhere, Category=Animation, AdvancedDisplay)
@@ -40,7 +40,7 @@ public:
 
 	/** Frame number to end capturing.  The frame number range depends on whether the bUseRelativeFrameNumbers option is enabled. */
 	UPROPERTY(config, EditAnywhere, Category=Animation, AdvancedDisplay, meta=(EditCondition="bUseCustomEndFrame"))
-	int32 EndFrame;
+	FFrameNumber EndFrame;
 
 	/** The number of extra frames to play before the sequence's start frame, to "warm up" the animation.  This is useful if your
 	    animation contains particles or other runtime effects that are spawned into the scene earlier than your capture start frame */
@@ -51,6 +51,10 @@ public:
 	UPROPERTY(config, EditAnywhere, Category=Animation, AdvancedDisplay, meta=(Units=Seconds, ClampMin=0))
 	float DelayBeforeWarmUp;
 
+	/** The number of seconds to wait (in real-time) at shot boundaries.  Useful for allowing post processing effects to settle down before capturing the animation. */
+	UPROPERTY(config, EditAnywhere, Category=Animation, AdvancedDisplay, meta=(Units=Seconds, ClampMin=0))
+	float DelayBeforeShotWarmUp;
+
 	UPROPERTY(EditAnywhere, Category=CaptureSettings, AdvancedDisplay, meta=(EditInline))
 	ULevelSequenceBurnInOptions* BurnInOptions;
 
@@ -69,7 +73,7 @@ public:
 	virtual void Close() override;
 
 	/** Override the render frames with the given start/end frames. Restore the values when done rendering. */
-	void SetFrameOverrides(int32 InStartFrame, int32 InEndFrame);
+	void SetFrameOverrides(FFrameNumber InStartFrame, FFrameNumber InEndFrame);
 
 protected:
 
@@ -87,7 +91,7 @@ private:
 	void UpdateFrameState();
 
 	/** Called when the level sequence has updated the world */
-	void SequenceUpdated(const UMovieSceneSequencePlayer& Player, float CurrentTime, float PreviousTime);
+	void SequenceUpdated(const UMovieSceneSequencePlayer& Player, FFrameTime CurrentTime, FFrameTime PreviousTime);
 
 	/** Called to set up the player's playback range */
 	void SetupFrameRange();
@@ -109,7 +113,7 @@ private:
 	bool InitializeShots();
 
 	/** Set up the current shot to be recorded, ie. expand playback range to the section range */
-	bool SetupShot(float& StartTime, float& EndTime);
+	bool SetupShot(FFrameNumber& StartTime, FFrameNumber& EndTime);
 
 	/** Restore any modification to shots */
 	void RestoreShots();
@@ -161,25 +165,25 @@ private:
 
 	struct FCinematicShotCache
 	{
-		FCinematicShotCache(bool bInActive, bool bInLocked, const TRange<float>& InShotRange, const TRange<float>& InMovieSceneRange) : 
+		FCinematicShotCache(bool bInActive, bool bInLocked, const TRange<FFrameNumber>& InShotRange, const TRange<FFrameNumber>& InMovieSceneRange) : 
 			  bActive(bInActive)
 			, bLocked(bInLocked)
 			, ShotRange(InShotRange)
 			, MovieSceneRange(InMovieSceneRange)
 		{
-		};
+		}
 
 		bool bActive;
 		bool bLocked;
-		TRange<float> ShotRange;
-		TRange<float> MovieSceneRange;
+		TRange<FFrameNumber> ShotRange;
+		TRange<FFrameNumber> MovieSceneRange;
 	};
 
 	TArray<FCinematicShotCache> CachedShotStates;
-	TRange<float> CachedPlaybackRange;
+	TRange<FFrameNumber> CachedPlaybackRange;
 
-	TOptional<int32> CachedStartFrame;
-	TOptional<int32> CachedEndFrame;
+	TOptional<FFrameNumber> CachedStartFrame;
+	TOptional<FFrameNumber> CachedEndFrame;
 	TOptional<bool> bCachedUseCustomStartFrame;
 	TOptional<bool> bCachedUseCustomEndFrame;
 #endif
diff --git a/Engine/Source/Runtime/MovieSceneCapture/Public/IMovieSceneCaptureProtocol.h b/Engine/Source/Runtime/MovieSceneCapture/Public/IMovieSceneCaptureProtocol.h
--- a/Engine/Source/Runtime/MovieSceneCapture/Public/IMovieSceneCaptureProtocol.h
+++ b/Engine/Source/Runtime/MovieSceneCapture/Public/IMovieSceneCaptureProtocol.h
@@ -6,6 +6,7 @@
 
 class FSceneViewport;
 struct FCaptureProtocolInitSettings;
+struct FFrameRate;
 struct FFrameMetrics;
 struct ICaptureProtocolHost;
 
@@ -116,7 +117,7 @@ struct ICaptureProtocolHost
 	virtual void EnsureFileWritable(const FString& File) const = 0;
 
 	/** Get the capture frequency */
-	virtual float GetCaptureFrequency() const = 0;
+	virtual FFrameRate GetCaptureFrameRate() const = 0;
 
 	/** Access the host's capture strategy */
 	virtual const ICaptureStrategy& GetCaptureStrategy() const = 0;
diff --git a/Engine/Source/Runtime/MovieSceneCapture/Public/MovieSceneCapture.h b/Engine/Source/Runtime/MovieSceneCapture/Public/MovieSceneCapture.h
--- a/Engine/Source/Runtime/MovieSceneCapture/Public/MovieSceneCapture.h
+++ b/Engine/Source/Runtime/MovieSceneCapture/Public/MovieSceneCapture.h
@@ -141,7 +141,7 @@ protected:
 	/**~ ICaptureProtocolHost interface */
 	virtual FString GenerateFilename(const FFrameMetrics& FrameMetrics, const TCHAR* Extension) const override;
 	virtual void EnsureFileWritable(const FString& File) const override;
-	virtual float GetCaptureFrequency() const { return Settings.FrameRate; }
+	virtual FFrameRate GetCaptureFrameRate() const { return Settings.FrameRate; }
 	virtual const ICaptureStrategy& GetCaptureStrategy() const { return *CaptureStrategy; }
 
 	/** Add additional format mappings to be used when generating filenames */
@@ -174,8 +174,6 @@ protected:
 	FCachedMetrics CachedMetrics;
 	/** Format mappings used for generating filenames */
 	TMap<FString, FStringFormatArg> FormatMappings;
-	/** The number of frames to capture.  If this is zero, we'll capture the entire sequence. */
-	int32 FrameCount;
 	/** Whether we have started capturing or not */
 	bool bCapturing;
 	/** Frame number index offset when saving out frames.  This is used to allow the frame numbers on disk to match
@@ -190,7 +188,7 @@ protected:
 /** A strategy that employs a fixed frame time-step, and as such never drops a frame. Potentially accelerated. */
 struct MOVIESCENECAPTURE_API FFixedTimeStepCaptureStrategy : ICaptureStrategy
 {
-	FFixedTimeStepCaptureStrategy(uint32 InTargetFPS);
+	FFixedTimeStepCaptureStrategy(FFrameRate InFrameRate);
 
 	virtual void OnWarmup() override;
 	virtual void OnStart() override;
@@ -200,13 +198,13 @@ struct MOVIESCENECAPTURE_API FFixedTimeStepCaptureStrategy : ICaptureStrategy
 	virtual int32 GetDroppedFrames(double CurrentTimeSeconds, uint32 FrameIndex) const override;
 
 private:
-	uint32 TargetFPS;
+	FFrameRate FrameRate;
 };
 
 /** A capture strategy that captures in real-time, potentially dropping frames to maintain a stable constant framerate video. */
 struct MOVIESCENECAPTURE_API FRealTimeCaptureStrategy : ICaptureStrategy
 {
-	FRealTimeCaptureStrategy(uint32 InTargetFPS);
+	FRealTimeCaptureStrategy(FFrameRate InFrameRate);
 
 	virtual void OnWarmup() override;
 	virtual void OnStart() override;
diff --git a/Engine/Source/Runtime/MovieSceneCapture/Public/MovieSceneCaptureHelpers.h b/Engine/Source/Runtime/MovieSceneCapture/Public/MovieSceneCaptureHelpers.h
--- a/Engine/Source/Runtime/MovieSceneCapture/Public/MovieSceneCaptureHelpers.h
+++ b/Engine/Source/Runtime/MovieSceneCapture/Public/MovieSceneCaptureHelpers.h
@@ -4,6 +4,8 @@
 
 #include "CoreMinimal.h"
 
+struct FFrameRate;
+struct FFrameNumber;
 class UMovieScene;
 
 class MOVIESCENECAPTURE_API MovieSceneCaptureHelpers
@@ -18,7 +20,7 @@ public:
 	 * @param InFilename The filename to import
 	 * @return Whether the import was successful
 	 */
-	static bool ImportEDL(UMovieScene* InMovieScene, float InFrameRate, FString InFilename);
+	static bool ImportEDL(UMovieScene* InMovieScene, FFrameRate InFrameRate, FString InFilename);
 
 	/**
 	 * Export EDL
@@ -29,5 +31,5 @@ public:
 	 * @param InHandleFrames The number of handle frames to include for each shot.
 	 * @return Whether the export was successful
 	 */
-	static bool ExportEDL(const UMovieScene* InMovieScene, float InFrameRate, FString InSaveFilename, int32 InHandleFrames);
+	static bool ExportEDL(const UMovieScene* InMovieScene, FFrameRate InFrameRate, FString InSaveFilename, int32 InHandleFrames);
 };
diff --git a/Engine/Source/Runtime/MovieSceneCapture/Public/MovieSceneCaptureModule.h b/Engine/Source/Runtime/MovieSceneCapture/Public/MovieSceneCaptureModule.h
--- a/Engine/Source/Runtime/MovieSceneCapture/Public/MovieSceneCaptureModule.h
+++ b/Engine/Source/Runtime/MovieSceneCapture/Public/MovieSceneCaptureModule.h
@@ -11,6 +11,8 @@ class FMovieSceneCaptureProtocolRegistry;
 class FSceneViewport;
 class IMovieSceneCaptureInterface;
 
+MOVIESCENECAPTURE_API DECLARE_LOG_CATEGORY_EXTERN(LogMovieSceneCapture, Log, All);
+
 class IMovieSceneCaptureModule : public IModuleInterface
 {
 public:
@@ -29,5 +31,6 @@ public:
 
 	virtual void DestroyAllActiveCaptures() = 0;
 	virtual FMovieSceneCaptureProtocolRegistry& GetProtocolRegistry() = 0;
+	virtual bool IsStereoAllowed() = 0;
 };
 
diff --git a/Engine/Source/Runtime/MovieSceneCapture/Public/MovieSceneCaptureSettings.h b/Engine/Source/Runtime/MovieSceneCapture/Public/MovieSceneCaptureSettings.h
--- a/Engine/Source/Runtime/MovieSceneCapture/Public/MovieSceneCaptureSettings.h
+++ b/Engine/Source/Runtime/MovieSceneCapture/Public/MovieSceneCaptureSettings.h
@@ -6,6 +6,7 @@
 #include "UObject/ObjectMacros.h"
 #include "Templates/SubclassOf.h"
 #include "Engine/EngineTypes.h"
+#include "FrameRate.h"
 #include "MovieSceneCaptureSettings.generated.h"
 
 
@@ -49,6 +50,7 @@ struct MOVIESCENECAPTURE_API FMovieSceneCaptureSettings
 	 * {quality}	- The image compression quality setting
 	 * {material}   - The material/render pass
 	 * {shot}       - The name of the level sequence asset shot being played
+	 * {camera}     - The name of the current camera
 	 */
 	UPROPERTY(config, EditAnywhere, Category=General, DisplayName="Filename Format")
 	FString OutputFormat;
@@ -71,7 +73,7 @@ struct MOVIESCENECAPTURE_API FMovieSceneCaptureSettings
 
 	/** The frame rate at which to capture */
 	UPROPERTY(config, EditAnywhere, Category=CaptureSettings, meta=(ClampMin=1, UIMin=1, ClampMax=200, UIMax=200))
-	int32 FrameRate;
+	FFrameRate FrameRate;
 
 	/** The resolution at which to capture */
 	UPROPERTY(config, EditAnywhere, Category=CaptureSettings, meta=(ShowOnlyInnerProperties))
diff --git a/Engine/Source/Runtime/MovieSceneCapture/Public/Protocols/CompositionGraphCaptureProtocol.h b/Engine/Source/Runtime/MovieSceneCapture/Public/Protocols/CompositionGraphCaptureProtocol.h
--- a/Engine/Source/Runtime/MovieSceneCapture/Public/Protocols/CompositionGraphCaptureProtocol.h
+++ b/Engine/Source/Runtime/MovieSceneCapture/Public/Protocols/CompositionGraphCaptureProtocol.h
@@ -39,6 +39,8 @@ UCLASS(config=EditorPerProjectUserSettings, DisplayName="Composition Graph Optio
 class MOVIESCENECAPTURE_API UCompositionGraphCaptureSettings : public UMovieSceneCaptureProtocolSettings
 {
 public:
+	UCompositionGraphCaptureSettings(const FObjectInitializer& Init) : UMovieSceneCaptureProtocolSettings(Init), bDisableScreenPercentage(true) {}
+
 	GENERATED_BODY()
 	
 	/**~ UMovieSceneCaptureProtocolSettings implementation */
@@ -54,7 +56,7 @@ public:
 	bool bCaptureFramesInHDR;
 
 	/** Compression Quality for HDR Frames (0 for no compression, 1 for default compression which can be slow) */
-	UPROPERTY(config, EditAnywhere, Category="Composition Graph Options", meta = (EditCondition = "bCaptureFramesInHDR"))
+	UPROPERTY(config, EditAnywhere, Category="Composition Graph Options", meta = (EditCondition = "bCaptureFramesInHDR", UIMin=0, ClampMin=0, UIMax=1, ClampMax=1))
 	int32 HDRCompressionQuality;
 
 	/** The color gamut to use when storing HDR captured data. The gamut depends on whether the bCaptureFramesInHDR option is enabled. */
@@ -64,6 +66,10 @@ public:
 	/** Custom post processing material to use for rendering */
 	UPROPERTY(config, EditAnywhere, Category="Composition Graph Options", meta=(AllowedClasses=""))
 	FSoftObjectPath PostProcessingMaterial;
+
+	/** Whether to disable screen percentage */
+	UPROPERTY(config, EditAnywhere, Category="Composition Graph Options")
+	bool bDisableScreenPercentage;
 };
 
 struct MOVIESCENECAPTURE_API FCompositionGraphCaptureProtocol : IMovieSceneCaptureProtocol
diff --git a/Engine/Source/Runtime/MovieSceneTracks/Private/Evaluation/MovieScene3DAttachTemplate.cpp b/Engine/Source/Runtime/MovieSceneTracks/Private/Evaluation/MovieScene3DAttachTemplate.cpp
--- a/Engine/Source/Runtime/MovieSceneTracks/Private/Evaluation/MovieScene3DAttachTemplate.cpp
+++ b/Engine/Source/Runtime/MovieSceneTracks/Private/Evaluation/MovieScene3DAttachTemplate.cpp
@@ -173,9 +173,12 @@ struct F3DAttachExecutionToken
 		FMovieSceneSequenceID SequenceID = Operand.SequenceID;
 		if (AttachBindingID.GetSequenceID().IsValid())
 		{
-			// Ensure that this ID is resolvable from the root, based on the current local sequence ID
-			FMovieSceneObjectBindingID RootBindingID = AttachBindingID.ResolveLocalToRoot(SequenceID, Player.GetEvaluationTemplate().GetHierarchy());
-			SequenceID = RootBindingID.GetSequenceID();
+			if (const FMovieSceneSubSequenceData* SubData = Player.GetEvaluationTemplate().GetHierarchy().FindSubData(SequenceID))
+			{
+				// Ensure that this ID is resolvable from the root, based on the current local sequence ID
+				FMovieSceneObjectBindingID RootBindingID = AttachBindingID.ResolveLocalToRoot(SequenceID, Player.GetEvaluationTemplate().GetHierarchy());
+				SequenceID = RootBindingID.GetSequenceID();
+			}
 		}
 
 		// If the transform is set, otherwise use the bound actor's transform
diff --git a/Engine/Source/Runtime/MovieSceneTracks/Private/Evaluation/MovieScene3DPathTemplate.cpp b/Engine/Source/Runtime/MovieSceneTracks/Private/Evaluation/MovieScene3DPathTemplate.cpp
--- a/Engine/Source/Runtime/MovieSceneTracks/Private/Evaluation/MovieScene3DPathTemplate.cpp
+++ b/Engine/Source/Runtime/MovieSceneTracks/Private/Evaluation/MovieScene3DPathTemplate.cpp
@@ -128,9 +128,12 @@ struct F3DPathExecutionToken
 		FMovieSceneSequenceID SequenceID = Operand.SequenceID;
 		if (PathBindingID.GetSequenceID().IsValid())
 		{
-			// Ensure that this ID is resolvable from the root, based on the current local sequence ID
-			FMovieSceneObjectBindingID RootBindingID = PathBindingID.ResolveLocalToRoot(SequenceID, Player.GetEvaluationTemplate().GetHierarchy());
-			SequenceID = RootBindingID.GetSequenceID();
+			if (const FMovieSceneSubSequenceData* SubData = Player.GetEvaluationTemplate().GetHierarchy().FindSubData(SequenceID))
+			{
+				// Ensure that this ID is resolvable from the root, based on the current local sequence ID
+				FMovieSceneObjectBindingID RootBindingID = PathBindingID.ResolveLocalToRoot(SequenceID, Player.GetEvaluationTemplate().GetHierarchy());
+				SequenceID = RootBindingID.GetSequenceID();
+			}
 		}
 
 		// If the transform is set, otherwise use the bound actor's transform
@@ -187,7 +190,7 @@ struct F3DPathExecutionToken
 
 FMovieScene3DPathSectionTemplate::FMovieScene3DPathSectionTemplate(const UMovieScene3DPathSection& Section)
 	: PathBindingID(Section.GetConstraintBindingID())
-	, TimingCurve(Section.GetTimingCurve())
+	, TimingCurve(Section.GetTimingChannel())
 	, FrontAxisEnum(Section.GetFrontAxisEnum())
 	, UpAxisEnum(Section.GetUpAxisEnum())
 	, bFollow(Section.GetFollow())
@@ -199,8 +202,10 @@ FMovieScene3DPathSectionTemplate::FMovieScene3DPathSectionTemplate(const UMovieS
 void FMovieScene3DPathSectionTemplate::Evaluate(const FMovieSceneEvaluationOperand& Operand, const FMovieSceneContext& Context, const FPersistentEvaluationData& PersistentData, FMovieSceneExecutionTokens& ExecutionTokens) const
 {
 	MOVIESCENE_DETAILED_SCOPE_CYCLE_COUNTER(MovieSceneEval_PathTrack_Evaluate)
-		
-	float Timing = TimingCurve.Eval(Context.GetTime());
 
-	ExecutionTokens.Add(F3DPathExecutionToken(PathBindingID, Timing, FrontAxisEnum, UpAxisEnum, bFollow, bReverse, bForceUpright));
+	float Timing = 0.f;
+	if (TimingCurve.Evaluate(Context.GetTime(), Timing))
+	{
+		ExecutionTokens.Add(F3DPathExecutionToken(PathBindingID, Timing, FrontAxisEnum, UpAxisEnum, bFollow, bReverse, bForceUpright));
+	}
 }
diff --git a/Engine/Source/Runtime/MovieSceneTracks/Private/Evaluation/MovieScene3DPathTemplate.h b/Engine/Source/Runtime/MovieSceneTracks/Private/Evaluation/MovieScene3DPathTemplate.h
--- a/Engine/Source/Runtime/MovieSceneTracks/Private/Evaluation/MovieScene3DPathTemplate.h
+++ b/Engine/Source/Runtime/MovieSceneTracks/Private/Evaluation/MovieScene3DPathTemplate.h
@@ -5,7 +5,7 @@
 #include "CoreMinimal.h"
 #include "UObject/ObjectMacros.h"
 #include "Misc/Guid.h"
-#include "Curves/RichCurve.h"
+#include "Channels/MovieSceneFloatChannel.h"
 #include "Evaluation/MovieSceneEvalTemplate.h"
 #include "Sections/MovieScene3DPathSection.h"
 #include "MovieSceneObjectBindingID.h"
@@ -25,7 +25,7 @@ struct FMovieScene3DPathSectionTemplate : public FMovieSceneEvalTemplate
 
 	/** The timing curve */
 	UPROPERTY()
-	FRichCurve TimingCurve;
+	FMovieSceneFloatChannel TimingCurve;
 
 	/** Front Axis */
 	UPROPERTY()
diff --git a/Engine/Source/Runtime/MovieSceneTracks/Private/Evaluation/MovieScene3DTransformTemplate.cpp b/Engine/Source/Runtime/MovieSceneTracks/Private/Evaluation/MovieScene3DTransformTemplate.cpp
--- a/Engine/Source/Runtime/MovieSceneTracks/Private/Evaluation/MovieScene3DTransformTemplate.cpp
+++ b/Engine/Source/Runtime/MovieSceneTracks/Private/Evaluation/MovieScene3DTransformTemplate.cpp
@@ -7,6 +7,7 @@
 #include "MovieSceneCommonHelpers.h"
 #include "Evaluation/MovieSceneEvaluation.h"
 #include "IMovieScenePlayer.h"
+#include "Channels/MovieSceneChannelProxy.h"
 #include "Evaluation/Blending/BlendableTokenStack.h"
 #include "Evaluation/Blending/MovieSceneBlendingActuatorID.h"
 #include "Tracks/IMovieSceneTransformOrigin.h"
@@ -84,7 +85,7 @@ struct FComponentTransformActuator : TMovieSceneBlendingActuator<F3DTransformTra
 
 			SceneComponent->SetMobility(EComponentMobility::Movable);
 
-			InFinalValue.Apply(*SceneComponent, Context.GetDelta());
+			InFinalValue.Apply(*SceneComponent, Context.GetDelta() / Context.GetFrameRate());
 		}
 	}
 
@@ -171,9 +172,11 @@ void FMovieSceneComponentTransformSectionTemplate::Evaluate(const FMovieSceneEva
 
 	// Add the blendable to the accumulator
 	float Weight = EvaluateEasing(Context.GetTime());
-	if (EnumHasAllFlags(TemplateData.Mask.GetChannels(), EMovieSceneTransformChannel::Weight) && TemplateData.ManualWeight.HasAnyData())
+	if (EnumHasAllFlags(TemplateData.Mask.GetChannels(), EMovieSceneTransformChannel::Weight))
 	{
-		Weight *= TemplateData.ManualWeight.Eval(Context.GetTime());
+		float ManualWeight = 1.f;
+		TemplateData.ManualWeight.Evaluate(Context.GetTime(), ManualWeight);
+		Weight *= ManualWeight;
 	}
 
 	ExecutionTokens.BlendToken(ActuatorTypeID, TBlendableToken<F3DTransformTrackToken>(TransformValue, TemplateData.BlendType, Weight));
@@ -198,9 +201,11 @@ void FMovieSceneComponentTransformSectionTemplate::Interrogate(const FMovieScene
 
 	// Add the blendable to the accumulator
 	float Weight = EvaluateEasing(Context.GetTime());
-	if (EnumHasAllFlags(TemplateData.Mask.GetChannels(), EMovieSceneTransformChannel::Weight) && TemplateData.ManualWeight.HasAnyData())
+	if (EnumHasAllFlags(TemplateData.Mask.GetChannels(), EMovieSceneTransformChannel::Weight))
 	{
-		Weight *= TemplateData.ManualWeight.Eval(Context.GetTime());
+		float ManualWeight = 1.f;
+		TemplateData.ManualWeight.Evaluate(Context.GetTime(), ManualWeight);
+		Weight *= ManualWeight;
 	}
 
 	Container.GetAccumulator().BlendToken(FMovieSceneEvaluationOperand(), ActuatorTypeID, FMovieSceneEvaluationScope(), Context, TBlendableToken<F3DTransformTrackToken>(TransformValue, TemplateData.BlendType, Weight));
@@ -211,33 +216,35 @@ FMovieScene3DTransformTemplateData::FMovieScene3DTransformTemplateData(const UMo
 	, Mask(Section.GetMask())
 {
 	EMovieSceneTransformChannel MaskChannels = Mask.GetChannels();
+	TArrayView<FMovieSceneFloatChannel*> FloatChannels = Section.GetChannelProxy().GetChannels<FMovieSceneFloatChannel>();
 
-	if (EnumHasAllFlags(MaskChannels, EMovieSceneTransformChannel::TranslationX))	TranslationCurve[0]	= Section.GetTranslationCurve(EAxis::X);
-	if (EnumHasAllFlags(MaskChannels, EMovieSceneTransformChannel::TranslationY))	TranslationCurve[1]	= Section.GetTranslationCurve(EAxis::Y);
-	if (EnumHasAllFlags(MaskChannels, EMovieSceneTransformChannel::TranslationZ))	TranslationCurve[2]	= Section.GetTranslationCurve(EAxis::Z);
+	if (EnumHasAllFlags(MaskChannels, EMovieSceneTransformChannel::TranslationX))	TranslationCurve[0]	= *FloatChannels[0];
+	if (EnumHasAllFlags(MaskChannels, EMovieSceneTransformChannel::TranslationY))	TranslationCurve[1]	= *FloatChannels[1];
+	if (EnumHasAllFlags(MaskChannels, EMovieSceneTransformChannel::TranslationZ))	TranslationCurve[2]	= *FloatChannels[2];
 
-	if (EnumHasAllFlags(MaskChannels, EMovieSceneTransformChannel::RotationX))		RotationCurve[0]	= Section.GetRotationCurve(EAxis::X);
-	if (EnumHasAllFlags(MaskChannels, EMovieSceneTransformChannel::RotationY))		RotationCurve[1]	= Section.GetRotationCurve(EAxis::Y);
-	if (EnumHasAllFlags(MaskChannels, EMovieSceneTransformChannel::RotationZ))		RotationCurve[2]	= Section.GetRotationCurve(EAxis::Z);
+	if (EnumHasAllFlags(MaskChannels, EMovieSceneTransformChannel::RotationX))		RotationCurve[0]	= *FloatChannels[3];
+	if (EnumHasAllFlags(MaskChannels, EMovieSceneTransformChannel::RotationY))		RotationCurve[1]	= *FloatChannels[4];
+	if (EnumHasAllFlags(MaskChannels, EMovieSceneTransformChannel::RotationZ))		RotationCurve[2]	= *FloatChannels[5];
 
-	if (EnumHasAllFlags(MaskChannels, EMovieSceneTransformChannel::ScaleX))			ScaleCurve[0]		= Section.GetScaleCurve(EAxis::X);
-	if (EnumHasAllFlags(MaskChannels, EMovieSceneTransformChannel::ScaleY))			ScaleCurve[1]		= Section.GetScaleCurve(EAxis::Y);
-	if (EnumHasAllFlags(MaskChannels, EMovieSceneTransformChannel::ScaleZ))			ScaleCurve[2]		= Section.GetScaleCurve(EAxis::Z);
+	if (EnumHasAllFlags(MaskChannels, EMovieSceneTransformChannel::ScaleX))			ScaleCurve[0]		= *FloatChannels[6];
+	if (EnumHasAllFlags(MaskChannels, EMovieSceneTransformChannel::ScaleY))			ScaleCurve[1]		= *FloatChannels[7];
+	if (EnumHasAllFlags(MaskChannels, EMovieSceneTransformChannel::ScaleZ))			ScaleCurve[2]		= *FloatChannels[8];
 
-	if (EnumHasAllFlags(MaskChannels, EMovieSceneTransformChannel::Weight))			ManualWeight		= Section.GetManualWeightCurve();
+	if (EnumHasAllFlags(MaskChannels, EMovieSceneTransformChannel::Weight))			ManualWeight		= *FloatChannels[9];
 }
 
-MovieScene::TMultiChannelValue<float, 9> FMovieScene3DTransformTemplateData::Evaluate(float Time) const
+MovieScene::TMultiChannelValue<float, 9> FMovieScene3DTransformTemplateData::Evaluate(FFrameTime Time) const
 {
 	MovieScene::TMultiChannelValue<float, 9> AnimatedData;
 
 	EMovieSceneTransformChannel ChannelMask = Mask.GetChannels();
 
-	auto EvalChannel = [&AnimatedData, Time, ChannelMask](uint8 ChanneIndex, EMovieSceneTransformChannel Channel, const FRichCurve& Curve)
+	auto EvalChannel = [&AnimatedData, Time, ChannelMask](uint8 ChanneIndex, EMovieSceneTransformChannel ChannelType, const FMovieSceneFloatChannel& Channel)
 	{
-		if (EnumHasAllFlags(ChannelMask, Channel) && Curve.HasAnyData())
+		float Value = 0.f;
+		if (EnumHasAllFlags(ChannelMask, ChannelType) && Channel.Evaluate(Time, Value))
 		{
-			AnimatedData.Set(ChanneIndex, Curve.Eval(Time));
+			AnimatedData.Set(ChanneIndex, Value);
 		}
 	};
 
diff --git a/Engine/Source/Runtime/MovieSceneTracks/Private/Evaluation/MovieSceneActorReferenceTemplate.cpp b/Engine/Source/Runtime/MovieSceneTracks/Private/Evaluation/MovieSceneActorReferenceTemplate.cpp
--- a/Engine/Source/Runtime/MovieSceneTracks/Private/Evaluation/MovieSceneActorReferenceTemplate.cpp
+++ b/Engine/Source/Runtime/MovieSceneTracks/Private/Evaluation/MovieSceneActorReferenceTemplate.cpp
@@ -11,12 +11,11 @@
 namespace PropertyTemplate
 {
 	template<>
-	AActor* ConvertFromIntermediateType<AActor*, FGuid>(const FGuid& InObjectGuid, const FMovieSceneEvaluationOperand& Operand, FPersistentEvaluationData& PersistentData, IMovieScenePlayer& Player)
+	AActor* ConvertFromIntermediateType<AActor*, FMovieSceneObjectBindingID>(const FMovieSceneObjectBindingID& InObjectBinding, const FMovieSceneEvaluationOperand& Operand, FPersistentEvaluationData& PersistentData, IMovieScenePlayer& Player)
 	{
-		FMovieSceneEvaluationOperand NewOperand = Operand;
-		NewOperand.ObjectBindingID = InObjectGuid;
+		FMovieSceneObjectBindingID ResolvedID = InObjectBinding.ResolveLocalToRoot(Operand.SequenceID, Player.GetEvaluationTemplate().GetHierarchy());
 
-		for (TWeakObjectPtr<>& WeakObject : Player.FindBoundObjects(NewOperand))
+		for (TWeakObjectPtr<>& WeakObject : Player.FindBoundObjects(ResolvedID.GetGuid(), ResolvedID.GetSequenceID()))
 		{
 			if (AActor* Actor = Cast<AActor>(WeakObject.Get()))
 			{
@@ -43,7 +42,7 @@ namespace PropertyTemplate
 		return InValue != nullptr;
 	}
 
-	template<> IMovieScenePreAnimatedTokenPtr CacheExistingState<AActor*, FGuid>(UObject& Object, FTrackInstancePropertyBindings& PropertyBindings)
+	template<> IMovieScenePreAnimatedTokenPtr CacheExistingState<AActor*, FMovieSceneObjectBindingID>(UObject& Object, FTrackInstancePropertyBindings& PropertyBindings)
 	{
 		return TCachedState<AActor*, TWeakObjectPtr<AActor>>(PropertyBindings.GetCurrentValue<AActor*>(Object), PropertyBindings);
 	}
@@ -51,18 +50,14 @@ namespace PropertyTemplate
 
 FMovieSceneActorReferenceSectionTemplate::FMovieSceneActorReferenceSectionTemplate(const UMovieSceneActorReferenceSection& Section, const UMovieScenePropertyTrack& Track)
 	: PropertyData(Track.GetPropertyName(), Track.GetPropertyPath())
-	, ActorGuidIndexCurve(Section.GetActorReferenceCurve())
-	, ActorGuids(Section.GetActorGuids())
+	, ActorReferenceData(Section.GetActorReferenceData())
 {
 }
 
 void FMovieSceneActorReferenceSectionTemplate::Evaluate(const FMovieSceneEvaluationOperand& Operand, const FMovieSceneContext& Context, const FPersistentEvaluationData& PersistentData, FMovieSceneExecutionTokens& ExecutionTokens) const
 {
 	using namespace PropertyTemplate;
-	FSectionData& TrackData = PersistentData.GetSectionData<FSectionData>();
 
-	const int32 ActorGuidIndex = ActorGuidIndexCurve.Evaluate(Context.GetTime());
-	FGuid ObjectBindingID = ActorGuids.IsValidIndex(ActorGuidIndex) ? ActorGuids[ActorGuidIndex] : FGuid();
-
-	ExecutionTokens.Add(TPropertyTrackExecutionToken<AActor*, FGuid>(ObjectBindingID));
+	FMovieSceneActorReferenceKey ObjectBinding = ActorReferenceData.Evaluate(Context.GetTime());
+	ExecutionTokens.Add(TPropertyTrackExecutionToken<AActor*, FMovieSceneObjectBindingID>(ObjectBinding.Object));
 }
diff --git a/Engine/Source/Runtime/MovieSceneTracks/Private/Evaluation/MovieSceneActorReferenceTemplate.h b/Engine/Source/Runtime/MovieSceneTracks/Private/Evaluation/MovieSceneActorReferenceTemplate.h
--- a/Engine/Source/Runtime/MovieSceneTracks/Private/Evaluation/MovieSceneActorReferenceTemplate.h
+++ b/Engine/Source/Runtime/MovieSceneTracks/Private/Evaluation/MovieSceneActorReferenceTemplate.h
@@ -4,13 +4,11 @@
 
 #include "CoreMinimal.h"
 #include "UObject/ObjectMacros.h"
-#include "Misc/Guid.h"
 #include "Evaluation/MovieSceneEvalTemplate.h"
 #include "Evaluation/MovieScenePropertyTemplate.h"
-#include "Curves/IntegralCurve.h"
+#include "Sections/MovieSceneActorReferenceSection.h"
 #include "MovieSceneActorReferenceTemplate.generated.h"
 
-class UMovieSceneActorReferenceSection;
 class UMovieScenePropertyTrack;
 
 USTRUCT()
@@ -42,8 +40,5 @@ private:
 	FMovieScenePropertySectionData PropertyData;
 
 	UPROPERTY()
-	FIntegralCurve ActorGuidIndexCurve;
-
-	UPROPERTY()
-	TArray<FGuid> ActorGuids;
+	FMovieSceneActorReferenceData ActorReferenceData;
 };
diff --git a/Engine/Source/Runtime/MovieSceneTracks/Private/Evaluation/MovieSceneAudioTemplate.cpp b/Engine/Source/Runtime/MovieSceneTracks/Private/Evaluation/MovieSceneAudioTemplate.cpp
--- a/Engine/Source/Runtime/MovieSceneTracks/Private/Evaluation/MovieSceneAudioTemplate.cpp
+++ b/Engine/Source/Runtime/MovieSceneTracks/Private/Evaluation/MovieSceneAudioTemplate.cpp
@@ -273,9 +273,8 @@ struct FAudioSectionExecutionToken : IMovieSceneExecutionToken
 FMovieSceneAudioSectionTemplateData::FMovieSceneAudioSectionTemplateData(const UMovieSceneAudioSection& Section)
 	: Sound(Section.GetSound())
 	, AudioStartOffset(Section.GetStartOffset())
-	, AudioRange(Section.GetAudioRange())
-	, AudioPitchMultiplierCurve(Section.GetPitchMultiplierCurve())
-	, AudioVolumeCurve(Section.GetSoundVolumeCurve())
+	, AudioPitchMultiplierCurve(Section.GetPitchMultiplierChannel())
+	, AudioVolumeCurve(Section.GetSoundVolumeChannel())
 	, RowIndex(Section.GetRowIndex())
 	, bOverrideAttenuation(Section.GetOverrideAttenuation())
 	, AttenuationSettings(Section.GetAttenuationSettings())
@@ -283,6 +282,7 @@ FMovieSceneAudioSectionTemplateData::FMovieSceneAudioSectionTemplateData(const U
 	, OnAudioFinished(Section.GetOnAudioFinished())
 	, OnAudioPlaybackPercent(Section.GetOnAudioPlaybackPercent())
 {
+	SectionStartTimeSeconds = (Section.HasStartFrame() ? Section.GetInclusiveStartFrame() : 0) / Section.GetTypedOuter<UMovieScene>()->GetFrameResolution();
 }
 
 void FMovieSceneAudioSectionTemplateData::EnsureAudioIsPlaying(UAudioComponent& AudioComponent, FPersistentEvaluationData& PersistentData, const FMovieSceneContext& Context, bool bAllowSpatialization, IMovieScenePlayer& Player) const
@@ -290,14 +290,16 @@ void FMovieSceneAudioSectionTemplateData::EnsureAudioIsPlaying(UAudioComponent&
 	Player.SavePreAnimatedState(AudioComponent, FStopAudioPreAnimatedToken::GetAnimTypeID(), FStopAudioPreAnimatedToken::FProducer());
 
 	bool bPlaySound = !AudioComponent.IsPlaying() || AudioComponent.Sound != Sound;
-	float AudioVolume = AudioVolumeCurve.Eval(Context.GetTime());
-	float PitchMultiplier = AudioPitchMultiplierCurve.Eval(Context.GetTime());
 
+	float AudioVolume = 1.f;
+	AudioVolumeCurve.Evaluate(Context.GetTime(), AudioVolume);
 	if (AudioComponent.VolumeMultiplier != AudioVolume)
 	{
 		AudioComponent.SetVolumeMultiplier(AudioVolume);
 	}
-	
+
+	float PitchMultiplier = 1.f;
+	AudioPitchMultiplierCurve.Evaluate(Context.GetTime(), PitchMultiplier);
 	if (AudioComponent.PitchMultiplier != PitchMultiplier)
 	{
 		AudioComponent.SetPitchMultiplier(PitchMultiplier);
@@ -328,7 +330,8 @@ void FMovieSceneAudioSectionTemplateData::EnsureAudioIsPlaying(UAudioComponent&
 			AudioComponent.bIsUISound = false;
 		}
 
-		AudioComponent.Play((Context.GetTime() - AudioRange.GetLowerBoundValue()) + FMath::Max(AudioStartOffset, 0.f));
+		const float AudioTime = (Context.GetTime() / Context.GetFrameRate()) - SectionStartTimeSeconds + FMath::Max(AudioStartOffset, 0.f);
+		AudioComponent.Play(AudioTime);
 
 		if (Context.GetStatus() == EMovieScenePlayerStatus::Scrubbing)
 		{
diff --git a/Engine/Source/Runtime/MovieSceneTracks/Private/Evaluation/MovieSceneAudioTemplate.h b/Engine/Source/Runtime/MovieSceneTracks/Private/Evaluation/MovieSceneAudioTemplate.h
--- a/Engine/Source/Runtime/MovieSceneTracks/Private/Evaluation/MovieSceneAudioTemplate.h
+++ b/Engine/Source/Runtime/MovieSceneTracks/Private/Evaluation/MovieSceneAudioTemplate.h
@@ -7,6 +7,7 @@
 #include "Evaluation/MovieSceneEvalTemplate.h"
 #include "Runtime/Engine/Classes/Components/AudioComponent.h"
 #include "Sound/SoundAttenuation.h"
+#include "Channels/MovieSceneFloatChannel.h"
 #include "MovieSceneAudioTemplate.generated.h"
 
 class UAudioComponent;
@@ -31,18 +32,18 @@ struct FMovieSceneAudioSectionTemplateData
 	/** The offset into the beginning of the audio clip */
 	UPROPERTY()
 	float AudioStartOffset;
-	
-	/** The time range this audio plays */
+
+	/** The frame number at which the audio starts playing */
 	UPROPERTY()
-	FFloatRange AudioRange;
+	double SectionStartTimeSeconds;
 
 	/** The amount which this audio is time dilated by */
 	UPROPERTY()
-	FRichCurve AudioPitchMultiplierCurve;
+	FMovieSceneFloatChannel AudioPitchMultiplierCurve;
 
 	/** The volume the sound will be played with. */
 	UPROPERTY()
-	FRichCurve AudioVolumeCurve;
+	FMovieSceneFloatChannel AudioVolumeCurve;
 
 	/** The row index of the section */
 	UPROPERTY()
diff --git a/Engine/Source/Runtime/MovieSceneTracks/Private/Evaluation/MovieSceneCameraAnimTemplate.cpp b/Engine/Source/Runtime/MovieSceneTracks/Private/Evaluation/MovieSceneCameraAnimTemplate.cpp
--- a/Engine/Source/Runtime/MovieSceneTracks/Private/Evaluation/MovieSceneCameraAnimTemplate.cpp
+++ b/Engine/Source/Runtime/MovieSceneTracks/Private/Evaluation/MovieSceneCameraAnimTemplate.cpp
@@ -349,7 +349,7 @@ struct FPreAnimatedCameraAnimTokenProducer : IMovieScenePreAnimatedTokenProducer
 
 FMovieSceneCameraAnimSectionTemplate::FMovieSceneCameraAnimSectionTemplate(const UMovieSceneCameraAnimSection& Section)
 	: SourceData(Section.AnimData)
-	, SectionStartTime(Section.GetStartTime())
+	, SectionStartTime(Section.HasStartFrame() ? Section.GetInclusiveStartFrame() : 0)
 {
 }
 
@@ -416,8 +416,8 @@ bool FMovieSceneCameraAnimSectionTemplate::UpdateCamera(ACameraActor& TempCamera
 	}
 	
 	// set camera anim to the correct time
-	const float NewCameraAnimTime = Context.GetTime() - SectionStartTime;
-	CameraAnimInstance->SetCurrentTime(NewCameraAnimTime);
+	const FFrameTime NewCameraAnimTime = Context.GetTime() - SectionStartTime;
+	CameraAnimInstance->SetCurrentTime(NewCameraAnimTime / Context.GetFrameRate());
 
 	if (CameraAnimInstance->CurrentBlendWeight <= 0.f)
 	{
@@ -457,7 +457,7 @@ struct FPreAnimatedCameraShakeTokenProducer : IMovieScenePreAnimatedTokenProduce
 
 FMovieSceneCameraShakeSectionTemplate::FMovieSceneCameraShakeSectionTemplate(const UMovieSceneCameraShakeSection& Section)
 	: SourceData(Section.ShakeData)
-	, SectionStartTime(Section.GetStartTime())
+	, SectionStartTime(Section.HasStartFrame() ? Section.GetInclusiveStartFrame() : 0)
 {
 }
 
@@ -529,8 +529,8 @@ bool FMovieSceneCameraShakeSectionTemplate::UpdateCamera(ACameraActor& TempCamer
 	}
 	
 	// set camera anim to the correct time
-	const float NewCameraAnimTime = Context.GetTime() - SectionStartTime;
-	CameraShakeInstance->SetCurrentTimeAndApplyShake(NewCameraAnimTime, POV);
+	const FFrameTime NewCameraAnimTime = Context.GetTime() - SectionStartTime;
+	CameraShakeInstance->SetCurrentTimeAndApplyShake(NewCameraAnimTime / Context.GetFrameRate(), POV);
 
 	return true;
 }
diff --git a/Engine/Source/Runtime/MovieSceneTracks/Private/Evaluation/MovieSceneCameraAnimTemplate.h b/Engine/Source/Runtime/MovieSceneTracks/Private/Evaluation/MovieSceneCameraAnimTemplate.h
--- a/Engine/Source/Runtime/MovieSceneTracks/Private/Evaluation/MovieSceneCameraAnimTemplate.h
+++ b/Engine/Source/Runtime/MovieSceneTracks/Private/Evaluation/MovieSceneCameraAnimTemplate.h
@@ -50,7 +50,7 @@ private:
 
 	/** Cached section start time */
 	UPROPERTY()
-	float SectionStartTime;
+	FFrameNumber SectionStartTime;
 };
 
 /** Section template for a camera anim */
@@ -75,5 +75,5 @@ private:
 
 	/** Cached section start time */
 	UPROPERTY()
-	float SectionStartTime;
+	FFrameNumber SectionStartTime;
 };
diff --git a/Engine/Source/Runtime/MovieSceneTracks/Private/Evaluation/MovieSceneCameraCutTemplate.cpp b/Engine/Source/Runtime/MovieSceneTracks/Private/Evaluation/MovieSceneCameraCutTemplate.cpp
--- a/Engine/Source/Runtime/MovieSceneTracks/Private/Evaluation/MovieSceneCameraCutTemplate.cpp
+++ b/Engine/Source/Runtime/MovieSceneTracks/Private/Evaluation/MovieSceneCameraCutTemplate.cpp
@@ -41,9 +41,12 @@ struct FCameraCutPreRollExecutionToken : IMovieSceneExecutionToken
 			FMovieSceneSequenceID SequenceID = Operand.SequenceID;
 			if (CameraBindingID.GetSequenceID().IsValid())
 			{
-				// Ensure that this ID is resolvable from the root, based on the current local sequence ID
-				FMovieSceneObjectBindingID RootBindingID = CameraBindingID.ResolveLocalToRoot(SequenceID, Player.GetEvaluationTemplate().GetHierarchy());
-				SequenceID = RootBindingID.GetSequenceID();
+				if (const FMovieSceneSubSequenceData* SubData = Player.GetEvaluationTemplate().GetHierarchy().FindSubData(SequenceID))
+				{
+					// Ensure that this ID is resolvable from the root, based on the current local sequence ID
+					FMovieSceneObjectBindingID RootBindingID = CameraBindingID.ResolveLocalToRoot(SequenceID, Player.GetEvaluationTemplate().GetHierarchy());
+					SequenceID = RootBindingID.GetSequenceID();
+				}
 			}
 
 			// If the transform is set, otherwise use the bound actor's transform
@@ -111,9 +114,12 @@ struct FCameraCutExecutionToken : IMovieSceneExecutionToken
 		FMovieSceneSequenceID SequenceID = Operand.SequenceID;
 		if (CameraBindingID.GetSequenceID().IsValid())
 		{
-			// Ensure that this ID is resolvable from the root, based on the current local sequence ID
-			FMovieSceneObjectBindingID RootBindingID = CameraBindingID.ResolveLocalToRoot(SequenceID, Player.GetEvaluationTemplate().GetHierarchy());
-			SequenceID = RootBindingID.GetSequenceID();
+			if (const FMovieSceneSubSequenceData* SubData = Player.GetEvaluationTemplate().GetHierarchy().FindSubData(SequenceID))
+			{
+				// Ensure that this ID is resolvable from the root, based on the current local sequence ID
+				FMovieSceneObjectBindingID RootBindingID = CameraBindingID.ResolveLocalToRoot(SequenceID, Player.GetEvaluationTemplate().GetHierarchy());
+				SequenceID = RootBindingID.GetSequenceID();
+			}
 		}
 
 		FMovieSceneEvaluationOperand CameraOperand(SequenceID, CameraBindingID.GetGuid());
diff --git a/Engine/Source/Runtime/MovieSceneTracks/Private/Evaluation/MovieSceneColorTemplate.cpp b/Engine/Source/Runtime/MovieSceneTracks/Private/Evaluation/MovieSceneColorTemplate.cpp
--- a/Engine/Source/Runtime/MovieSceneTracks/Private/Evaluation/MovieSceneColorTemplate.cpp
+++ b/Engine/Source/Runtime/MovieSceneTracks/Private/Evaluation/MovieSceneColorTemplate.cpp
@@ -234,23 +234,23 @@ FMovieSceneColorSectionTemplate::FMovieSceneColorSectionTemplate(const UMovieSce
 	: FMovieScenePropertySectionTemplate(Track.GetPropertyName(), Track.GetPropertyPath())
 	, BlendType(Section.GetBlendType().Get())
 {
-	Curves[0] = Section.GetRedCurve();
-	Curves[1] = Section.GetGreenCurve();
-	Curves[2] = Section.GetBlueCurve();
-	Curves[3] = Section.GetAlphaCurve();
+	Curves[0] = Section.GetRedChannel();
+	Curves[1] = Section.GetGreenChannel();
+	Curves[2] = Section.GetBlueChannel();
+	Curves[3] = Section.GetAlphaChannel();
 }
 
 void FMovieSceneColorSectionTemplate::Evaluate(const FMovieSceneEvaluationOperand& Operand, const FMovieSceneContext& Context, const FPersistentEvaluationData& PersistentData, FMovieSceneExecutionTokens& ExecutionTokens) const
 {
-	const float Time = Context.GetTime();
+	const FFrameTime Time = Context.GetTime();
 	MovieScene::TMultiChannelValue<float, 4> AnimationData;
 
 	for (uint8 Index = 0; Index < 4; ++Index)
 	{
-		const FRichCurve& Curve = Curves[Index];
-		if (Curve.HasAnyData())
+		float ChannelValue = 0.f;
+		if (Curves[Index].Evaluate(Time, ChannelValue))
 		{
-			AnimationData.Set(Index, Curve.Eval(Time));
+			AnimationData.Set(Index, ChannelValue);
 		}
 	}
 
diff --git a/Engine/Source/Runtime/MovieSceneTracks/Private/Evaluation/MovieSceneColorTemplate.h b/Engine/Source/Runtime/MovieSceneTracks/Private/Evaluation/MovieSceneColorTemplate.h
--- a/Engine/Source/Runtime/MovieSceneTracks/Private/Evaluation/MovieSceneColorTemplate.h
+++ b/Engine/Source/Runtime/MovieSceneTracks/Private/Evaluation/MovieSceneColorTemplate.h
@@ -4,7 +4,7 @@
 
 #include "CoreMinimal.h"
 #include "UObject/ObjectMacros.h"
-#include "Curves/RichCurve.h"
+#include "Channels/MovieSceneFloatChannel.h"
 #include "Evaluation/MovieSceneEvalTemplate.h"
 #include "Evaluation/MovieScenePropertyTemplate.h"
 
@@ -23,7 +23,7 @@ struct FMovieSceneColorSectionTemplate : public FMovieScenePropertySectionTempla
 
 	/** Curve data as RGBA */
 	UPROPERTY()
-	FRichCurve Curves[4];
+	FMovieSceneFloatChannel Curves[4];
 
 	UPROPERTY()
 	EMovieSceneBlendType BlendType;
diff --git a/Engine/Source/Runtime/MovieSceneTracks/Private/Evaluation/MovieSceneEventTemplate.cpp b/Engine/Source/Runtime/MovieSceneTracks/Private/Evaluation/MovieSceneEventTemplate.cpp
--- a/Engine/Source/Runtime/MovieSceneTracks/Private/Evaluation/MovieSceneEventTemplate.cpp
+++ b/Engine/Source/Runtime/MovieSceneTracks/Private/Evaluation/MovieSceneEventTemplate.cpp
@@ -187,36 +187,36 @@ void FMovieSceneEventSectionTemplate::EvaluateSwept(const FMovieSceneEvaluationO
 
 	TArray<FMovieSceneEventData> Events;
 
-	TRange<float> SweptRange = Context.GetRange();
+	TRange<FFrameNumber>            SweptRange = Context.GetFrameNumberRange();
 
-	const TArray<float>& KeyTimes = EventData.KeyTimes;
-	const TArray<FEventPayload>& KeyValues = EventData.KeyValues;
+	TArrayView<const FFrameNumber>  KeyTimes   = EventData.GetKeyTimes();
+	TArrayView<const FEventPayload> KeyValues  = EventData.GetKeyValues();
 
 	const int32 First = bBackwards ? KeyTimes.Num() - 1 : 0;
 	const int32 Last = bBackwards ? 0 : KeyTimes.Num() - 1;
 	const int32 Inc = bBackwards ? -1 : 1;
 
-	const float Position = Context.GetTime() * Context.GetRootToSequenceTransform().Inverse();
+	const float PositionInSeconds = Context.GetTime() * Context.GetRootToSequenceTransform().Inverse() / Context.GetFrameRate();
 
 	if (bBackwards)
 	{
 		// Trigger events backwards
 		for (int32 KeyIndex = KeyTimes.Num() - 1; KeyIndex >= 0; --KeyIndex)
 		{
-			float Time = KeyTimes[KeyIndex];
+			FFrameNumber Time = KeyTimes[KeyIndex];
 			if (SweptRange.Contains(Time))
 			{
-				Events.Add(FMovieSceneEventData(KeyValues[KeyIndex], Position));
+				Events.Add(FMovieSceneEventData(KeyValues[KeyIndex], PositionInSeconds));
 			}
 		}
 	}
 	// Trigger events forwards
 	else for (int32 KeyIndex = 0; KeyIndex < KeyTimes.Num(); ++KeyIndex)
 	{
-		float Time = KeyTimes[KeyIndex];
+		FFrameNumber Time = KeyTimes[KeyIndex];
 		if (SweptRange.Contains(Time))
 		{
-			Events.Add(FMovieSceneEventData(KeyValues[KeyIndex], Position));
+			Events.Add(FMovieSceneEventData(KeyValues[KeyIndex], PositionInSeconds));
 		}
 	}
 
diff --git a/Engine/Source/Runtime/MovieSceneTracks/Private/Evaluation/MovieSceneFadeTemplate.cpp b/Engine/Source/Runtime/MovieSceneTracks/Private/Evaluation/MovieSceneFadeTemplate.cpp
--- a/Engine/Source/Runtime/MovieSceneTracks/Private/Evaluation/MovieSceneFadeTemplate.cpp
+++ b/Engine/Source/Runtime/MovieSceneTracks/Private/Evaluation/MovieSceneFadeTemplate.cpp
@@ -121,7 +121,7 @@ struct FFadeExecutionToken : IMovieSceneExecutionToken, FFadeTrackToken
 };
 
 FMovieSceneFadeSectionTemplate::FMovieSceneFadeSectionTemplate(const UMovieSceneFadeSection& Section)
-	: FadeCurve(Section.GetFloatCurve())
+	: FadeCurve(Section.GetChannel())
 	, FadeColor(Section.FadeColor)
 	, bFadeAudio(Section.bFadeAudio)
 {
@@ -129,6 +129,9 @@ FMovieSceneFadeSectionTemplate::FMovieSceneFadeSectionTemplate(const UMovieScene
 
 void FMovieSceneFadeSectionTemplate::Evaluate(const FMovieSceneEvaluationOperand& Operand, const FMovieSceneContext& Context, const FPersistentEvaluationData& PersistentData, FMovieSceneExecutionTokens& ExecutionTokens) const
 {
-	float FadeValue = FadeCurve.Eval(Context.GetTime());
-	ExecutionTokens.Add(FFadeExecutionToken(FadeValue, FadeColor, bFadeAudio));
+	float FadeValue = 1.f;
+	if (FadeCurve.Evaluate(Context.GetTime(), FadeValue))
+	{
+		ExecutionTokens.Add(FFadeExecutionToken(FadeValue, FadeColor, bFadeAudio));
+	}
 }
diff --git a/Engine/Source/Runtime/MovieSceneTracks/Private/Evaluation/MovieSceneFadeTemplate.h b/Engine/Source/Runtime/MovieSceneTracks/Private/Evaluation/MovieSceneFadeTemplate.h
--- a/Engine/Source/Runtime/MovieSceneTracks/Private/Evaluation/MovieSceneFadeTemplate.h
+++ b/Engine/Source/Runtime/MovieSceneTracks/Private/Evaluation/MovieSceneFadeTemplate.h
@@ -4,7 +4,7 @@
 
 #include "CoreMinimal.h"
 #include "UObject/ObjectMacros.h"
-#include "Curves/RichCurve.h"
+#include "Channels/MovieSceneFloatChannel.h"
 #include "Evaluation/MovieSceneTrackImplementation.h"
 #include "Evaluation/MovieSceneEvalTemplate.h"
 #include "Sections/MovieSceneFadeSection.h"
@@ -25,8 +25,8 @@ private:
 	virtual void Evaluate(const FMovieSceneEvaluationOperand& Operand, const FMovieSceneContext& Context, const FPersistentEvaluationData& PersistentData, FMovieSceneExecutionTokens& ExecutionTokens) const override;
 
 	UPROPERTY()
-	FRichCurve FadeCurve;
-	
+	FMovieSceneFloatChannel FadeCurve;
+
 	UPROPERTY()
 	FLinearColor FadeColor;
 
diff --git a/Engine/Source/Runtime/MovieSceneTracks/Private/Evaluation/MovieSceneParameterTemplate.cpp b/Engine/Source/Runtime/MovieSceneTracks/Private/Evaluation/MovieSceneParameterTemplate.cpp
--- a/Engine/Source/Runtime/MovieSceneTracks/Private/Evaluation/MovieSceneParameterTemplate.cpp
+++ b/Engine/Source/Runtime/MovieSceneTracks/Private/Evaluation/MovieSceneParameterTemplate.cpp
@@ -15,42 +15,46 @@ FMovieSceneParameterSectionTemplate::FMovieSceneParameterSectionTemplate(const U
 
 void FMovieSceneParameterSectionTemplate::EvaluateCurves(const FMovieSceneContext& Context, FEvaluatedParameterSectionValues& Values) const
 {
-	float Time = Context.GetTime();
+	const FFrameTime Time = Context.GetTime();
 
 	for ( const FScalarParameterNameAndCurve& Scalar : Scalars )
 	{
-		Values.ScalarValues.Add(
-			FScalarParameterNameAndValue(Scalar.ParameterName, Scalar.ParameterCurve.Eval(Time))
-			);
+		float Value = 0;
+		if (Scalar.ParameterCurve.Evaluate(Time, Value))
+		{
+			Values.ScalarValues.Emplace(Scalar.ParameterName, Value);
+		}
 	}
 
 	for ( const FVectorParameterNameAndCurves& Vector : Vectors )
 	{
-		Values.VectorValues.Add(
-			FVectorParameterNameAndValue(
-				Vector.ParameterName,
-				FVector(
-					Vector.XCurve.Eval(Time),
-					Vector.YCurve.Eval(Time),
-					Vector.ZCurve.Eval(Time)
-				)
-			)
-		);
+		FVector Value(ForceInitToZero);
+
+		bool bAnyEvaluated = false;
+		bAnyEvaluated |= Vector.XCurve.Evaluate(Time, Value.X);
+		bAnyEvaluated |= Vector.YCurve.Evaluate(Time, Value.Y);
+		bAnyEvaluated |= Vector.ZCurve.Evaluate(Time, Value.Z);
+
+		if (bAnyEvaluated)
+		{
+			Values.VectorValues.Emplace(Vector.ParameterName, Value);
+		}
 	}
 
 	for ( const FColorParameterNameAndCurves& Color : Colors )
 	{
-		Values.ColorValues.Add(
-			FColorParameterNameAndValue(
-				Color.ParameterName,
-				FLinearColor(
-					Color.RedCurve.Eval(Time ),
-					Color.GreenCurve.Eval(Time),
-					Color.BlueCurve.Eval(Time),
-					Color.AlphaCurve.Eval(Time)
-				)
-			)
-		);
+		FLinearColor ColorValue = FLinearColor::White;
+
+		bool bAnyEvaluated = false;
+		bAnyEvaluated |= Color.RedCurve.Evaluate(  Time, ColorValue.R);
+		bAnyEvaluated |= Color.GreenCurve.Evaluate(Time, ColorValue.G);
+		bAnyEvaluated |= Color.BlueCurve.Evaluate( Time, ColorValue.B);
+		bAnyEvaluated |= Color.AlphaCurve.Evaluate(Time, ColorValue.A);
+
+		if (bAnyEvaluated)
+		{
+			Values.ColorValues.Emplace(Color.ParameterName, ColorValue);
+		}
 	}
 }
 
diff --git a/Engine/Source/Runtime/MovieSceneTracks/Private/Evaluation/MovieSceneParticleTemplate.cpp b/Engine/Source/Runtime/MovieSceneTracks/Private/Evaluation/MovieSceneParticleTemplate.cpp
--- a/Engine/Source/Runtime/MovieSceneTracks/Private/Evaluation/MovieSceneParticleTemplate.cpp
+++ b/Engine/Source/Runtime/MovieSceneTracks/Private/Evaluation/MovieSceneParticleTemplate.cpp
@@ -22,12 +22,6 @@ static UParticleSystemComponent* GetParticleSystemComponentFromObject(UObject* O
 	}
 }
 
-struct FParticleKeyState : IPersistentEvaluationData
-{
-	FKeyHandle LastKeyHandle;
-	FKeyHandle InvalidKeyHandle;
-};
-
 /** A movie scene pre-animated token that stores a pre-animated active state */
 struct FActivePreAnimatedToken : IMovieScenePreAnimatedToken
 {
@@ -73,20 +67,14 @@ private:
 struct FParticleTrackExecutionToken
 	: IMovieSceneExecutionToken
 {
-	FParticleTrackExecutionToken(EParticleKey::Type InParticleKey, TOptional<FKeyHandle> InKeyHandle = TOptional<FKeyHandle>())
-		: ParticleKey(InParticleKey), KeyHandle(InKeyHandle)
-	{
-	}
+	FParticleTrackExecutionToken(EParticleKey InParticleKey)
+		: ParticleKey(InParticleKey)
+	{}
 
 	/** Execute this token, operating on all objects referenced by 'Operand' */
 	virtual void Execute(const FMovieSceneContext& Context, const FMovieSceneEvaluationOperand& Operand, FPersistentEvaluationData& PersistentData, IMovieScenePlayer& Player) override
 	{
 		MOVIESCENE_DETAILED_SCOPE_CYCLE_COUNTER(MovieSceneEval_ParticleTrack_TokenExecute)
-		
-		if (KeyHandle.IsSet())
-		{
-			PersistentData.GetOrAddSectionData<FParticleKeyState>().LastKeyHandle = KeyHandle.GetValue();
-		}
 
 		for (TWeakObjectPtr<> Object : Player.FindBoundObjects(Operand))
 		{
@@ -117,34 +105,37 @@ struct FParticleTrackExecutionToken
 		}
 	}
 
-	EParticleKey::Type ParticleKey;
+	EParticleKey ParticleKey;
 	TOptional<FKeyHandle> KeyHandle;
 };
 
 FMovieSceneParticleSectionTemplate::FMovieSceneParticleSectionTemplate(const UMovieSceneParticleSection& Section)
-	: ParticleKeys(Section.GetParticleCurve())
+	: ParticleKeys(Section.ParticleKeys)
 {
 }
 
 void FMovieSceneParticleSectionTemplate::Evaluate(const FMovieSceneEvaluationOperand& Operand, const FMovieSceneContext& Context, const FPersistentEvaluationData& PersistentData, FMovieSceneExecutionTokens& ExecutionTokens) const
 {
-	const bool bPlaying = Context.GetDirection() == EPlayDirection::Forwards && Context.GetRange().Size<float>() >= 0.f && Context.GetStatus() == EMovieScenePlayerStatus::Playing;
-
-	const FParticleKeyState* SectionData = PersistentData.FindSectionData<FParticleKeyState>();
+	const bool bPlaying = Context.GetDirection() == EPlayDirection::Forwards && Context.GetRange().Size<FFrameTime>() >= FFrameTime(0) && Context.GetStatus() == EMovieScenePlayerStatus::Playing;
 
 	if (!bPlaying)
 	{
-		ExecutionTokens.Add(FParticleTrackExecutionToken(EParticleKey::Deactivate, SectionData ? SectionData->InvalidKeyHandle : TOptional<FKeyHandle>()));
+		ExecutionTokens.Add(FParticleTrackExecutionToken(EParticleKey::Deactivate));
 	}
 	else
 	{
-		FKeyHandle PreviousHandle = ParticleKeys.FindKeyBeforeOrAt(Context.GetTime());
-		if (ParticleKeys.IsKeyHandleValid(PreviousHandle) && (!SectionData || SectionData->LastKeyHandle != PreviousHandle))
-		{
-			FParticleTrackExecutionToken NewToken(
-				(EParticleKey::Type)ParticleKeys.GetKey(PreviousHandle).Value,
-				PreviousHandle);
+		TRange<FFrameNumber> PlaybackRange = Context.GetFrameNumberRange();
 
+		TMovieSceneChannel<const uint8> ChannelInterface = ParticleKeys.GetInterface();
+
+		// Find the index of the key handle that exists before this time
+		TArrayView<const FFrameNumber> Times = ChannelInterface.GetTimes();
+		TArrayView<const uint8>        Values = ChannelInterface.GetValues();
+
+		const int32 LastKeyIndex = Algo::UpperBound(Times, PlaybackRange.GetUpperBoundValue())-1;
+		if (LastKeyIndex >= 0 && PlaybackRange.Contains(Times[LastKeyIndex]))
+		{
+			FParticleTrackExecutionToken NewToken((EParticleKey)Values[LastKeyIndex]);
 			ExecutionTokens.Add(MoveTemp(NewToken));
 		}
 	}
diff --git a/Engine/Source/Runtime/MovieSceneTracks/Private/Evaluation/MovieSceneParticleTemplate.h b/Engine/Source/Runtime/MovieSceneTracks/Private/Evaluation/MovieSceneParticleTemplate.h
--- a/Engine/Source/Runtime/MovieSceneTracks/Private/Evaluation/MovieSceneParticleTemplate.h
+++ b/Engine/Source/Runtime/MovieSceneTracks/Private/Evaluation/MovieSceneParticleTemplate.h
@@ -4,7 +4,7 @@
 
 #include "CoreMinimal.h"
 #include "UObject/ObjectMacros.h"
-#include "Curves/IntegralCurve.h"
+#include "Sections/MovieSceneParticleSection.h"
 #include "Evaluation/MovieSceneEvalTemplate.h"
 
 #include "MovieSceneParticleTemplate.generated.h"
@@ -20,7 +20,7 @@ struct FMovieSceneParticleSectionTemplate : public FMovieSceneEvalTemplate
 	FMovieSceneParticleSectionTemplate(const UMovieSceneParticleSection& Section);
 
 	UPROPERTY()
-	FIntegralCurve ParticleKeys;
+	FMovieSceneParticleChannel ParticleKeys;
 
 private:
 
diff --git a/Engine/Source/Runtime/MovieSceneTracks/Private/Evaluation/MovieScenePropertyTemplates.cpp b/Engine/Source/Runtime/MovieSceneTracks/Private/Evaluation/MovieScenePropertyTemplates.cpp
--- a/Engine/Source/Runtime/MovieSceneTracks/Private/Evaluation/MovieScenePropertyTemplates.cpp
+++ b/Engine/Source/Runtime/MovieSceneTracks/Private/Evaluation/MovieScenePropertyTemplates.cpp
@@ -29,7 +29,7 @@ namespace
 //	Boolean Property Template
 FMovieSceneBoolPropertySectionTemplate::FMovieSceneBoolPropertySectionTemplate(const UMovieSceneBoolSection& Section, const UMovieScenePropertyTrack& Track)
 	: FMovieScenePropertySectionTemplate(Track.GetPropertyName(), Track.GetPropertyPath())
-	, BoolCurve(Section.GetCurve())
+	, BoolCurve(Section.GetChannel())
 {
 	PropertyData.PropertyName = SanitizeBoolPropertyName(PropertyData.PropertyName);
 }
@@ -37,9 +37,10 @@ FMovieSceneBoolPropertySectionTemplate::FMovieSceneBoolPropertySectionTemplate(c
 void FMovieSceneBoolPropertySectionTemplate::Evaluate(const FMovieSceneEvaluationOperand& Operand, const FMovieSceneContext& Context, const FPersistentEvaluationData& PersistentData, FMovieSceneExecutionTokens& ExecutionTokens) const
 {
 	// Only evaluate if the curve has any data
-	if (BoolCurve.HasAnyData())
+	bool Result = false;
+	if (BoolCurve.Evaluate(Context.GetTime(), Result))
 	{
-		ExecutionTokens.Add(TPropertyTrackExecutionToken<bool>(!!BoolCurve.Evaluate(Context.GetTime())));
+		ExecutionTokens.Add(TPropertyTrackExecutionToken<bool>(Result));
 	}
 }
 
@@ -48,22 +49,23 @@ void FMovieSceneBoolPropertySectionTemplate::Evaluate(const FMovieSceneEvaluatio
 //	Float Property Template
 FMovieSceneFloatPropertySectionTemplate::FMovieSceneFloatPropertySectionTemplate(const UMovieSceneFloatSection& Section, const UMovieScenePropertyTrack& Track)
 	: FMovieScenePropertySectionTemplate(Track.GetPropertyName(), Track.GetPropertyPath())
-	, FloatCurve(Section.GetFloatCurve())
+	, FloatFunction(Section.GetChannel())
 	, BlendType(Section.GetBlendType().Get())
 {}
 
 void FMovieSceneFloatPropertySectionTemplate::Evaluate(const FMovieSceneEvaluationOperand& Operand, const FMovieSceneContext& Context, const FPersistentEvaluationData& PersistentData, FMovieSceneExecutionTokens& ExecutionTokens) const
 {
+	float Result = 0.f;
+
 	// Only evaluate if the curve has any data
-	if (FloatCurve.HasAnyData())
+	if (FloatFunction.Evaluate(Context.GetTime(), Result))
 	{
 		// Actuator type ID for this property
 		FMovieSceneBlendingActuatorID ActuatorTypeID = EnsureActuator<float>(ExecutionTokens.GetBlendingAccumulator());
 
 		// Add the blendable to the accumulator
-		const float Value = FloatCurve.Eval(Context.GetTime());
 		const float Weight = EvaluateEasing(Context.GetTime());
-		ExecutionTokens.BlendToken(ActuatorTypeID, TBlendableToken<float>(Value, BlendType, Weight));
+		ExecutionTokens.BlendToken(ActuatorTypeID, TBlendableToken<float>(Result, BlendType, Weight));
 	}
 }
 
@@ -72,15 +74,15 @@ void FMovieSceneFloatPropertySectionTemplate::Evaluate(const FMovieSceneEvaluati
 //	Byte Property Template
 FMovieSceneBytePropertySectionTemplate::FMovieSceneBytePropertySectionTemplate(const UMovieSceneByteSection& Section, const UMovieScenePropertyTrack& Track)
 	: FMovieScenePropertySectionTemplate(Track.GetPropertyName(), Track.GetPropertyPath())
-	, ByteCurve(Section.GetCurve())
+	, ByteCurve(Section.ByteCurve)
 {}
 
 void FMovieSceneBytePropertySectionTemplate::Evaluate(const FMovieSceneEvaluationOperand& Operand, const FMovieSceneContext& Context, const FPersistentEvaluationData& PersistentData, FMovieSceneExecutionTokens& ExecutionTokens) const
 {
-	// Only evaluate if the curve has any data
-	if (ByteCurve.HasAnyData())
+	uint8 Result = 0;
+	if (ByteCurve.Evaluate(Context.GetTime(), Result))
 	{
-		ExecutionTokens.Add(TPropertyTrackExecutionToken<uint8>(ByteCurve.Evaluate(Context.GetTime())));
+		ExecutionTokens.Add(TPropertyTrackExecutionToken<uint8>(Result));
 	}
 }
 
@@ -89,15 +91,15 @@ void FMovieSceneBytePropertySectionTemplate::Evaluate(const FMovieSceneEvaluatio
 //	Enum Property Template
 FMovieSceneEnumPropertySectionTemplate::FMovieSceneEnumPropertySectionTemplate(const UMovieSceneEnumSection& Section, const UMovieScenePropertyTrack& Track)
 	: FMovieScenePropertySectionTemplate(Track.GetPropertyName(), Track.GetPropertyPath())
-	, EnumCurve(Section.GetCurve())
+	, EnumCurve(Section.EnumCurve)
 {}
 
 void FMovieSceneEnumPropertySectionTemplate::Evaluate(const FMovieSceneEvaluationOperand& Operand, const FMovieSceneContext& Context, const FPersistentEvaluationData& PersistentData, FMovieSceneExecutionTokens& ExecutionTokens) const
 {
-	// Only evaluate if the curve has any data
-	if (EnumCurve.HasAnyData())
+	uint8 Result = 0;
+	if (EnumCurve.Evaluate(Context.GetTime(), Result))
 	{
-		ExecutionTokens.Add(TPropertyTrackExecutionToken<int64>(EnumCurve.Evaluate(Context.GetTime())));
+		ExecutionTokens.Add(TPropertyTrackExecutionToken<uint8>(Result));
 	}
 }
 
@@ -106,21 +108,21 @@ void FMovieSceneEnumPropertySectionTemplate::Evaluate(const FMovieSceneEvaluatio
 //	Integer Property Template
 FMovieSceneIntegerPropertySectionTemplate::FMovieSceneIntegerPropertySectionTemplate(const UMovieSceneIntegerSection& Section, const UMovieScenePropertyTrack& Track)
 	: FMovieScenePropertySectionTemplate(Track.GetPropertyName(), Track.GetPropertyPath())
-	, IntegerCurve(Section.GetCurve())
+	, IntegerCurve(Section.GetChannel())
 	, BlendType(Section.GetBlendType().Get())
 {}
 
 void FMovieSceneIntegerPropertySectionTemplate::Evaluate(const FMovieSceneEvaluationOperand& Operand, const FMovieSceneContext& Context, const FPersistentEvaluationData& PersistentData, FMovieSceneExecutionTokens& ExecutionTokens) const
 {
-	if (IntegerCurve.HasAnyData())
+	int32 Result = 0;
+	if (IntegerCurve.Evaluate(Context.GetTime(), Result))
 	{
 		// Actuator type ID for this property
 		FMovieSceneBlendingActuatorID ActuatorTypeID = EnsureActuator<int32>(ExecutionTokens.GetBlendingAccumulator());
 
 		// Add the blendable to the accumulator
-		const int32 Value = IntegerCurve.Evaluate(Context.GetTime());
 		const float Weight = EvaluateEasing(Context.GetTime());
-		ExecutionTokens.BlendToken(ActuatorTypeID, TBlendableToken<int32>(Value, BlendType, Weight));
+		ExecutionTokens.BlendToken(ActuatorTypeID, TBlendableToken<int32>(Result, BlendType, Weight));
 	}
 }
 
@@ -129,14 +131,15 @@ void FMovieSceneIntegerPropertySectionTemplate::Evaluate(const FMovieSceneEvalua
 //	String Property Template
 FMovieSceneStringPropertySectionTemplate::FMovieSceneStringPropertySectionTemplate(const UMovieSceneStringSection& Section, const UMovieScenePropertyTrack& Track)
 	: FMovieScenePropertySectionTemplate(Track.GetPropertyName(), Track.GetPropertyPath())
-	, StringCurve(Section.GetStringCurve())
+	, StringCurve(Section.GetChannel())
 {}
 
 void FMovieSceneStringPropertySectionTemplate::Evaluate(const FMovieSceneEvaluationOperand& Operand, const FMovieSceneContext& Context, const FPersistentEvaluationData& PersistentData, FMovieSceneExecutionTokens& ExecutionTokens) const
 {
-	if (StringCurve.HasAnyData())
+	const FString* Result = StringCurve.Evaluate(Context.GetTime());
+	if (Result)
 	{
-		ExecutionTokens.Add(TPropertyTrackExecutionToken<FString>(StringCurve.Eval(Context.GetTime(), FString())));
+		ExecutionTokens.Add(TPropertyTrackExecutionToken<FString>(*Result));
 	}
 }
 
@@ -150,22 +153,23 @@ FMovieSceneVectorPropertySectionTemplate::FMovieSceneVectorPropertySectionTempla
 {
 	for (int32 Index = 0; Index < NumChannelsUsed; ++Index)
 	{
-		ComponentCurves[Index] = Section.GetCurve(Index);
+		ComponentCurves[Index] = Section.GetChannel(Index);
 	}
 }
 
 /** Helper function fo evaluating a number of curves for a specific vector type */
 template<typename VectorType, uint8 N>
-void EvaluateVectorCurve(EMovieSceneBlendType BlendType, float Weight, float Time, const FRichCurve* Curves, FMovieSceneBlendingActuatorID ActuatorTypeID, FMovieSceneExecutionTokens& ExecutionTokens)
+void EvaluateVectorCurve(EMovieSceneBlendType BlendType, float Weight, FFrameTime Time, const FMovieSceneFloatChannel* Channels, FMovieSceneBlendingActuatorID ActuatorTypeID, FMovieSceneExecutionTokens& ExecutionTokens)
 {
 	MovieScene::TMultiChannelValue<float, N> AnimatedChannels;
 
 	for (uint8 Index = 0; Index < N; ++Index)
 	{
-		const FRichCurve& Curve = Curves[Index];
-		if (Curve.HasAnyData())
+		const FMovieSceneFloatChannel& Channel = Channels[Index];
+		float Result = 0;
+		if (Channel.Evaluate(Time, Result))
 		{
-			AnimatedChannels.Set(Index, Curve.Eval(Time));
+			AnimatedChannels.Set(Index, Result);
 		}
 	}
 
@@ -180,8 +184,8 @@ void FMovieSceneVectorPropertySectionTemplate::Evaluate(const FMovieSceneEvaluat
 {
 	FMovieSceneBlendingAccumulator& Accumulator = ExecutionTokens.GetBlendingAccumulator();
 
-	const float Time = Context.GetTime();
-	const float Weight = EvaluateEasing(Context.GetTime());
+	const FFrameTime Time   = Context.GetTime();
+	const float      Weight = EvaluateEasing(Context.GetTime());
 
 	switch (NumChannelsUsed)
 	{
@@ -218,11 +222,13 @@ void FMovieSceneTransformPropertySectionTemplate::Evaluate(const FMovieSceneEval
 	// Actuator type ID for this property
 	FMovieSceneBlendingActuatorID ActuatorTypeID = EnsureActuator<FTransform>(ExecutionTokens.GetBlendingAccumulator());
 
-	// Add the blendable to the accumulator
+	// Evaluate the easing, and multiply this with the manual weight specified on the manual weight curve
 	float Weight = EvaluateEasing(Context.GetTime());
 	if (EnumHasAllFlags(TemplateData.Mask.GetChannels(), EMovieSceneTransformChannel::Weight))
 	{
-		Weight *= TemplateData.ManualWeight.Eval(Context.GetTime());
+		float ManualWeight = 1.f;
+		TemplateData.ManualWeight.Evaluate(Context.GetTime(), ManualWeight);
+		Weight *= ManualWeight;
 	}
 
 	// Add the blendable to the accumulator
diff --git a/Engine/Source/Runtime/MovieSceneTracks/Private/Evaluation/MovieScenePropertyTemplates.h b/Engine/Source/Runtime/MovieSceneTracks/Private/Evaluation/MovieScenePropertyTemplates.h
--- a/Engine/Source/Runtime/MovieSceneTracks/Private/Evaluation/MovieScenePropertyTemplates.h
+++ b/Engine/Source/Runtime/MovieSceneTracks/Private/Evaluation/MovieScenePropertyTemplates.h
@@ -4,14 +4,15 @@
 
 #include "CoreMinimal.h"
 #include "UObject/ObjectMacros.h"
-#include "Curves/RichCurve.h"
 #include "Evaluation/MovieSceneEvalTemplate.h"
 #include "Evaluation/MovieScenePropertyTemplate.h"
-#include "Curves/IntegralCurve.h"
-#include "Curves/StringCurve.h"
-#include "Templates/UnrealTypeTraits.h"
+#include "Sections/MovieSceneStringSection.h"
 #include "Evaluation/MovieScene3DTransformTemplate.h"
 #include "Evaluation/Blending/MovieSceneBlendType.h"
+#include "Channels/MovieSceneBoolChannel.h"
+#include "Channels/MovieSceneByteChannel.h"
+#include "Channels/MovieSceneIntegerChannel.h"
+#include "Channels/MovieSceneFloatChannel.h"
 #include "MovieScenePropertyTemplates.generated.h"
 
 class UMovieSceneBoolSection;
@@ -19,7 +20,6 @@ class UMovieSceneByteSection;
 class UMovieSceneFloatSection;
 class UMovieSceneIntegerSection;
 class UMovieScenePropertyTrack;
-class UMovieSceneStringSection;
 class UMovieSceneVectorSection;
 class UMovieSceneEnumSection;
 class UMovieScene3DTransformSection;
@@ -40,7 +40,7 @@ protected:
 	virtual void Evaluate(const FMovieSceneEvaluationOperand& Operand, const FMovieSceneContext& Context, const FPersistentEvaluationData& PersistentData, FMovieSceneExecutionTokens& ExecutionTokens) const override;
 
 	UPROPERTY()
-	FIntegralCurve BoolCurve;
+	FMovieSceneBoolChannel BoolCurve;
 };
 
 USTRUCT()
@@ -57,7 +57,7 @@ protected:
 	virtual void Evaluate(const FMovieSceneEvaluationOperand& Operand, const FMovieSceneContext& Context, const FPersistentEvaluationData& PersistentData, FMovieSceneExecutionTokens& ExecutionTokens) const override;
 
 	UPROPERTY()
-	FRichCurve FloatCurve;
+	FMovieSceneFloatChannel FloatFunction;
 
 	UPROPERTY()
 	EMovieSceneBlendType BlendType;
@@ -78,7 +78,7 @@ protected:
 	virtual void Evaluate(const FMovieSceneEvaluationOperand& Operand, const FMovieSceneContext& Context, const FPersistentEvaluationData& PersistentData, FMovieSceneExecutionTokens& ExecutionTokens) const override;
 
 	UPROPERTY()
-	FIntegralCurve ByteCurve;
+	FMovieSceneByteChannel ByteCurve;
 };
 
 USTRUCT()
@@ -96,7 +96,7 @@ protected:
 	virtual void Evaluate(const FMovieSceneEvaluationOperand& Operand, const FMovieSceneContext& Context, const FPersistentEvaluationData& PersistentData, FMovieSceneExecutionTokens& ExecutionTokens) const override;
 
 	UPROPERTY()
-	FIntegralCurve EnumCurve;
+	FMovieSceneByteChannel EnumCurve;
 };
 
 USTRUCT()
@@ -113,7 +113,7 @@ protected:
 	virtual void Evaluate(const FMovieSceneEvaluationOperand& Operand, const FMovieSceneContext& Context, const FPersistentEvaluationData& PersistentData, FMovieSceneExecutionTokens& ExecutionTokens) const override;
 
 	UPROPERTY()
-	FIntegralCurve IntegerCurve;
+	FMovieSceneIntegerChannel IntegerCurve;
 
 	UPROPERTY()
 	EMovieSceneBlendType BlendType;
@@ -134,7 +134,7 @@ protected:
 	virtual void Evaluate(const FMovieSceneEvaluationOperand& Operand, const FMovieSceneContext& Context, const FPersistentEvaluationData& PersistentData, FMovieSceneExecutionTokens& ExecutionTokens) const override;
 
 	UPROPERTY()
-	FStringCurve StringCurve;
+	FMovieSceneStringChannel StringCurve;
 };
 
 Expose_TNameOf(FVector2D);
@@ -155,7 +155,7 @@ protected:
 	virtual void Evaluate(const FMovieSceneEvaluationOperand& Operand, const FMovieSceneContext& Context, const FPersistentEvaluationData& PersistentData, FMovieSceneExecutionTokens& ExecutionTokens) const override;
 
 	UPROPERTY()
-	FRichCurve ComponentCurves[4];
+	FMovieSceneFloatChannel ComponentCurves[4];
 
 	UPROPERTY()
 	int32 NumChannelsUsed;
diff --git a/Engine/Source/Runtime/MovieSceneTracks/Private/Evaluation/MovieSceneSkeletalAnimationTemplate.cpp b/Engine/Source/Runtime/MovieSceneTracks/Private/Evaluation/MovieSceneSkeletalAnimationTemplate.cpp
--- a/Engine/Source/Runtime/MovieSceneTracks/Private/Evaluation/MovieSceneSkeletalAnimationTemplate.cpp
+++ b/Engine/Source/Runtime/MovieSceneTracks/Private/Evaluation/MovieSceneSkeletalAnimationTemplate.cpp
@@ -199,7 +199,7 @@ namespace MovieScene
 			const bool bFireNotifies = !bPreviewPlayback || (PlayerStatus != EMovieScenePlayerStatus::Jumping && PlayerStatus != EMovieScenePlayerStatus::Stopped);
 
 			// When jumping from one cut to another cut, the delta time should be 0 so that anim notifies before the current position are not evaluated. Note, anim notifies at the current time should still be evaluated.
-			const float DeltaTime = Context.HasJumped() ? 0.f : Context.GetRange().Size<float>();
+			const double DeltaTime = ( Context.HasJumped() ? FFrameTime(0) : Context.GetRange().Size<FFrameTime>() ) / Context.GetFrameRate();
 
 			const bool bResetDynamics = PlayerStatus == EMovieScenePlayerStatus::Stepping || 
 										PlayerStatus == EMovieScenePlayerStatus::Jumping || 
@@ -343,7 +343,7 @@ template<> FMovieSceneAnimTypeID GetBlendingDataType<MovieScene::FBlendedAnimati
 }
 
 FMovieSceneSkeletalAnimationSectionTemplate::FMovieSceneSkeletalAnimationSectionTemplate(const UMovieSceneSkeletalAnimationSection& InSection)
-	: Params(InSection.Params, InSection.GetStartTime(), InSection.GetEndTime())
+	: Params(InSection.Params, InSection.GetInclusiveStartFrame(), InSection.GetExclusiveEndFrame())
 {
 }
 
@@ -352,8 +352,12 @@ void FMovieSceneSkeletalAnimationSectionTemplate::Evaluate(const FMovieSceneEval
 	if (Params.Animation)
 	{
 		// calculate the time at which to evaluate the animation
-		float EvalTime = Params.MapTimeToAnimation(Context.GetTime());
-		float Weight = Params.Weight.Eval(Context.GetTime()) * EvaluateEasing(Context.GetTime());
+		float EvalTime = Params.MapTimeToAnimation(Context.GetTime(), Context.GetFrameRate());
+
+		float ManualWeight = 1.f;
+		Params.Weight.Evaluate(Context.GetTime(), ManualWeight);
+
+		const float Weight = ManualWeight * EvaluateEasing(Context.GetTime());
 
 		FOptionalMovieSceneBlendType BlendType = SourceSection->GetBlendType();
 		check(BlendType.IsValid());
@@ -374,28 +378,25 @@ void FMovieSceneSkeletalAnimationSectionTemplate::Evaluate(const FMovieSceneEval
 	}
 }
 
-float FMovieSceneSkeletalAnimationSectionTemplateParameters::MapTimeToAnimation(float ThisPosition) const
+float FMovieSceneSkeletalAnimationSectionTemplateParameters::MapTimeToAnimation(FFrameTime InPosition, FFrameRate InFrameRate) const
 {
-	// @todo: Sequencer: what is this for??
-	//ThisPosition -= 1 / 1000.0f;
-
-	ThisPosition = FMath::Clamp(ThisPosition, SectionStartTime, SectionEndTime);
+	InPosition = FMath::Clamp(InPosition, FFrameTime(SectionStartTime), FFrameTime(SectionEndTime-1));
 
 	const float SectionPlayRate = PlayRate;
 	const float AnimPlayRate = FMath::IsNearlyZero(SectionPlayRate) ? 1.0f : SectionPlayRate;
 
 	const float SeqLength = GetSequenceLength() - (StartOffset + EndOffset);
 
-	ThisPosition = (ThisPosition - SectionStartTime) * AnimPlayRate;
+	float AnimPosition = FFrameTime::FromDecimal((InPosition - SectionStartTime).AsDecimal() * AnimPlayRate) / InFrameRate;
 	if (SeqLength > 0.f)
 	{
-		ThisPosition = FMath::Fmod(ThisPosition, SeqLength);
+		AnimPosition = FMath::Fmod(AnimPosition, SeqLength);
 	}
-	ThisPosition += StartOffset;
+	AnimPosition += StartOffset;
 	if (bReverse)
 	{
-		ThisPosition = (SeqLength - (ThisPosition - StartOffset)) + StartOffset;
+		AnimPosition = (SeqLength - (AnimPosition - StartOffset)) + StartOffset;
 	}
 
-	return ThisPosition;
+	return AnimPosition;
 }
diff --git a/Engine/Source/Runtime/MovieSceneTracks/Private/Evaluation/MovieSceneSkeletalAnimationTemplate.h b/Engine/Source/Runtime/MovieSceneTracks/Private/Evaluation/MovieSceneSkeletalAnimationTemplate.h
--- a/Engine/Source/Runtime/MovieSceneTracks/Private/Evaluation/MovieSceneSkeletalAnimationTemplate.h
+++ b/Engine/Source/Runtime/MovieSceneTracks/Private/Evaluation/MovieSceneSkeletalAnimationTemplate.h
@@ -14,19 +14,19 @@ struct FMovieSceneSkeletalAnimationSectionTemplateParameters : public FMovieScen
 	GENERATED_BODY()
 
 	FMovieSceneSkeletalAnimationSectionTemplateParameters() {}
-	FMovieSceneSkeletalAnimationSectionTemplateParameters(const FMovieSceneSkeletalAnimationParams& BaseParams, float InSectionStartTime, float InSectionEndTime)
+	FMovieSceneSkeletalAnimationSectionTemplateParameters(const FMovieSceneSkeletalAnimationParams& BaseParams, FFrameNumber InSectionStartTime, FFrameNumber InSectionEndTime)
 		: FMovieSceneSkeletalAnimationParams(BaseParams)
 		, SectionStartTime(InSectionStartTime)
 		, SectionEndTime(InSectionEndTime)
 	{}
 
-	float MapTimeToAnimation(float InPosition) const;
+	float MapTimeToAnimation(FFrameTime InPosition, FFrameRate InFrameRate) const;
 
 	UPROPERTY()
-	float SectionStartTime;
+	FFrameNumber SectionStartTime;
 
 	UPROPERTY()
-	float SectionEndTime;
+	FFrameNumber SectionEndTime;
 };
 
 USTRUCT()
diff --git a/Engine/Source/Runtime/MovieSceneTracks/Private/Evaluation/MovieSceneSlomoTemplate.cpp b/Engine/Source/Runtime/MovieSceneTracks/Private/Evaluation/MovieSceneSlomoTemplate.cpp
--- a/Engine/Source/Runtime/MovieSceneTracks/Private/Evaluation/MovieSceneSlomoTemplate.cpp
+++ b/Engine/Source/Runtime/MovieSceneTracks/Private/Evaluation/MovieSceneSlomoTemplate.cpp
@@ -101,14 +101,14 @@ struct FSlomoExecutionToken : IMovieSceneExecutionToken, FSlomoTrackToken
 };
 
 FMovieSceneSlomoSectionTemplate::FMovieSceneSlomoSectionTemplate(const UMovieSceneSlomoSection& Section)
-	: SlomoCurve(Section.GetFloatCurve())
+	: SlomoCurve(Section.GetChannel())
 {
 }
 
 void FMovieSceneSlomoSectionTemplate::Evaluate(const FMovieSceneEvaluationOperand& Operand, const FMovieSceneContext& Context, const FPersistentEvaluationData& PersistentData, FMovieSceneExecutionTokens& ExecutionTokens) const
 {
-	float SlomoValue = SlomoCurve.Eval(Context.GetTime());
-	if (SlomoValue >= 0.f)
+	float SlomoValue = 1.f;
+	if (SlomoCurve.Evaluate(Context.GetTime(), SlomoValue) && SlomoValue >= 0.f)
 	{
 		ExecutionTokens.Add(FSlomoExecutionToken(SlomoValue));
 	}
diff --git a/Engine/Source/Runtime/MovieSceneTracks/Private/Evaluation/MovieSceneSlomoTemplate.h b/Engine/Source/Runtime/MovieSceneTracks/Private/Evaluation/MovieSceneSlomoTemplate.h
--- a/Engine/Source/Runtime/MovieSceneTracks/Private/Evaluation/MovieSceneSlomoTemplate.h
+++ b/Engine/Source/Runtime/MovieSceneTracks/Private/Evaluation/MovieSceneSlomoTemplate.h
@@ -4,7 +4,7 @@
 
 #include "CoreMinimal.h"
 #include "UObject/ObjectMacros.h"
-#include "Curves/RichCurve.h"
+#include "Channels/MovieSceneFloatChannel.h"
 #include "Evaluation/MovieSceneTrackImplementation.h"
 #include "Evaluation/MovieSceneEvalTemplate.h"
 #include "Sections/MovieSceneSlomoSection.h"
@@ -25,5 +25,5 @@ private:
 	virtual void Evaluate(const FMovieSceneEvaluationOperand& Operand, const FMovieSceneContext& Context, const FPersistentEvaluationData& PersistentData, FMovieSceneExecutionTokens& ExecutionTokens) const override;
 
 	UPROPERTY()
-	FRichCurve SlomoCurve;
+	FMovieSceneFloatChannel SlomoCurve;
 };
diff --git a/Engine/Source/Runtime/MovieSceneTracks/Private/Evaluation/MovieSceneSpawnTemplate.cpp b/Engine/Source/Runtime/MovieSceneTracks/Private/Evaluation/MovieSceneSpawnTemplate.cpp
--- a/Engine/Source/Runtime/MovieSceneTracks/Private/Evaluation/MovieSceneSpawnTemplate.cpp
+++ b/Engine/Source/Runtime/MovieSceneTracks/Private/Evaluation/MovieSceneSpawnTemplate.cpp
@@ -94,7 +94,7 @@ struct FSpawnObjectToken : IMovieSceneExecutionToken
 };
 
 FMovieSceneSpawnSectionTemplate::FMovieSceneSpawnSectionTemplate(const UMovieSceneSpawnSection& SpawnSection)
-	: Curve(SpawnSection.GetCurve())
+	: Curve(SpawnSection.GetChannel())
 {
 }
 
@@ -102,10 +102,11 @@ void FMovieSceneSpawnSectionTemplate::Evaluate(const FMovieSceneEvaluationOperan
 {
 	MOVIESCENE_DETAILED_SCOPE_CYCLE_COUNTER(MovieSceneEval_SpawnTrack_Evaluate)
 
-	const float Time = Context.GetTime();
-	ExecutionTokens.Add(
-		FSpawnObjectToken(Curve.Evaluate(Time) != 0)
-		);
+	bool SpawnValue = false;
+	if (Curve.Evaluate(Context.GetTime(), SpawnValue))
+	{
+		ExecutionTokens.Add(FSpawnObjectToken(SpawnValue));
+	}
 }
 
 FMovieSceneAnimTypeID FMovieSceneSpawnSectionTemplate::GetAnimTypeID()
diff --git a/Engine/Source/Runtime/MovieSceneTracks/Private/Evaluation/MovieSceneVisibilityTemplate.cpp b/Engine/Source/Runtime/MovieSceneTracks/Private/Evaluation/MovieSceneVisibilityTemplate.cpp
--- a/Engine/Source/Runtime/MovieSceneTracks/Private/Evaluation/MovieSceneVisibilityTemplate.cpp
+++ b/Engine/Source/Runtime/MovieSceneTracks/Private/Evaluation/MovieSceneVisibilityTemplate.cpp
@@ -130,11 +130,11 @@ void FMovieSceneVisibilitySectionTemplate::Evaluate(const FMovieSceneEvaluationO
 {
 	MOVIESCENE_DETAILED_SCOPE_CYCLE_COUNTER(MovieSceneEval_VisibilityTrack_Evaluate);
 
-	if (BoolCurve.HasAnyData())
+	bool VisibilityValue = false;
+	if (BoolCurve.Evaluate(Context.GetTime(), VisibilityValue))
 	{
 		// Invert this evaluation since the property is "bHiddenInGame" and we want the visualization to be the inverse of that. Green means visible.
-		bool bIsHidden = !BoolCurve.Evaluate(Context.GetTime());
-		ExecutionTokens.Add(FTemporarilyHiddenInGameExecutionToken(bIsHidden));
+		ExecutionTokens.Add(FTemporarilyHiddenInGameExecutionToken(!VisibilityValue));
 	}
 }
 
diff --git a/Engine/Source/Runtime/MovieSceneTracks/Private/Evaluation/MovieSceneVisibilityTemplate.h b/Engine/Source/Runtime/MovieSceneTracks/Private/Evaluation/MovieSceneVisibilityTemplate.h
--- a/Engine/Source/Runtime/MovieSceneTracks/Private/Evaluation/MovieSceneVisibilityTemplate.h
+++ b/Engine/Source/Runtime/MovieSceneTracks/Private/Evaluation/MovieSceneVisibilityTemplate.h
@@ -19,10 +19,6 @@ struct FMovieSceneVisibilitySectionTemplate : public FMovieSceneBoolPropertySect
 	FMovieSceneVisibilitySectionTemplate(){}
 	FMovieSceneVisibilitySectionTemplate(const UMovieSceneBoolSection& Section, const UMovieScenePropertyTrack& Track);
 
-	/** Temporarily hidden in game */
-	UPROPERTY()
-	bool bTemporarilyHiddenInGame;
-
 private:
 	virtual UScriptStruct& GetScriptStructImpl() const override { return *StaticStruct(); }
 	virtual void Evaluate(const FMovieSceneEvaluationOperand& Operand, const FMovieSceneContext& Context, const FPersistentEvaluationData& PersistentData, FMovieSceneExecutionTokens& ExecutionTokens) const override;
diff --git a/Engine/Source/Runtime/MovieSceneTracks/Private/Sections/MovieScene3DAttachSection.cpp b/Engine/Source/Runtime/MovieSceneTracks/Private/Sections/MovieScene3DAttachSection.cpp
--- a/Engine/Source/Runtime/MovieSceneTracks/Private/Sections/MovieScene3DAttachSection.cpp
+++ b/Engine/Source/Runtime/MovieSceneTracks/Private/Sections/MovieScene3DAttachSection.cpp
@@ -23,7 +23,7 @@ UMovieScene3DAttachSection::UMovieScene3DAttachSection( const FObjectInitializer
 }
 
 
-void UMovieScene3DAttachSection::AddAttach( float Time, float SequenceEndTime, const FMovieSceneObjectBindingID& InAttachBindingID )
+void UMovieScene3DAttachSection::SetAttachTargetID( const FMovieSceneObjectBindingID& InAttachBindingID )
 {
 	if (TryModify())
 	{
diff --git a/Engine/Source/Runtime/MovieSceneTracks/Private/Sections/MovieScene3DConstraintSection.cpp b/Engine/Source/Runtime/MovieSceneTracks/Private/Sections/MovieScene3DConstraintSection.cpp
--- a/Engine/Source/Runtime/MovieSceneTracks/Private/Sections/MovieScene3DConstraintSection.cpp
+++ b/Engine/Source/Runtime/MovieSceneTracks/Private/Sections/MovieScene3DConstraintSection.cpp
@@ -24,18 +24,6 @@ FGuid UMovieScene3DConstraintSection::GetConstraintId() const
 }
 
 
-TOptional<float> UMovieScene3DConstraintSection::GetKeyTime(FKeyHandle KeyHandle) const
-{
-	return TOptional<float>();
-}
-
-
-void UMovieScene3DConstraintSection::SetKeyTime(FKeyHandle KeyHandle, float Time)
-{
-	// do nothing
-}
-
-
 void UMovieScene3DConstraintSection::OnBindingsUpdated(const TMap<FGuid, FGuid>& OldGuidToNewGuidMap)
 {
 	if (OldGuidToNewGuidMap.Contains(ConstraintBindingID.GetGuid()))
diff --git a/Engine/Source/Runtime/MovieSceneTracks/Private/Sections/MovieScene3DPathSection.cpp b/Engine/Source/Runtime/MovieSceneTracks/Private/Sections/MovieScene3DPathSection.cpp
--- a/Engine/Source/Runtime/MovieSceneTracks/Private/Sections/MovieScene3DPathSection.cpp
+++ b/Engine/Source/Runtime/MovieSceneTracks/Private/Sections/MovieScene3DPathSection.cpp
@@ -2,6 +2,7 @@
 
 #include "Sections/MovieScene3DPathSection.h"
 #include "Components/SplineComponent.h"
+#include "Channels/MovieSceneChannelProxy.h"
 
 
 UMovieScene3DPathSection::UMovieScene3DPathSection( const FObjectInitializer& ObjectInitializer )
@@ -11,12 +12,35 @@ UMovieScene3DPathSection::UMovieScene3DPathSection( const FObjectInitializer& Ob
 	, bFollow (true)
 	, bReverse (false)
 	, bForceUpright (false)
-{ }
+{
+#if WITH_EDITOR
+
+	static const FMovieSceneChannelEditorData EditorData("Timing", NSLOCTEXT("MovieScene3DPathSection", "TimingArea", "Timing"));
+	ChannelProxy = MakeShared<FMovieSceneChannelProxy>(TimingCurve, EditorData, TMovieSceneExternalValue<float>());
+
+#else
+
+	ChannelProxy = MakeShared<FMovieSceneChannelProxy>(TimingCurve);
+
+#endif
+}
+
+void UMovieScene3DPathSection::InitialPlacement(const TArray<UMovieSceneSection*>& Sections, FFrameNumber InStartTime, int32 Duration, bool bAllowMultipleRows)
+{
+	Super::InitialPlacement(Sections, InStartTime, Duration, bAllowMultipleRows);
 
+	TMovieSceneChannel<FMovieSceneFloatValue> KeyInterface = TimingCurve.GetInterface();
+	KeyInterface.UpdateOrAddKey(InStartTime, FMovieSceneFloatValue(0.f));
+	if (Duration > 0)
+	{
+		KeyInterface.UpdateOrAddKey(InStartTime + Duration, FMovieSceneFloatValue(1.f));
+	}
+}
 
-void UMovieScene3DPathSection::Eval( USceneComponent* SceneComponent, float Position, USplineComponent* SplineComponent, FVector& OutTranslation, FRotator& OutRotation ) const
+void UMovieScene3DPathSection::Eval( USceneComponent* SceneComponent, FFrameTime Position, USplineComponent* SplineComponent, FVector& OutTranslation, FRotator& OutRotation ) const
 {
-	float Timing = TimingCurve.Eval( Position );
+	float Timing = 0.f;
+	TimingCurve.Evaluate( Position, Timing );
 
 	if (Timing < 0.f)
 	{
@@ -110,48 +134,10 @@ void UMovieScene3DPathSection::Eval( USceneComponent* SceneComponent, float Posi
 }
 
 
-void UMovieScene3DPathSection::MoveSection( float DeltaPosition, TSet<FKeyHandle>& KeyHandles )
-{
-	Super::MoveSection( DeltaPosition, KeyHandles );
-
-	// Move the curve
-	TimingCurve.ShiftCurve(DeltaPosition, KeyHandles);
-}
-
-
-void UMovieScene3DPathSection::DilateSection( float DilationFactor, float Origin, TSet<FKeyHandle>& KeyHandles )
-{	
-	Super::DilateSection(DilationFactor, Origin, KeyHandles);
-	
-	TimingCurve.ScaleCurve(Origin, DilationFactor, KeyHandles);
-}
-
-
-void UMovieScene3DPathSection::GetKeyHandles(TSet<FKeyHandle>& OutKeyHandles, TRange<float> TimeRange) const
-{
-	if (!TimeRange.Overlaps(GetRange()))
-	{
-		return;
-	}
-
-	for (auto It(TimingCurve.GetKeyHandleIterator()); It; ++It)
-	{
-		float Time = TimingCurve.GetKeyTime(It.Key());
-		if (TimeRange.Contains(Time))
-		{
-			OutKeyHandles.Add(It.Key());
-		}
-	}
-}
-
-
-void UMovieScene3DPathSection::AddPath( float Time, float SequenceEndTime, const FMovieSceneObjectBindingID& InPathBindingID )
+void UMovieScene3DPathSection::SetPathBindingID( const FMovieSceneObjectBindingID& InPathBindingID )
 {
 	if (TryModify())
 	{
 		ConstraintBindingID = InPathBindingID;
-
-		TimingCurve.UpdateOrAddKey(Time, 0);
-		TimingCurve.UpdateOrAddKey(SequenceEndTime, 1);
-	}		
+	}
 }
