Copying //UE4/Dev-Sequencer to //UE4/Dev-Main (Source: //UE4/Dev-Sequencer @ 3946692)

diff --git a/Engine/Source/Runtime/MovieSceneTracks/Private/Sections/MovieScene3DTransformSection.cpp b/Engine/Source/Runtime/MovieSceneTracks/Private/Sections/MovieScene3DTransformSection.cpp
--- a/Engine/Source/Runtime/MovieSceneTracks/Private/Sections/MovieScene3DTransformSection.cpp
+++ b/Engine/Source/Runtime/MovieSceneTracks/Private/Sections/MovieScene3DTransformSection.cpp
@@ -4,641 +4,404 @@
 #include "UObject/StructOnScope.h"
 #include "Evaluation/MovieScene3DTransformTemplate.h"
 #include "UObject/SequencerObjectVersion.h"
+#include "Algo/AnyOf.h"
+#include "Channels/MovieSceneChannelProxy.h"
+#include "GameFramework/Actor.h"
 
+#if WITH_EDITOR
 
-/* FMovieScene3DLocationKeyStruct interface
- *****************************************************************************/
-
-void FMovieScene3DLocationKeyStruct::PropagateChanges(const FPropertyChangedEvent& ChangeEvent)
+struct F3DTransformChannelEditorData
 {
-	for (int32 Index = 0; Index <= 2; ++Index)
+	F3DTransformChannelEditorData(EMovieSceneTransformChannel Mask)
 	{
-		if(LocationKeys[Index] != nullptr)
+		FText LocationGroup = NSLOCTEXT("MovieSceneTransformSection", "Location", "Location");
+		FText RotationGroup = NSLOCTEXT("MovieSceneTransformSection", "Rotation", "Rotation");
+		FText ScaleGroup    = NSLOCTEXT("MovieSceneTransformSection", "Scale",    "Scale");
 		{
-			LocationKeys[Index]->Value = Location[Index];
-			LocationKeys[Index]->Time = Time;
+			CommonData[0].SetIdentifiers("Location.X", FCommonChannelData::ChannelX, LocationGroup);
+			CommonData[0].bEnabled = EnumHasAllFlags(Mask, EMovieSceneTransformChannel::TranslationX);
+			CommonData[0].Color = FCommonChannelData::RedChannelColor;
+			CommonData[0].SortOrder = 0;
+
+			CommonData[1].SetIdentifiers("Location.Y", FCommonChannelData::ChannelY, LocationGroup);
+			CommonData[1].bEnabled = EnumHasAllFlags(Mask, EMovieSceneTransformChannel::TranslationY);
+			CommonData[1].Color = FCommonChannelData::GreenChannelColor;
+			CommonData[1].SortOrder = 1;
+
+			CommonData[2].SetIdentifiers("Location.Z", FCommonChannelData::ChannelZ, LocationGroup);
+			CommonData[2].bEnabled = EnumHasAllFlags(Mask, EMovieSceneTransformChannel::TranslationZ);
+			CommonData[2].Color = FCommonChannelData::BlueChannelColor;
+			CommonData[2].SortOrder = 2;
+		}
+		{
+			CommonData[3].SetIdentifiers("Rotation.X", NSLOCTEXT("MovieSceneTransformSection", "RotationX", "Roll"), RotationGroup);
+			CommonData[3].bEnabled = EnumHasAllFlags(Mask, EMovieSceneTransformChannel::RotationX);
+			CommonData[3].Color = FCommonChannelData::RedChannelColor;
+			CommonData[3].SortOrder = 3;
+
+			CommonData[4].SetIdentifiers("Rotation.Y", NSLOCTEXT("MovieSceneTransformSection", "RotationY", "Pitch"), RotationGroup);
+			CommonData[4].bEnabled = EnumHasAllFlags(Mask, EMovieSceneTransformChannel::RotationY);
+			CommonData[4].Color = FCommonChannelData::GreenChannelColor;
+			CommonData[4].SortOrder = 4;
+
+			CommonData[5].SetIdentifiers("Rotation.Z", NSLOCTEXT("MovieSceneTransformSection", "RotationZ", "Yaw"), RotationGroup);
+			CommonData[5].bEnabled = EnumHasAllFlags(Mask, EMovieSceneTransformChannel::RotationZ);
+			CommonData[5].Color = FCommonChannelData::BlueChannelColor;
+			CommonData[5].SortOrder = 5;
+		}
+		{
+			CommonData[6].SetIdentifiers("Scale.X", FCommonChannelData::ChannelX, ScaleGroup);
+			CommonData[6].bEnabled = EnumHasAllFlags(Mask, EMovieSceneTransformChannel::ScaleX);
+			CommonData[6].Color = FCommonChannelData::RedChannelColor;
+			CommonData[6].SortOrder = 6;
+
+			CommonData[7].SetIdentifiers("Scale.Y", FCommonChannelData::ChannelY, ScaleGroup);
+			CommonData[7].bEnabled = EnumHasAllFlags(Mask, EMovieSceneTransformChannel::ScaleY);
+			CommonData[7].Color = FCommonChannelData::GreenChannelColor;
+			CommonData[7].SortOrder = 7;
+
+			CommonData[8].SetIdentifiers("Scale.Z", FCommonChannelData::ChannelZ, ScaleGroup);
+			CommonData[8].bEnabled = EnumHasAllFlags(Mask, EMovieSceneTransformChannel::ScaleZ);
+			CommonData[8].Color = FCommonChannelData::BlueChannelColor;
+			CommonData[8].SortOrder = 8;
+		}
+		{
+			CommonData[9].SetIdentifiers("Weight", NSLOCTEXT("MovieSceneTransformSection", "Weight", "Weight"));
+			CommonData[9].bEnabled = EnumHasAllFlags(Mask, EMovieSceneTransformChannel::Weight);
 		}
-	}
-}
-
-
-/* FMovieScene3DRotationKeyStruct interface
- *****************************************************************************/
 
-void FMovieScene3DRotationKeyStruct::PropagateChanges(const FPropertyChangedEvent& ChangeEvent)
-{
-	if(RotationKeys[0] != nullptr)
-	{
-		RotationKeys[0]->Value = Rotation.Roll;
-		RotationKeys[0]->Time = Time;
+		ExternalValues[0].OnGetExternalValue = ExtractTranslationX;
+		ExternalValues[1].OnGetExternalValue = ExtractTranslationY;
+		ExternalValues[2].OnGetExternalValue = ExtractTranslationZ;
+		ExternalValues[3].OnGetExternalValue = ExtractRotationX;
+		ExternalValues[4].OnGetExternalValue = ExtractRotationY;
+		ExternalValues[5].OnGetExternalValue = ExtractRotationZ;
+		ExternalValues[6].OnGetExternalValue = ExtractScaleX;
+		ExternalValues[7].OnGetExternalValue = ExtractScaleY;
+		ExternalValues[8].OnGetExternalValue = ExtractScaleZ;
 	}
-	if(RotationKeys[1] != nullptr)
-	{	
-		RotationKeys[1]->Value = Rotation.Pitch;
-		RotationKeys[1]->Time = Time;
-	}
-	if(RotationKeys[2] != nullptr)
-	{
-		RotationKeys[2]->Value = Rotation.Yaw;
-		RotationKeys[2]->Time = Time;
-	}
-}
-
 
-/* FMovieScene3DScaleKeyStruct interface
- *****************************************************************************/
-
-void FMovieScene3DScaleKeyStruct::PropagateChanges(const FPropertyChangedEvent& ChangeEvent)
-{
-	for (int32 Index = 0; Index <= 2; ++Index)
+	static TOptional<FTransform> GetTransform(UObject& InObject, FTrackInstancePropertyBindings* Bindings)
 	{
-		if(ScaleKeys[Index] != nullptr)
+		TOptional<FTransform> Transform = Bindings ? Bindings->GetOptionalValue<FTransform>(InObject) : TOptional<FTransform>();
+		if (Transform.IsSet())
 		{
-			ScaleKeys[Index]->Value = Scale[Index];
-			ScaleKeys[Index]->Time = Time;
+			return Transform;
 		}
-	}
-}
 
+		if (AActor* Actor = Cast<AActor>(&InObject))
+		{
+			if (USceneComponent* RootComponent = Actor->GetRootComponent())
+			{
+				Transform = RootComponent->GetRelativeTransform();
+			}
+		}
 
-/* FMovieScene3DTransformKeyStruct interface
- *****************************************************************************/
+		return Transform;
+	}
 
-void FMovieScene3DTransformKeyStruct::PropagateChanges(const FPropertyChangedEvent& ChangeEvent)
-{
-	for (int32 Index = 0; Index <= 2; ++Index)
+	static TOptional<FRotator> GetRotator(UObject& InObject, FTrackInstancePropertyBindings* Bindings)
 	{
-		if(LocationKeys[Index] != nullptr)
+		TOptional<FRotator> Rotator = Bindings ? Bindings->GetOptionalValue<FRotator>(InObject) : TOptional<FRotator>();
+		if (Rotator.IsSet())
 		{
-			LocationKeys[Index]->Value = Location[Index];
-			LocationKeys[Index]->Time = Time;
+			return Rotator;
 		}
-		if(ScaleKeys[Index] != nullptr)
+
+		if (AActor* Actor = Cast<AActor>(&InObject))
 		{
-			ScaleKeys[Index]->Value = Scale[Index];
-			ScaleKeys[Index]->Time = Time;
+			if (USceneComponent* RootComponent = Actor->GetRootComponent())
+			{
+				Rotator = RootComponent->RelativeRotation;
+			}
 		}
+
+		return Rotator;
 	}
 
-	if(RotationKeys[0] != nullptr)
+	static TOptional<float> ExtractTranslationX(UObject& InObject, FTrackInstancePropertyBindings* Bindings)
 	{
-		RotationKeys[0]->Value = Rotation.Roll;
-		RotationKeys[0]->Time = Time;
+		TOptional<FTransform> Transform = GetTransform(InObject, Bindings);
+		return Transform.IsSet() ? Transform->GetTranslation().X : TOptional<float>();
 	}
-	if(RotationKeys[1] != nullptr)
-	{	
-		RotationKeys[1]->Value = Rotation.Pitch;
-		RotationKeys[1]->Time = Time;
+	static TOptional<float> ExtractTranslationY(UObject& InObject, FTrackInstancePropertyBindings* Bindings)
+	{
+		TOptional<FTransform> Transform = GetTransform(InObject, Bindings);
+		return Transform.IsSet() ? Transform->GetTranslation().Y : TOptional<float>();
 	}
-	if(RotationKeys[2] != nullptr)
+	static TOptional<float> ExtractTranslationZ(UObject& InObject, FTrackInstancePropertyBindings* Bindings)
 	{
-		RotationKeys[2]->Value = Rotation.Yaw;
-		RotationKeys[2]->Time = Time;
+		TOptional<FTransform> Transform = GetTransform(InObject, Bindings);
+		return Transform.IsSet() ? Transform->GetTranslation().Z : TOptional<float>();
 	}
-}
 
+	static TOptional<float> ExtractRotationX(UObject& InObject, FTrackInstancePropertyBindings* Bindings)
+	{
+		TOptional<FRotator> Rotator = GetRotator(InObject, Bindings);
+		return Rotator.IsSet() ? Rotator.GetValue().Roll : TOptional<float>();
+	}
+	static TOptional<float> ExtractRotationY(UObject& InObject, FTrackInstancePropertyBindings* Bindings)
+	{
+		TOptional<FRotator> Rotator = GetRotator(InObject, Bindings);
+		return Rotator.IsSet() ? Rotator.GetValue().Pitch : TOptional<float>();
+	}
+	static TOptional<float> ExtractRotationZ(UObject& InObject, FTrackInstancePropertyBindings* Bindings)
+	{
+		TOptional<FRotator> Rotator = GetRotator(InObject, Bindings);
+		return Rotator.IsSet() ? Rotator.GetValue().Yaw : TOptional<float>();
+	}
 
-/* FMovieScene3DTransformKeyStruct interface
- *****************************************************************************/
+	static TOptional<float> ExtractScaleX(UObject& InObject, FTrackInstancePropertyBindings* Bindings)
+	{
+		TOptional<FTransform> Transform = GetTransform(InObject, Bindings);
+		return Transform.IsSet() ? Transform->GetScale3D().X : TOptional<float>();
+	}
+	static TOptional<float> ExtractScaleY(UObject& InObject, FTrackInstancePropertyBindings* Bindings)
+	{
+		TOptional<FTransform> Transform = GetTransform(InObject, Bindings);
+		return Transform.IsSet() ? Transform->GetScale3D().Y : TOptional<float>();
+	}
+	static TOptional<float> ExtractScaleZ(UObject& InObject, FTrackInstancePropertyBindings* Bindings)
+	{
+		TOptional<FTransform> Transform = GetTransform(InObject, Bindings);
+		return Transform.IsSet() ? Transform->GetScale3D().Z : TOptional<float>();
+	}
 
-UMovieScene3DTransformSection::UMovieScene3DTransformSection(const FObjectInitializer& ObjectInitializer)
-	: Super(ObjectInitializer)
-#if WITH_EDITORONLY_DATA
-	, Show3DTrajectory(EShow3DTrajectory::EST_OnlyWhenSelected)
-#endif
-{
-	EvalOptions.EnableAndSetCompletionMode
-		(GetLinkerCustomVersion(FSequencerObjectVersion::GUID) < FSequencerObjectVersion::WhenFinishedDefaultsToRestoreState ? 
-			EMovieSceneCompletionMode::KeepState : 
-			GetLinkerCustomVersion(FSequencerObjectVersion::GUID) < FSequencerObjectVersion::WhenFinishedDefaultsToProjectDefault ? 
-			EMovieSceneCompletionMode::RestoreState : 
-			EMovieSceneCompletionMode::ProjectDefault);
+	FMovieSceneChannelEditorData    CommonData[10];
+	TMovieSceneExternalValue<float> ExternalValues[10];
+};
 
-	TransformMask = EMovieSceneTransformChannel::AllTransform;
-	BlendType = EMovieSceneBlendType::Absolute;
-}
+#endif // WITH_EDITOR
 
 
-/* UMovieScene3DTransformSection interface
- *****************************************************************************/
-void UMovieScene3DTransformSection::EvalTranslation(float Time, FVector& InOutTranslation) const
-{
-	InOutTranslation.X = Translation[0].Eval(Time, InOutTranslation.X);
-	InOutTranslation.Y = Translation[1].Eval(Time, InOutTranslation.Y);
-	InOutTranslation.Z = Translation[2].Eval(Time, InOutTranslation.Z);
-}
 
+/* FMovieScene3DLocationKeyStruct interface
+ *****************************************************************************/
 
-void UMovieScene3DTransformSection::EvalRotation(float Time, FRotator& InOutRotation) const
+void FMovieScene3DLocationKeyStruct::PropagateChanges(const FPropertyChangedEvent& ChangeEvent)
 {
-	InOutRotation.Roll = Rotation[0].Eval(Time, InOutRotation.Roll);
-	InOutRotation.Pitch = Rotation[1].Eval(Time, InOutRotation.Pitch);
-	InOutRotation.Yaw = Rotation[2].Eval(Time, InOutRotation.Yaw);
+	KeyStructInterop.Apply(Time);
 }
 
 
-void UMovieScene3DTransformSection::EvalScale(float Time, FVector& InOutScale) const
-{
-	InOutScale.X = Scale[0].Eval(Time, InOutScale.X);
-	InOutScale.Y = Scale[1].Eval(Time, InOutScale.Y);
-	InOutScale.Z = Scale[2].Eval(Time, InOutScale.Z);
-}
-
+/* FMovieScene3DRotationKeyStruct interface
+ *****************************************************************************/
 
-/**
- * Chooses an appropriate curve from an axis and a set of curves
- */
-template<typename T>
-static T* ChooseCurve(EAxis::Type Axis, T* Curves)
+void FMovieScene3DRotationKeyStruct::PropagateChanges(const FPropertyChangedEvent& ChangeEvent)
 {
-	switch(Axis)
-	{
-	case EAxis::X:
-		return &Curves[0];
-
-	case EAxis::Y:
-		return &Curves[1];
-
-	case EAxis::Z:
-		return &Curves[2];
-
-	default:
-		check(false);
-		return nullptr;
-	}
+	KeyStructInterop.Apply(Time);
 }
 
 
-FRichCurve& UMovieScene3DTransformSection::GetTranslationCurve(EAxis::Type Axis) 
-{
-	return *ChooseCurve(Axis, Translation);
-}
-
+/* FMovieScene3DScaleKeyStruct interface
+ *****************************************************************************/
 
-const FRichCurve& UMovieScene3DTransformSection::GetTranslationCurve(EAxis::Type Axis) const
+void FMovieScene3DScaleKeyStruct::PropagateChanges(const FPropertyChangedEvent& ChangeEvent)
 {
-	return *ChooseCurve(Axis, Translation);
+	KeyStructInterop.Apply(Time);
 }
 
 
-FRichCurve& UMovieScene3DTransformSection::GetRotationCurve(EAxis::Type Axis)
-{
-	return *ChooseCurve(Axis, Rotation);
-}
-
+/* FMovieScene3DTransformKeyStruct interface
+ *****************************************************************************/
 
-const FRichCurve& UMovieScene3DTransformSection::GetRotationCurve(EAxis::Type Axis) const
+void FMovieScene3DTransformKeyStruct::PropagateChanges(const FPropertyChangedEvent& ChangeEvent)
 {
-	return *ChooseCurve(Axis, Rotation);
+	KeyStructInterop.Apply(Time);
 }
 
 
-FRichCurve& UMovieScene3DTransformSection::GetScaleCurve(EAxis::Type Axis)
-{
-	return *ChooseCurve(Axis, Scale);
-}
-
+/* UMovieScene3DTransformSection interface
+ *****************************************************************************/
 
-const FRichCurve& UMovieScene3DTransformSection::GetScaleCurve(EAxis::Type Axis) const
+UMovieScene3DTransformSection::UMovieScene3DTransformSection(const FObjectInitializer& ObjectInitializer)
+	: Super(ObjectInitializer)
+#if WITH_EDITORONLY_DATA
+	, Show3DTrajectory(EShow3DTrajectory::EST_OnlyWhenSelected)
+#endif
 {
-	return *ChooseCurve(Axis, Scale);
-}
+	EvalOptions.EnableAndSetCompletionMode
+		(GetLinkerCustomVersion(FSequencerObjectVersion::GUID) < FSequencerObjectVersion::WhenFinishedDefaultsToRestoreState ? 
+			EMovieSceneCompletionMode::KeepState : 
+			GetLinkerCustomVersion(FSequencerObjectVersion::GUID) < FSequencerObjectVersion::WhenFinishedDefaultsToProjectDefault ? 
+			EMovieSceneCompletionMode::RestoreState : 
+			EMovieSceneCompletionMode::ProjectDefault);
 
-FRichCurve& UMovieScene3DTransformSection::GetManualWeightCurve()
-{
-	return ManualWeight;
-}
+	ProxyChannels = EMovieSceneTransformChannel::None;
+	TransformMask = EMovieSceneTransformChannel::AllTransform;
+	BlendType = EMovieSceneBlendType::Absolute;
+	bSupportsInfiniteRange = true;
 
-const FRichCurve& UMovieScene3DTransformSection::GetManualWeightCurve() const
-{
-	return ManualWeight;
+	UpdateChannelProxy();
 }
 
-
-/* UMovieSceneSection interface
- *****************************************************************************/
-
-void UMovieScene3DTransformSection::MoveSection(float DeltaTime, TSet<FKeyHandle>& KeyHandles)
+void UMovieScene3DTransformSection::Serialize(FArchive& Ar)
 {
-	Super::MoveSection(DeltaTime, KeyHandles);
+	Super::Serialize(Ar);
 
-	// Move all the curves in this section
-	for (int32 Axis = 0; Axis < 3; ++Axis)
+	if (Ar.IsLoading())
 	{
-		Translation[Axis].ShiftCurve(DeltaTime, KeyHandles);
-		Rotation[Axis].ShiftCurve(DeltaTime, KeyHandles);
-		Scale[Axis].ShiftCurve(DeltaTime, KeyHandles);
+		UpdateChannelProxy();
 	}
-	ManualWeight.ShiftCurve(DeltaTime, KeyHandles);
 }
 
-
-void UMovieScene3DTransformSection::DilateSection(float DilationFactor, float Origin, TSet<FKeyHandle>& KeyHandles)
+FMovieSceneTransformMask UMovieScene3DTransformSection::GetMask() const
 {
-	Super::DilateSection(DilationFactor, Origin, KeyHandles);
-
-	for (int32 Axis = 0; Axis < 3; ++Axis)
-	{
-		Translation[Axis].ScaleCurve(Origin, DilationFactor, KeyHandles);
-		Rotation[Axis].ScaleCurve(Origin, DilationFactor, KeyHandles);
-		Scale[Axis].ScaleCurve(Origin, DilationFactor, KeyHandles);
-	}
-	ManualWeight.ScaleCurve(Origin, DilationFactor, KeyHandles);
+	return TransformMask;
 }
 
+void UMovieScene3DTransformSection::SetMask(FMovieSceneTransformMask NewMask)
+{
+	TransformMask = NewMask;
+	UpdateChannelProxy();
+}
 
-void UMovieScene3DTransformSection::GetKeyHandles(TSet<FKeyHandle>& OutKeyHandles, TRange<float> TimeRange) const
+void UMovieScene3DTransformSection::UpdateChannelProxy()
 {
-	if (!TimeRange.Overlaps(GetRange()))
+	if (ProxyChannels == TransformMask.GetChannels())
 	{
 		return;
 	}
 
-	for (int32 Axis = 0; Axis < 3; ++Axis)
-	{
-		for (auto It(Translation[Axis].GetKeyHandleIterator()); It; ++It)
-		{
-			float Time = Translation[Axis].GetKeyTime(It.Key());
-			if (TimeRange.Contains(Time))
-			{
-				OutKeyHandles.Add(It.Key());
-			}
-		}
+	ProxyChannels = TransformMask.GetChannels();
 
-		for (auto It(Rotation[Axis].GetKeyHandleIterator()); It; ++It)
-		{
-			float Time = Rotation[Axis].GetKeyTime(It.Key());
-			if (TimeRange.Contains(Time))
-			{
-				OutKeyHandles.Add(It.Key());
-			}
-		}
+	FMovieSceneChannelData Channels;
 
-		for (auto It(Scale[Axis].GetKeyHandleIterator()); It; ++It)
-		{
-			float Time = Scale[Axis].GetKeyTime(It.Key());
-			if (TimeRange.Contains(Time))
-			{
-				OutKeyHandles.Add(It.Key());
-			}
-		}
-	}
-	for (auto It(ManualWeight.GetKeyHandleIterator()); It; ++It)
-	{
-		float Time = ManualWeight.GetKeyTime(It.Key());
-		if (TimeRange.Contains(Time))
-		{
-			OutKeyHandles.Add(It.Key());
-		}
-	}
-}
+#if WITH_EDITOR
 
+	F3DTransformChannelEditorData EditorData(TransformMask.GetChannels());
 
-TSharedPtr<FStructOnScope> UMovieScene3DTransformSection::GetKeyStruct(const TArray<FKeyHandle>& KeyHandles)
-{
-	FRichCurveKey* TranslationKeys[3];
-	FRichCurveKey* RotationKeys[3];
-	FRichCurveKey* ScaleKeys[3];
+	Channels.Add(Translation[0], EditorData.CommonData[0], EditorData.ExternalValues[0]);
+	Channels.Add(Translation[1], EditorData.CommonData[1], EditorData.ExternalValues[1]);
+	Channels.Add(Translation[2], EditorData.CommonData[2], EditorData.ExternalValues[2]);
+	Channels.Add(Rotation[0],    EditorData.CommonData[3], EditorData.ExternalValues[3]);
+	Channels.Add(Rotation[1],    EditorData.CommonData[4], EditorData.ExternalValues[4]);
+	Channels.Add(Rotation[2],    EditorData.CommonData[5], EditorData.ExternalValues[5]);
+	Channels.Add(Scale[0],       EditorData.CommonData[6], EditorData.ExternalValues[6]);
+	Channels.Add(Scale[1],       EditorData.CommonData[7], EditorData.ExternalValues[7]);
+	Channels.Add(Scale[2],       EditorData.CommonData[8], EditorData.ExternalValues[8]);
+	Channels.Add(ManualWeight,   EditorData.CommonData[9], EditorData.ExternalValues[9]);
 
-	bool bHasTranslationKeys = false;
-	bool bHasRotationKeys = false;
-	bool bHasScaleKeys = false;
+#else
 
-	for (int32 Index = 0; Index <= 2; ++Index)
-	{
-		TranslationKeys[Index] = Translation[Index].GetFirstMatchingKey(KeyHandles);
-		if(TranslationKeys[Index] != nullptr)
-		{
-			bHasTranslationKeys = true;
-		}
-		RotationKeys[Index] = Rotation[Index].GetFirstMatchingKey(KeyHandles);
-		if(RotationKeys[Index] != nullptr)
-		{
-			bHasRotationKeys = true;
-		}
-		ScaleKeys[Index] = Scale[Index].GetFirstMatchingKey(KeyHandles);
-		if(ScaleKeys[Index] != nullptr)
-		{
-			bHasScaleKeys = true;
-		}
-	}
+	Channels.Add(Translation[0]);
+	Channels.Add(Translation[1]);
+	Channels.Add(Translation[2]);
+	Channels.Add(Rotation[0]);
+	Channels.Add(Rotation[1]);
+	Channels.Add(Rotation[2]);
+	Channels.Add(Scale[0]);
+	Channels.Add(Scale[1]);
+	Channels.Add(Scale[2]);
+	Channels.Add(ManualWeight);
 
-	int32 KeyTypeCount = 0;
-	if(bHasTranslationKeys)
-	{
-		KeyTypeCount++;
-	}
-	if(bHasRotationKeys)
-	{
-		KeyTypeCount++;
-	}
-	if(bHasScaleKeys)
-	{
-		KeyTypeCount++;
-	}
+#endif
+
+	ChannelProxy = MakeShared<FMovieSceneChannelProxy>(MoveTemp(Channels));
+}
+
+/* UMovieSceneSection interface
+ *****************************************************************************/
+
+
+TSharedPtr<FStructOnScope> UMovieScene3DTransformSection::GetKeyStruct(TArrayView<const FKeyHandle> KeyHandles)
+{
+	FVector  StartingLocation;
+	FRotator StartingRotation;
+	FVector  StartingScale;
+
+	TArrayView<FMovieSceneFloatChannel*> FloatChannels = ChannelProxy->GetChannels<FMovieSceneFloatChannel>();
+
+	TOptional<TTuple<FKeyHandle, FFrameNumber>> LocationKeys[3] = {
+		FMovieSceneChannelValueHelper::FindFirstKey(FloatChannels[0], KeyHandles),
+		FMovieSceneChannelValueHelper::FindFirstKey(FloatChannels[1], KeyHandles),
+		FMovieSceneChannelValueHelper::FindFirstKey(FloatChannels[2], KeyHandles)
+	};
+
+	TOptional<TTuple<FKeyHandle, FFrameNumber>> RotationKeys[3] = {
+		FMovieSceneChannelValueHelper::FindFirstKey(FloatChannels[3], KeyHandles),
+		FMovieSceneChannelValueHelper::FindFirstKey(FloatChannels[4], KeyHandles),
+		FMovieSceneChannelValueHelper::FindFirstKey(FloatChannels[5], KeyHandles)
+	};
+
+	TOptional<TTuple<FKeyHandle, FFrameNumber>> ScaleKeys[3] = {
+		FMovieSceneChannelValueHelper::FindFirstKey(FloatChannels[6], KeyHandles),
+		FMovieSceneChannelValueHelper::FindFirstKey(FloatChannels[7], KeyHandles),
+		FMovieSceneChannelValueHelper::FindFirstKey(FloatChannels[8], KeyHandles)
+	};
+
+	const int32 AnyLocationKeys = Algo::AnyOf(LocationKeys);
+	const int32 AnyRotationKeys = Algo::AnyOf(RotationKeys);
+	const int32 AnyScaleKeys =    Algo::AnyOf(ScaleKeys);
 
 	// do we have multiple keys on multiple parts of the transform?
-	if (KeyTypeCount > 1)
+	if (AnyLocationKeys + AnyRotationKeys + AnyScaleKeys > 1)
 	{
 		TSharedRef<FStructOnScope> KeyStruct = MakeShareable(new FStructOnScope(FMovieScene3DTransformKeyStruct::StaticStruct()));
 		auto Struct = (FMovieScene3DTransformKeyStruct*)KeyStruct->GetStructMemory();
-		{
-			for (int32 Index = 0; Index <= 2; ++Index)
-			{
-				if(TranslationKeys[Index] != nullptr)
-				{
-					Struct->LocationKeys[Index] = TranslationKeys[Index];
-					Struct->Location[Index] = TranslationKeys[Index]->Value;
-					Struct->Time = TranslationKeys[Index]->Time;
-				}
-
-				if(RotationKeys[Index] != nullptr)
-				{
-					Struct->RotationKeys[Index] = RotationKeys[Index];
-					Struct->Time = RotationKeys[Index]->Time;
-				}
-				
-				if(ScaleKeys[Index] != nullptr)
-				{
-					Struct->ScaleKeys[Index] = ScaleKeys[Index];
-					Struct->Scale[Index] = ScaleKeys[Index]->Value;
-					Struct->Time = ScaleKeys[Index]->Time;
-				}
-			}
 
-			if(RotationKeys[0] != nullptr)
-			{
-				Struct->Rotation.Roll = RotationKeys[0]->Value;
-			}
-			if(RotationKeys[1] != nullptr)
-			{
-				Struct->Rotation.Pitch = RotationKeys[1]->Value;
-			}
-			if(RotationKeys[2] != nullptr)
-			{
-				Struct->Rotation.Yaw = RotationKeys[2]->Value;
-			}
-		}
+		Struct->KeyStructInterop.Add(FMovieSceneChannelValueHelper(ChannelProxy->MakeHandle(FloatChannels[0]), &Struct->Location.X,     LocationKeys[0]));
+		Struct->KeyStructInterop.Add(FMovieSceneChannelValueHelper(ChannelProxy->MakeHandle(FloatChannels[1]), &Struct->Location.Y,     LocationKeys[1]));
+		Struct->KeyStructInterop.Add(FMovieSceneChannelValueHelper(ChannelProxy->MakeHandle(FloatChannels[2]), &Struct->Location.Z,     LocationKeys[2]));
+
+		Struct->KeyStructInterop.Add(FMovieSceneChannelValueHelper(ChannelProxy->MakeHandle(FloatChannels[3]), &Struct->Rotation.Roll,  RotationKeys[0]));
+		Struct->KeyStructInterop.Add(FMovieSceneChannelValueHelper(ChannelProxy->MakeHandle(FloatChannels[4]), &Struct->Rotation.Pitch, RotationKeys[1]));
+		Struct->KeyStructInterop.Add(FMovieSceneChannelValueHelper(ChannelProxy->MakeHandle(FloatChannels[5]), &Struct->Rotation.Yaw,   RotationKeys[2]));
 
+		Struct->KeyStructInterop.Add(FMovieSceneChannelValueHelper(ChannelProxy->MakeHandle(FloatChannels[6]), &Struct->Scale.X,        ScaleKeys[0]));
+		Struct->KeyStructInterop.Add(FMovieSceneChannelValueHelper(ChannelProxy->MakeHandle(FloatChannels[7]), &Struct->Scale.Y,        ScaleKeys[1]));
+		Struct->KeyStructInterop.Add(FMovieSceneChannelValueHelper(ChannelProxy->MakeHandle(FloatChannels[8]), &Struct->Scale.Z,        ScaleKeys[2]));
+
+		Struct->KeyStructInterop.SetStartingValues();
+		Struct->Time = Struct->KeyStructInterop.GetUnifiedKeyTime().Get(0);
 		return KeyStruct;
 	}
-	
-	if (bHasTranslationKeys)
+
+	if (AnyLocationKeys > 1)
 	{
 		TSharedRef<FStructOnScope> KeyStruct = MakeShareable(new FStructOnScope(FMovieScene3DLocationKeyStruct::StaticStruct()));
 		auto Struct = (FMovieScene3DLocationKeyStruct*)KeyStruct->GetStructMemory();
-		{
-			for (int32 Index = 0; Index <= 2; ++Index)
-			{
-				if(TranslationKeys[Index] != nullptr)
-				{
-					Struct->LocationKeys[Index] = TranslationKeys[Index];
-					Struct->Location[Index] = TranslationKeys[Index]->Value;
-					Struct->Time = TranslationKeys[Index]->Time;
-				}
-			}
-		}
 
+		Struct->KeyStructInterop.Add(FMovieSceneChannelValueHelper(ChannelProxy->MakeHandle(FloatChannels[0]), &Struct->Location.X,     LocationKeys[0]));
+		Struct->KeyStructInterop.Add(FMovieSceneChannelValueHelper(ChannelProxy->MakeHandle(FloatChannels[1]), &Struct->Location.Y,     LocationKeys[1]));
+		Struct->KeyStructInterop.Add(FMovieSceneChannelValueHelper(ChannelProxy->MakeHandle(FloatChannels[2]), &Struct->Location.Z,     LocationKeys[2]));
+
+		Struct->KeyStructInterop.SetStartingValues();
+		Struct->Time = Struct->KeyStructInterop.GetUnifiedKeyTime().Get(0);
 		return KeyStruct;
 	}
-	
-	if (bHasRotationKeys)
+
+	if (AnyRotationKeys)
 	{
 		TSharedRef<FStructOnScope> KeyStruct = MakeShareable(new FStructOnScope(FMovieScene3DRotationKeyStruct::StaticStruct()));
 		auto Struct = (FMovieScene3DRotationKeyStruct*)KeyStruct->GetStructMemory();
-		{
-			for (int32 Index = 0; Index <= 2; ++Index)
-			{
-				if(RotationKeys[Index] != nullptr)
-				{
-					Struct->RotationKeys[Index] = RotationKeys[Index];
-					Struct->Time = RotationKeys[Index]->Time;
-				}
-			}
 
-			if (Struct->RotationKeys[0] != nullptr)
-			{
-				Struct->Rotation.Roll = Struct->RotationKeys[0]->Value;
-			}
-
-			if (Struct->RotationKeys[1] != nullptr)
-			{
-				Struct->Rotation.Pitch = Struct->RotationKeys[1]->Value;
-			}
-			
-			if (Struct->RotationKeys[2] != nullptr)
-			{
-				Struct->Rotation.Yaw = Struct->RotationKeys[2]->Value;
-			}
-		}
+		Struct->KeyStructInterop.Add(FMovieSceneChannelValueHelper(ChannelProxy->MakeHandle(FloatChannels[3]), &Struct->Rotation.Roll,  RotationKeys[0]));
+		Struct->KeyStructInterop.Add(FMovieSceneChannelValueHelper(ChannelProxy->MakeHandle(FloatChannels[4]), &Struct->Rotation.Pitch, RotationKeys[1]));
+		Struct->KeyStructInterop.Add(FMovieSceneChannelValueHelper(ChannelProxy->MakeHandle(FloatChannels[5]), &Struct->Rotation.Yaw,   RotationKeys[2]));
 
+		Struct->KeyStructInterop.SetStartingValues();
+		Struct->Time = Struct->KeyStructInterop.GetUnifiedKeyTime().Get(0);
 		return KeyStruct;
 	}
-	
-	if (bHasScaleKeys)
+
+	if (AnyScaleKeys)
 	{
 		TSharedRef<FStructOnScope> KeyStruct = MakeShareable(new FStructOnScope(FMovieScene3DScaleKeyStruct::StaticStruct()));
 		auto Struct = (FMovieScene3DScaleKeyStruct*)KeyStruct->GetStructMemory();
-		{
-			for (int32 Index = 0; Index <= 2; ++Index)
-			{
-				if(ScaleKeys[Index] != nullptr)
-				{
-					Struct->ScaleKeys[Index] = ScaleKeys[Index];
-					Struct->Scale[Index] = ScaleKeys[Index]->Value;
-					Struct->Time = ScaleKeys[Index]->Time;
-				}
-			}
-		}
-
-		return KeyStruct;
-	}
-
-	return nullptr;
-}
-
-
-TOptional<float> UMovieScene3DTransformSection::GetKeyTime( FKeyHandle KeyHandle ) const
-{
-	// Translation
-	if ( Translation[0].IsKeyHandleValid( KeyHandle ) )
-	{
-		return TOptional<float>( Translation[0].GetKeyTime( KeyHandle ) );
-	}
-	if ( Translation[1].IsKeyHandleValid( KeyHandle ) )
-	{
-		return TOptional<float>( Translation[1].GetKeyTime( KeyHandle ) );
-	}
-	if ( Translation[2].IsKeyHandleValid( KeyHandle ) )
-	{
-		return TOptional<float>( Translation[2].GetKeyTime( KeyHandle ) );
-	}
-	// Rotation
-	if ( Rotation[0].IsKeyHandleValid( KeyHandle ) )
-	{
-		return TOptional<float>( Rotation[0].GetKeyTime( KeyHandle ) );
-	}
-	if ( Rotation[1].IsKeyHandleValid( KeyHandle ) )
-	{
-		return TOptional<float>( Rotation[1].GetKeyTime( KeyHandle ) );
-	}
-	if ( Rotation[2].IsKeyHandleValid( KeyHandle ) )
-	{
-		return TOptional<float>( Rotation[2].GetKeyTime( KeyHandle ) );
-	}
-	// Scale
-	if ( Scale[0].IsKeyHandleValid( KeyHandle ) )
-	{
-		return TOptional<float>( Scale[0].GetKeyTime( KeyHandle ) );
-	}
-	if ( Scale[1].IsKeyHandleValid( KeyHandle ) )
-	{
-		return TOptional<float>( Scale[1].GetKeyTime( KeyHandle ) );
-	}
-	if ( Scale[2].IsKeyHandleValid( KeyHandle ) )
-	{
-		return TOptional<float>( Scale[2].GetKeyTime( KeyHandle ) );
-	}
-	if ( ManualWeight.IsKeyHandleValid( KeyHandle ) )
-	{
-		return TOptional<float>( ManualWeight.GetKeyTime( KeyHandle ) );
-	}
-	return TOptional<float>();
-}
 
+		Struct->KeyStructInterop.Add(FMovieSceneChannelValueHelper(ChannelProxy->MakeHandle(FloatChannels[6]), &Struct->Scale.X,        ScaleKeys[0]));
+		Struct->KeyStructInterop.Add(FMovieSceneChannelValueHelper(ChannelProxy->MakeHandle(FloatChannels[7]), &Struct->Scale.Y,        ScaleKeys[1]));
+		Struct->KeyStructInterop.Add(FMovieSceneChannelValueHelper(ChannelProxy->MakeHandle(FloatChannels[8]), &Struct->Scale.Z,        ScaleKeys[2]));
 
-void UMovieScene3DTransformSection::SetKeyTime( FKeyHandle KeyHandle, float Time )
-{
-	// Translation
-	if ( Translation[0].IsKeyHandleValid( KeyHandle ) )
-	{
-		Translation[0].SetKeyTime( KeyHandle, Time );
-	}
-	else if ( Translation[1].IsKeyHandleValid( KeyHandle ) )
-	{
-		Translation[1].SetKeyTime( KeyHandle, Time );
-	}
-	else if ( Translation[2].IsKeyHandleValid( KeyHandle ) )
-	{
-		Translation[2].SetKeyTime( KeyHandle, Time );
-	}
-	// Rotation
-	else if ( Rotation[0].IsKeyHandleValid( KeyHandle ) )
-	{
-		Rotation[0].SetKeyTime( KeyHandle, Time );
-	}
-	else if ( Rotation[1].IsKeyHandleValid( KeyHandle ) )
-	{
-		Rotation[1].SetKeyTime( KeyHandle, Time );
-	}
-	else if ( Rotation[2].IsKeyHandleValid( KeyHandle ) )
-	{
-		Rotation[2].SetKeyTime( KeyHandle, Time );
-	}
-	// Scale
-	else if ( Scale[0].IsKeyHandleValid( KeyHandle ) )
-	{
-		Scale[0].SetKeyTime( KeyHandle, Time );
-	}
-	else if ( Scale[1].IsKeyHandleValid( KeyHandle ) )
-	{
-		Scale[1].SetKeyTime( KeyHandle, Time );
-	}
-	else if ( Scale[2].IsKeyHandleValid( KeyHandle ) )
-	{
-		Scale[2].SetKeyTime( KeyHandle, Time );
-	}
-	else if ( ManualWeight.IsKeyHandleValid( KeyHandle ) )
-	{
-		ManualWeight.SetKeyTime( KeyHandle, Time );
-	}
-}
-
-
-/* UMovieSceneSection interface
- *****************************************************************************/
-
-// This function uses a template to avoid duplicating this for const and non-const versions
-template<typename CurveType>
-CurveType* GetCurveForChannelAndAxis(EKey3DTransformChannel::Type Channel, EAxis::Type Axis,
-	CurveType* TranslationCurves, CurveType* RotationCurves, CurveType* ScaleCurves)
-{
-	CurveType* ChannelCurves = nullptr;
-	switch (Channel)
-	{
-	case EKey3DTransformChannel::Translation:
-		ChannelCurves = TranslationCurves;
-		break;
-	case EKey3DTransformChannel::Rotation:
-		ChannelCurves = RotationCurves;
-		break;
-	case EKey3DTransformChannel::Scale:
-		ChannelCurves = ScaleCurves;
-		break;
-	}
-
-	if (ChannelCurves != nullptr)
-	{
-		switch (Axis)
-		{
-		case EAxis::X:
-			return &ChannelCurves[0];
-		case EAxis::Y:
-			return &ChannelCurves[1];
-		case EAxis::Z:
-			return &ChannelCurves[2];
-		}
+		Struct->KeyStructInterop.SetStartingValues();
+		Struct->Time = Struct->KeyStructInterop.GetUnifiedKeyTime().Get(0);
+		return KeyStruct;
 	}
 
-	checkf(false, TEXT("Invalid channel and axis combination."));
 	return nullptr;
 }
 
-
-bool UMovieScene3DTransformSection::NewKeyIsNewData(float Time, const FTransformKey& TransformKey) const
-{
-	const FRichCurve* KeyCurve = GetCurveForChannelAndAxis(TransformKey.Channel, TransformKey.Axis, Translation, Rotation, Scale);
-	return FMath::IsNearlyEqual(KeyCurve->Eval(Time), TransformKey.Value) == false;
-}
-
-
-bool UMovieScene3DTransformSection::HasKeys(const FTransformKey& TransformKey) const
-{
-	const FRichCurve* KeyCurve = GetCurveForChannelAndAxis(TransformKey.Channel, TransformKey.Axis, Translation, Rotation, Scale);
-	return KeyCurve->GetNumKeys() != 0;
-}
-
-
-void UMovieScene3DTransformSection::AddKey(float Time, const FTransformKey& TransformKey, EMovieSceneKeyInterpolation KeyInterpolation)
-{
-	FRichCurve* KeyCurve = GetCurveForChannelAndAxis(TransformKey.Channel, TransformKey.Axis, Translation, Rotation, Scale);
-	
-	bool bUnwindRotation = TransformKey.Channel == EKey3DTransformChannel::Rotation;
-	AddKeyToCurve(*KeyCurve, Time, TransformKey.Value, KeyInterpolation, bUnwindRotation);
-}
-
-
-void UMovieScene3DTransformSection::SetDefault(const FTransformKey& TransformKey)
-{
-	FRichCurve* KeyCurve = GetCurveForChannelAndAxis(TransformKey.Channel, TransformKey.Axis, Translation, Rotation, Scale);
-	SetCurveDefault(*KeyCurve, TransformKey.Value);
-}
-
-
-void UMovieScene3DTransformSection::ClearDefaults()
-{
-	Translation[0].ClearDefaultValue();
-	Translation[1].ClearDefaultValue();
-	Translation[2].ClearDefaultValue();
-	Rotation[0].ClearDefaultValue();
-	Rotation[1].ClearDefaultValue();
-	Rotation[2].ClearDefaultValue();
-	Scale[0].ClearDefaultValue();
-	Scale[1].ClearDefaultValue();
-	Scale[2].ClearDefaultValue();
-	ManualWeight.ClearDefaultValue();
-}
-
 FMovieSceneEvalTemplatePtr UMovieScene3DTransformSection::GenerateTemplate() const
 {
 	return FMovieSceneComponentTransformSectionTemplate(*this);
diff --git a/Engine/Source/Runtime/MovieSceneTracks/Private/Sections/MovieSceneActorReferenceSection.cpp b/Engine/Source/Runtime/MovieSceneTracks/Private/Sections/MovieSceneActorReferenceSection.cpp
--- a/Engine/Source/Runtime/MovieSceneTracks/Private/Sections/MovieSceneActorReferenceSection.cpp
+++ b/Engine/Source/Runtime/MovieSceneTracks/Private/Sections/MovieSceneActorReferenceSection.cpp
@@ -1,163 +1,66 @@
 // Copyright 1998-2018 Epic Games, Inc. All Rights Reserved.
 
 #include "Sections/MovieSceneActorReferenceSection.h"
+#include "Channels/MovieSceneChannelProxy.h"
 
-
-UMovieSceneActorReferenceSection::UMovieSceneActorReferenceSection( const FObjectInitializer& ObjectInitializer )
-	: Super( ObjectInitializer )
-{ }
-
-
-FGuid UMovieSceneActorReferenceSection::Eval( float Position ) const
-{
-	int32 ActorGuidIndex = ActorGuidIndexCurve.Evaluate( Position );
-	return ActorGuidIndex >= 0 && ActorGuidIndex < ActorGuids.Num() ? ActorGuids[ActorGuidIndex] : FGuid();
-}
-
-
-void UMovieSceneActorReferenceSection::MoveSection( float DeltaPosition, TSet<FKeyHandle>& KeyHandles )
+uint32 FMovieSceneActorReferenceData::GetChannelID()
 {
-	Super::MoveSection( DeltaPosition, KeyHandles );
-
-	// Move the curve
-	ActorGuidIndexCurve.ShiftCurve( DeltaPosition, KeyHandles );
+	static uint32 ID = FMovieSceneChannelEntry::RegisterNewID();
+	return ID;
 }
 
-
-void UMovieSceneActorReferenceSection::DilateSection( float DilationFactor, float Origin, TSet<FKeyHandle>& KeyHandles )
-{	
-	Super::DilateSection(DilationFactor, Origin, KeyHandles);
-	
-	ActorGuidIndexCurve.ScaleCurve( Origin, DilationFactor, KeyHandles );
-}
-
-
-void UMovieSceneActorReferenceSection::GetKeyHandles(TSet<FKeyHandle>& OutKeyHandles, TRange<float> TimeRange) const
+FMovieSceneActorReferenceKey FMovieSceneActorReferenceData::Evaluate(FFrameTime InTime) const
 {
-	if (!TimeRange.Overlaps(GetRange()))
+	if (KeyTimes.Num())
 	{
-		return;
+		const int32 Index = FMath::Max(0, Algo::UpperBound(KeyTimes, InTime.FrameNumber)-1);
+		return KeyValues[Index];
 	}
 
-	for ( auto It( ActorGuidIndexCurve.GetKeyHandleIterator() ); It; ++It )
-	{
-		float Time = ActorGuidIndexCurve.GetKeyTime( It.Key() );
-		if (TimeRange.Contains(Time))
-		{
-			OutKeyHandles.Add(It.Key());
-		}
-	}
+	return DefaultValue;
 }
 
-
-TOptional<float> UMovieSceneActorReferenceSection::GetKeyTime( FKeyHandle KeyHandle ) const
+UMovieSceneActorReferenceSection::UMovieSceneActorReferenceSection( const FObjectInitializer& ObjectInitializer )
+	: Super( ObjectInitializer )
 {
-	if ( ActorGuidIndexCurve.IsKeyHandleValid( KeyHandle ) )
-	{
-		return TOptional<float>( ActorGuidIndexCurve.GetKeyTime( KeyHandle ) );
-	}
-	return TOptional<float>();
-}
+#if WITH_EDITOR
 
+	ChannelProxy = MakeShared<FMovieSceneChannelProxy>(ActorReferenceData, FMovieSceneChannelEditorData());
 
-void UMovieSceneActorReferenceSection::SetKeyTime( FKeyHandle KeyHandle, float Time )
-{
-	if ( ActorGuidIndexCurve.IsKeyHandleValid( KeyHandle ) )
-	{
-		ActorGuidIndexCurve.SetKeyTime( KeyHandle, Time );
-	}
-}
+#else
 
+	ChannelProxy = MakeShared<FMovieSceneChannelProxy>(ActorReferenceData);
 
-void UMovieSceneActorReferenceSection::AddKey( float Time, const FGuid& Value, EMovieSceneKeyInterpolation KeyInterpolation )
-{
-	AddKey(Time, Value);
+#endif
 }
 
-
-FKeyHandle UMovieSceneActorReferenceSection::AddKey(float Time, const FGuid& Value)
+void UMovieSceneActorReferenceSection::PostLoad()
 {
-	if (IsTimeWithinSection(Time))
+	Super::PostLoad();
+
+	if (ActorGuidStrings_DEPRECATED.Num())
 	{
-		if (TryModify())
+		TArray<FGuid> Guids;
+
+		for (const FString& ActorGuidString : ActorGuidStrings_DEPRECATED)
 		{
-			FKeyHandle ExistingKeyHandle = ActorGuidIndexCurve.FindKey(Time);
-			if (ActorGuidIndexCurve.IsKeyHandleValid(ExistingKeyHandle))
-			{
-				int32 ActorGuidIndex = ActorGuidIndexCurve.GetKeyValue(ExistingKeyHandle);
-				if (ActorGuidIndex >= 0 && ActorGuidIndex < ActorGuids.Num())
-				{
-					ActorGuids[ActorGuidIndex] = Value;
-					return ExistingKeyHandle;
-				}
-				else
-				{
-					int32 NewActorGuidIndex = ActorGuids.Add(Value);
-					return ActorGuidIndexCurve.UpdateOrAddKey(Time, NewActorGuidIndex);
-				}
-			}
-			else
-			{
-				int32 NewActorGuidIndex = ActorGuids.Add(Value);
-				return ActorGuidIndexCurve.AddKey(Time, NewActorGuidIndex);
-			}
+			FGuid& ActorGuid = Guids[Guids.Emplace()];
+			FGuid::Parse( ActorGuidString,  ActorGuid);
 		}
-	}
-	return FKeyHandle();
-}
-
-
-bool UMovieSceneActorReferenceSection::NewKeyIsNewData(float Time, const FGuid& Value) const
-{
-	return Eval(Time) != Value;
-}
-
 
-bool UMovieSceneActorReferenceSection::HasKeys( const FGuid& Value ) const
-{
-	return ActorGuidIndexCurve.GetNumKeys() > 0;
-}
-
-
-void UMovieSceneActorReferenceSection::SetDefault( const FGuid& Value )
-{
-	int32 CurrentDefault = ActorGuidIndexCurve.GetDefaultValue();
-	if (!ActorGuids.IsValidIndex(CurrentDefault) || ActorGuids[CurrentDefault] != Value)
-	{
-		if (TryModify())
+		if (Guids.IsValidIndex(ActorGuidIndexCurve_DEPRECATED.GetDefaultValue()))
 		{
-			int32 DefaultIndex = ActorGuids.Add( Value );
-			ActorGuidIndexCurve.SetDefaultValue( DefaultIndex );
+			FMovieSceneObjectBindingID DefaultValue(Guids[ActorGuidIndexCurve_DEPRECATED.GetDefaultValue()], MovieSceneSequenceID::Root, EMovieSceneObjectBindingSpace::Local);
+			ActorReferenceData.SetDefault(DefaultValue);
 		}
-	}
-}
-
 
-void UMovieSceneActorReferenceSection::ClearDefaults()
-{
-	ActorGuidIndexCurve.ClearDefaultValue();
-}
-
-
-void UMovieSceneActorReferenceSection::PreSave(const class ITargetPlatform* TargetPlatform)
-{
-	ActorGuidStrings.Empty();
-	for ( const FGuid& ActorGuid : ActorGuids )
-	{
-		ActorGuidStrings.Add( ActorGuid.ToString() );
-	}
-	Super::PreSave(TargetPlatform);
-}
-
-
-void UMovieSceneActorReferenceSection::PostLoad()
-{
-	Super::PostLoad();
-	ActorGuids.Empty();
-	for ( const FString& ActorGuidString : ActorGuidStrings )
-	{
-		FGuid ActorGuid;
-		FGuid::Parse( ActorGuidString, ActorGuid );
-		ActorGuids.Add( ActorGuid );
+		for (auto It = ActorGuidIndexCurve_DEPRECATED.GetKeyIterator(); It; ++It)
+		{
+			if (ensure(Guids.IsValidIndex(It->Value)))
+			{
+				FMovieSceneObjectBindingID BindingID(Guids[It->Value], MovieSceneSequenceID::Root, EMovieSceneObjectBindingSpace::Local);
+				ActorReferenceData.UpgradeLegacyTime(this, It->Time, BindingID);
+			}
+		}
 	}
 }
diff --git a/Engine/Source/Runtime/MovieSceneTracks/Private/Sections/MovieSceneAudioSection.cpp b/Engine/Source/Runtime/MovieSceneTracks/Private/Sections/MovieSceneAudioSection.cpp
--- a/Engine/Source/Runtime/MovieSceneTracks/Private/Sections/MovieSceneAudioSection.cpp
+++ b/Engine/Source/Runtime/MovieSceneTracks/Private/Sections/MovieSceneAudioSection.cpp
@@ -4,6 +4,23 @@
 #include "Sound/SoundBase.h"
 #include "Evaluation/MovieSceneAudioTemplate.h"
 #include "UObject/SequencerObjectVersion.h"
+#include "Channels/MovieSceneChannelProxy.h"
+#include "MovieScene.h"
+
+#if WITH_EDITOR
+
+struct FAudioChannelEditorData
+{
+	FAudioChannelEditorData()
+	{
+		Data[0].SetIdentifiers("Volume", NSLOCTEXT("MovieSceneAudioSection", "SoundVolumeText", "Volume"));
+		Data[1].SetIdentifiers("Pitch", NSLOCTEXT("MovieSceneAudioSection", "PitchText", "Pitch"));
+	}
+
+	FMovieSceneChannelEditorData Data[2];
+};
+
+#endif // WITH_EDITOR
 
 UMovieSceneAudioSection::UMovieSceneAudioSection( const FObjectInitializer& ObjectInitializer )
 	: Super( ObjectInitializer )
@@ -13,8 +30,6 @@ UMovieSceneAudioSection::UMovieSceneAudioSection( const FObjectInitializer& Obje
 	AudioStartTime_DEPRECATED = 0.f;
 	AudioDilationFactor_DEPRECATED = 1.f;
 	AudioVolume_DEPRECATED = 1.f;
-	SoundVolume.SetDefaultValue(1.f);
-	PitchMultiplier.SetDefaultValue(1.f);
 	bSuppressSubtitles = false;
 	bOverrideAttenuation = false;
 
@@ -22,6 +37,27 @@ UMovieSceneAudioSection::UMovieSceneAudioSection( const FObjectInitializer& Obje
 		(GetLinkerCustomVersion(FSequencerObjectVersion::GUID) < FSequencerObjectVersion::WhenFinishedDefaultsToProjectDefault ? 
 			EMovieSceneCompletionMode::RestoreState : 
 			EMovieSceneCompletionMode::ProjectDefault);
+
+	SoundVolume.SetDefault(1.f);
+	PitchMultiplier.SetDefault(1.f);
+
+	// Set up the channel proxy
+	FMovieSceneChannelData Channels;
+
+#if WITH_EDITOR
+
+	static const FAudioChannelEditorData EditorData;
+	Channels.Add(SoundVolume,     EditorData.Data[0], TMovieSceneExternalValue<float>());
+	Channels.Add(PitchMultiplier, EditorData.Data[1], TMovieSceneExternalValue<float>());
+
+#else
+
+	Channels.Add(SoundVolume);
+	Channels.Add(PitchMultiplier);
+
+#endif
+
+	ChannelProxy = MakeShared<FMovieSceneChannelProxy>(MoveTemp(Channels));
 }
 
 FMovieSceneEvalTemplatePtr UMovieSceneAudioSection::GenerateTemplate() const
@@ -29,22 +65,9 @@ FMovieSceneEvalTemplatePtr UMovieSceneAudioSection::GenerateTemplate() const
 	return FMovieSceneAudioSectionTemplate(*this);
 }
 
-TRange<float> UMovieSceneAudioSection::GetAudioRange() const
+TOptional<FFrameTime> UMovieSceneAudioSection::GetOffsetTime() const
 {
-	if (!Sound)
-	{
-		return TRange<float>::Empty();
-	}
-
-	float AudioRangeStartTime = GetStartTime();
-	if (StartOffset < 0)
-	{
-		AudioRangeStartTime += FMath::Abs(StartOffset);
-	}
-
-	float AudioRangeEndTime = FMath::Min(AudioRangeStartTime + Sound->GetDuration() * PitchMultiplier.GetDefaultValue(), GetEndTime());
-
-	return TRange<float>(AudioRangeStartTime, AudioRangeEndTime);
+	return TOptional<FFrameTime>(StartOffset * GetTypedOuter<UMovieScene>()->GetFrameResolution());
 }
 
 void UMovieSceneAudioSection::PostLoad()
@@ -53,14 +76,14 @@ void UMovieSceneAudioSection::PostLoad()
 
 	if (AudioDilationFactor_DEPRECATED != FLT_MAX)
 	{
-		PitchMultiplier.SetDefaultValue(AudioDilationFactor_DEPRECATED);
+		PitchMultiplier.SetDefault(AudioDilationFactor_DEPRECATED);
 
 		AudioDilationFactor_DEPRECATED = FLT_MAX;
 	}
 
 	if (AudioVolume_DEPRECATED != FLT_MAX)
 	{
-		SoundVolume.SetDefaultValue(AudioVolume_DEPRECATED);
+		SoundVolume.SetDefault(AudioVolume_DEPRECATED);
 
 		AudioVolume_DEPRECATED = FLT_MAX;
 	}
@@ -69,35 +92,18 @@ void UMovieSceneAudioSection::PostLoad()
 	{
 		// Previously, start time in relation to the sequence. Start time was used to calculate the offset into the 
 		// clip at the start of the section evaluation as such: Section Start Time - Start Time. 
-		if (AudioStartTime_DEPRECATED != 0.f)
+		if (AudioStartTime_DEPRECATED != 0.f && HasStartFrame())
 		{
-			StartOffset = GetStartTime() - AudioStartTime_DEPRECATED;
+			StartOffset = GetInclusiveStartFrame() / GetTypedOuter<UMovieScene>()->GetFrameResolution() - AudioStartTime_DEPRECATED;
 		}
 		AudioStartTime_DEPRECATED = FLT_MAX;
 	}
 }
 
-void UMovieSceneAudioSection::MoveSection( float DeltaTime, TSet<FKeyHandle>& KeyHandles )
-{
-	Super::MoveSection(DeltaTime, KeyHandles);
-
-	SoundVolume.ShiftCurve(DeltaTime, KeyHandles);
-	PitchMultiplier.ShiftCurve(DeltaTime, KeyHandles);
-}
-
 
-void UMovieSceneAudioSection::DilateSection( float DilationFactor, float Origin, TSet<FKeyHandle>& KeyHandles )
+UMovieSceneSection* UMovieSceneAudioSection::SplitSection(FFrameNumber SplitTime)
 {
-	Super::DilateSection(DilationFactor, Origin, KeyHandles);
-
-	SoundVolume.ScaleCurve(Origin, DilationFactor, KeyHandles);
-	PitchMultiplier.ScaleCurve(Origin, DilationFactor, KeyHandles);
-}
-
-
-UMovieSceneSection* UMovieSceneAudioSection::SplitSection(float SplitTime)
-{
-	float NewOffset = StartOffset + (SplitTime - GetStartTime());
+	const float NewOffset = HasStartFrame() ? StartOffset + (SplitTime - GetInclusiveStartFrame()) / GetTypedOuter<UMovieScene>()->GetFrameResolution() : 0;
 
 	UMovieSceneSection* NewSection = Super::SplitSection(SplitTime);
 	if (NewSection != nullptr)
@@ -107,38 +113,3 @@ UMovieSceneSection* UMovieSceneAudioSection::SplitSection(float SplitTime)
 	}
 	return NewSection;
 }
-
-void UMovieSceneAudioSection::GetKeyHandles(TSet<FKeyHandle>& OutKeyHandles, TRange<float> TimeRange) const
-{
-	if (!TimeRange.Overlaps(GetRange()))
-	{
-		return;
-	}
-
-	for (auto It(SoundVolume.GetKeyHandleIterator()); It; ++It)
-	{
-		float Time = SoundVolume.GetKeyTime(It.Key());
-		if (TimeRange.Contains(Time))
-		{
-			OutKeyHandles.Add(It.Key());
-		}
-	}
-
-	for (auto It(PitchMultiplier.GetKeyHandleIterator()); It; ++It)
-	{
-		float Time = PitchMultiplier.GetKeyTime(It.Key());
-		if (TimeRange.Contains(Time))
-		{
-			OutKeyHandles.Add(It.Key());
-		}
-	}
-}
-
-
-void UMovieSceneAudioSection::GetSnapTimes(TArray<float>& OutSnapTimes, bool bGetSectionBorders) const
-{
-	Super::GetSnapTimes(OutSnapTimes, bGetSectionBorders);
-
-	// @todo Sequencer handle snapping for time dilation
-	// @todo Don't add redundant times (can't use AddUnique due to floating point equality issues)
-}
diff --git a/Engine/Source/Runtime/MovieSceneTracks/Private/Sections/MovieSceneBoolSection.cpp b/Engine/Source/Runtime/MovieSceneTracks/Private/Sections/MovieSceneBoolSection.cpp
--- a/Engine/Source/Runtime/MovieSceneTracks/Private/Sections/MovieSceneBoolSection.cpp
+++ b/Engine/Source/Runtime/MovieSceneTracks/Private/Sections/MovieSceneBoolSection.cpp
@@ -2,12 +2,18 @@
 
 #include "Sections/MovieSceneBoolSection.h"
 #include "UObject/SequencerObjectVersion.h"
+#include "Channels/MovieSceneChannelProxy.h"
 
 UMovieSceneBoolSection::UMovieSceneBoolSection( const FObjectInitializer& ObjectInitializer )
 	: Super( ObjectInitializer )
 	, DefaultValue_DEPRECATED(false)
 {
-	SetIsInfinite(true);
+	bSupportsInfiniteRange = true;
+#if WITH_EDITORONLY_DATA
+	bIsInfinite_DEPRECATED = true;
+#endif
+
+	SetRange(TRange<FFrameNumber>::All());
 
 	EvalOptions.EnableAndSetCompletionMode
 		(GetLinkerCustomVersion(FSequencerObjectVersion::GUID) < FSequencerObjectVersion::WhenFinishedDefaultsToRestoreState ? 
@@ -15,106 +21,24 @@ UMovieSceneBoolSection::UMovieSceneBoolSection( const FObjectInitializer& Object
 			GetLinkerCustomVersion(FSequencerObjectVersion::GUID) < FSequencerObjectVersion::WhenFinishedDefaultsToProjectDefault ? 
 			EMovieSceneCompletionMode::RestoreState : 
 			EMovieSceneCompletionMode::ProjectDefault);
-}
-
-void UMovieSceneBoolSection::PostLoad()
-{
-	if (GetCurve().GetDefaultValue() == MAX_int32 && DefaultValue_DEPRECATED)
-	{
-		GetCurve().SetDefaultValue(DefaultValue_DEPRECATED);
-	}
-	Super::PostLoad();
-}
-
-bool UMovieSceneBoolSection::Eval( float Position, bool DefaultValue ) const
-{
-	return !!BoolCurve.Evaluate(Position, DefaultValue);
-}
-
-
-void UMovieSceneBoolSection::MoveSection( float DeltaPosition, TSet<FKeyHandle>& KeyHandles )
-{
-	Super::MoveSection( DeltaPosition, KeyHandles );
 
-	BoolCurve.ShiftCurve(DeltaPosition, KeyHandles);
-}
+#if WITH_EDITOR
 
+	ChannelProxy = MakeShared<FMovieSceneChannelProxy>(BoolCurve, FMovieSceneChannelEditorData(), TMovieSceneExternalValue<bool>::Make());
 
-void UMovieSceneBoolSection::DilateSection( float DilationFactor, float Origin, TSet<FKeyHandle>& KeyHandles )
-{
-	Super::DilateSection(DilationFactor, Origin, KeyHandles);
-	
-	BoolCurve.ScaleCurve(Origin, DilationFactor, KeyHandles);
-}
+#else
 
+	ChannelProxy = MakeShared<FMovieSceneChannelProxy>(BoolCurve);
 
-void UMovieSceneBoolSection::GetKeyHandles(TSet<FKeyHandle>& OutKeyHandles, TRange<float> TimeRange) const
-{
-	if (!TimeRange.Overlaps(GetRange()))
-	{
-		return;
-	}
-
-	for (auto It(BoolCurve.GetKeyHandleIterator()); It; ++It)
-	{
-		float Time = BoolCurve.GetKeyTime(It.Key());
-		if (TimeRange.Contains(Time))
-		{
-			OutKeyHandles.Add(It.Key());
-		}
-	}
+#endif
 }
 
 
-TOptional<float> UMovieSceneBoolSection::GetKeyTime( FKeyHandle KeyHandle ) const
-{
-	if ( BoolCurve.IsKeyHandleValid( KeyHandle ) )
-	{
-		return TOptional<float>( BoolCurve.GetKeyTime( KeyHandle ) );
-	}
-	return TOptional<float>();
-}
-
-
-void UMovieSceneBoolSection::SetKeyTime( FKeyHandle KeyHandle, float Time )
-{
-	if ( BoolCurve.IsKeyHandleValid( KeyHandle ) )
-	{
-		BoolCurve.SetKeyTime( KeyHandle, Time );
-	}
-}
-
-
-void UMovieSceneBoolSection::AddKey( float Time, const bool& Value, EMovieSceneKeyInterpolation KeyInterpolation )
-{
-	if (TryModify())
-	{
-		BoolCurve.UpdateOrAddKey(Time, Value ? 1 : 0);
-	}
-}
-
-
-void UMovieSceneBoolSection::SetDefault( const bool& Value )
+void UMovieSceneBoolSection::PostLoad()
 {
-	uint8 NewValue = Value ? 1 : 0;
-	if (BoolCurve.GetDefaultValue() != NewValue && TryModify())
+	if (!BoolCurve.GetDefault().IsSet() && DefaultValue_DEPRECATED)
 	{
-		BoolCurve.SetDefaultValue(NewValue);
+		BoolCurve.SetDefault(DefaultValue_DEPRECATED);
 	}
-}
-
-
-bool UMovieSceneBoolSection::NewKeyIsNewData( float Time, const bool& Value ) const
-{
-	return Eval(Time, Value) != Value;
-}
-
-bool UMovieSceneBoolSection::HasKeys( const bool& Value ) const
-{
-	return BoolCurve.GetNumKeys() != 0;
-}
-
-void UMovieSceneBoolSection::ClearDefaults()
-{
-	BoolCurve.ClearDefaultValue();
-}
+	Super::PostLoad();
+}
\ No newline at end of file
diff --git a/Engine/Source/Runtime/MovieSceneTracks/Private/Sections/MovieSceneByteSection.cpp b/Engine/Source/Runtime/MovieSceneTracks/Private/Sections/MovieSceneByteSection.cpp
--- a/Engine/Source/Runtime/MovieSceneTracks/Private/Sections/MovieSceneByteSection.cpp
+++ b/Engine/Source/Runtime/MovieSceneTracks/Private/Sections/MovieSceneByteSection.cpp
@@ -3,6 +3,7 @@
 #include "Sections/MovieSceneByteSection.h"
 #include "Evaluation/MovieScenePropertyTemplates.h"
 #include "UObject/SequencerObjectVersion.h"
+#include "Channels/MovieSceneChannelProxy.h"
 
 UMovieSceneByteSection::UMovieSceneByteSection( const FObjectInitializer& ObjectInitializer )
 	: Super( ObjectInitializer )
@@ -13,93 +14,15 @@ UMovieSceneByteSection::UMovieSceneByteSection( const FObjectInitializer& Object
 			GetLinkerCustomVersion(FSequencerObjectVersion::GUID) < FSequencerObjectVersion::WhenFinishedDefaultsToProjectDefault ? 
 			EMovieSceneCompletionMode::RestoreState : 
 			EMovieSceneCompletionMode::ProjectDefault);
-}
 
+	bSupportsInfiniteRange = true;
+#if WITH_EDITOR
 
-void UMovieSceneByteSection::MoveSection( float DeltaPosition, TSet<FKeyHandle>& KeyHandles )
-{
-	Super::MoveSection( DeltaPosition, KeyHandles );
+	ChannelProxy = MakeShared<FMovieSceneChannelProxy>(ByteCurve, FMovieSceneChannelEditorData(), TMovieSceneExternalValue<uint8>::Make());
 
-	ByteCurve.ShiftCurve(DeltaPosition, KeyHandles);
-}
+#else
 
+	ChannelProxy = MakeShared<FMovieSceneChannelProxy>(ByteCurve);
 
-void UMovieSceneByteSection::DilateSection( float DilationFactor, float Origin, TSet<FKeyHandle>& KeyHandles )
-{
-	Super::DilateSection(DilationFactor, Origin, KeyHandles);
-	
-	ByteCurve.ScaleCurve(Origin, DilationFactor, KeyHandles);
-}
-
-
-void UMovieSceneByteSection::GetKeyHandles(TSet<FKeyHandle>& OutKeyHandles, TRange<float> TimeRange) const
-{
-	if (!TimeRange.Overlaps(GetRange()))
-	{
-		return;
-	}
-
-	for (auto It(ByteCurve.GetKeyHandleIterator()); It; ++It)
-	{
-		float Time = ByteCurve.GetKeyTime(It.Key());
-		if (TimeRange.Contains(Time))
-		{
-			OutKeyHandles.Add(It.Key());
-		}
-	}
-}
-
-
-TOptional<float> UMovieSceneByteSection::GetKeyTime( FKeyHandle KeyHandle ) const
-{
-	if ( ByteCurve.IsKeyHandleValid( KeyHandle ) )
-	{
-		return TOptional<float>( ByteCurve.GetKeyTime( KeyHandle ) );
-	}
-	return TOptional<float>();
-}
-
-
-void UMovieSceneByteSection::SetKeyTime( FKeyHandle KeyHandle, float Time )
-{
-	if ( ByteCurve.IsKeyHandleValid( KeyHandle ) )
-	{
-		ByteCurve.SetKeyTime( KeyHandle, Time );
-	}
-}
-
-
-void UMovieSceneByteSection::AddKey( float Time, const uint8& Value, EMovieSceneKeyInterpolation KeyInterpolation )
-{
-	if (TryModify())
-	{
-		ByteCurve.UpdateOrAddKey(Time, Value);
-	}
-}
-
-
-bool UMovieSceneByteSection::NewKeyIsNewData( float Time, const uint8& Value ) const
-{
-	return ByteCurve.Evaluate( Time, Value ) != Value;
-}
-
-
-bool UMovieSceneByteSection::HasKeys( const uint8& Value ) const
-{
-	return ByteCurve.GetNumKeys() > 0;
-}
-
-
-void UMovieSceneByteSection::SetDefault( const uint8& Value )
-{
-	if (ByteCurve.GetDefaultValue() != Value && TryModify())
-	{
-		ByteCurve.SetDefaultValue(Value);
-	}
-}
-
-
-void UMovieSceneByteSection::ClearDefaults()
-{
-	ByteCurve.ClearDefaultValue();
-}
+#endif
+}
\ No newline at end of file
diff --git a/Engine/Source/Runtime/MovieSceneTracks/Private/Sections/MovieSceneCameraCutSection.cpp b/Engine/Source/Runtime/MovieSceneTracks/Private/Sections/MovieSceneCameraCutSection.cpp
--- a/Engine/Source/Runtime/MovieSceneTracks/Private/Sections/MovieSceneCameraCutSection.cpp
+++ b/Engine/Source/Runtime/MovieSceneTracks/Private/Sections/MovieSceneCameraCutSection.cpp
@@ -31,7 +31,7 @@ FMovieSceneEvalTemplatePtr UMovieSceneCameraCutSection::GenerateTemplate() const
 				{
 					// Extract the transform
 					FMovieSceneEvaluationTrack TransformTrackTemplate = TransformTrack->GenerateTrackTemplate();
-					FMovieSceneContext Context = FMovieSceneEvaluationRange(GetStartTime());
+					FMovieSceneContext Context = FMovieSceneEvaluationRange(GetInclusiveStartFrame(), MovieScene->GetFrameResolution());
 
 					FMovieSceneInterrogationData Container;
 					TransformTrackTemplate.Interrogate(Context, Container);
diff --git a/Engine/Source/Runtime/MovieSceneTracks/Private/Sections/MovieSceneColorSection.cpp b/Engine/Source/Runtime/MovieSceneTracks/Private/Sections/MovieSceneColorSection.cpp
--- a/Engine/Source/Runtime/MovieSceneTracks/Private/Sections/MovieSceneColorSection.cpp
+++ b/Engine/Source/Runtime/MovieSceneTracks/Private/Sections/MovieSceneColorSection.cpp
@@ -3,256 +3,145 @@
 #include "Sections/MovieSceneColorSection.h"
 #include "UObject/StructOnScope.h"
 #include "UObject/SequencerObjectVersion.h"
+#include "Channels/MovieSceneChannelProxy.h"
+#include "Styling/SlateColor.h"
 
-/* FMovieSceneColorKeyStruct interface
- *****************************************************************************/
-
-void FMovieSceneColorKeyStruct::PropagateChanges(const FPropertyChangedEvent& ChangeEvent)
+#if WITH_EDITOR
+struct FColorSectionEditorData
 {
-	for (int32 Index = 0; Index <= 3; ++Index)
+	FColorSectionEditorData()
 	{
-		if (Keys[Index] == nullptr)
-		{
-			if (Curves[Index] != nullptr)
-			{
-				Curves[Index]->SetDefaultValue(Color.Component(Index));
-			}
-		}
-		else
-		{
-			Keys[Index]->Value = Color.Component(Index);
-			Keys[Index]->Time = Time;
-		}
-	}
-}
+		CommonData[0].SetIdentifiers("Color.R", FCommonChannelData::ChannelR);
+		CommonData[0].SortOrder = 0;
+		CommonData[0].Color = FCommonChannelData::RedChannelColor;
 
+		CommonData[1].SetIdentifiers("Color.G", FCommonChannelData::ChannelG);
+		CommonData[1].SortOrder = 1;
+		CommonData[1].Color = FCommonChannelData::GreenChannelColor;
 
-/* UMovieSceneColorSection structors
- *****************************************************************************/
+		CommonData[2].SetIdentifiers("Color.B", FCommonChannelData::ChannelB);
+		CommonData[2].SortOrder = 2;
+		CommonData[2].Color = FCommonChannelData::BlueChannelColor;
 
-UMovieSceneColorSection::UMovieSceneColorSection(const FObjectInitializer& ObjectInitializer)
-	: Super(ObjectInitializer)
-{ 
-	EvalOptions.EnableAndSetCompletionMode
-		(GetLinkerCustomVersion(FSequencerObjectVersion::GUID) < FSequencerObjectVersion::WhenFinishedDefaultsToRestoreState ? 
-			EMovieSceneCompletionMode::KeepState : 
-			GetLinkerCustomVersion(FSequencerObjectVersion::GUID) < FSequencerObjectVersion::WhenFinishedDefaultsToProjectDefault ? 
-			EMovieSceneCompletionMode::RestoreState : 
-			EMovieSceneCompletionMode::ProjectDefault);
-	BlendType = EMovieSceneBlendType::Absolute;
-}
-
-/* UMovieSceneSection interface
- *****************************************************************************/
-
-void UMovieSceneColorSection::MoveSection(float DeltaTime, TSet<FKeyHandle>& KeyHandles)
-{
-	Super::MoveSection(DeltaTime, KeyHandles);
+		CommonData[3].SetIdentifiers("Color.A", FCommonChannelData::ChannelA);
+		CommonData[3].SortOrder = 3;
 
-	// Move all the curves in this section
-	RedCurve.ShiftCurve(DeltaTime, KeyHandles);
-	GreenCurve.ShiftCurve(DeltaTime, KeyHandles);
-	BlueCurve.ShiftCurve(DeltaTime, KeyHandles);
-	AlphaCurve.ShiftCurve(DeltaTime, KeyHandles);
-}
-
-
-void UMovieSceneColorSection::DilateSection(float DilationFactor, float Origin, TSet<FKeyHandle>& KeyHandles)
-{
-	Super::DilateSection(DilationFactor, Origin, KeyHandles);
-
-	RedCurve.ScaleCurve(Origin, DilationFactor, KeyHandles);
-	GreenCurve.ScaleCurve(Origin, DilationFactor, KeyHandles);
-	BlueCurve.ScaleCurve(Origin, DilationFactor, KeyHandles);
-	AlphaCurve.ScaleCurve(Origin, DilationFactor, KeyHandles);
-}
-
-
-void UMovieSceneColorSection::GetKeyHandles(TSet<FKeyHandle>& OutKeyHandles, TRange<float> TimeRange) const
-{
-	if (!TimeRange.Overlaps(GetRange()))
-	{
-		return;
-	}
-
-	for (auto It(RedCurve.GetKeyHandleIterator()); It; ++It)
-	{
-		float Time = RedCurve.GetKeyTime(It.Key());
-		if (TimeRange.Contains(Time))
-		{
-			OutKeyHandles.Add(It.Key());
-		}
-	}
-
-	for (auto It(GreenCurve.GetKeyHandleIterator()); It; ++It)
-	{
-		float Time = GreenCurve.GetKeyTime(It.Key());
-		if (TimeRange.Contains(Time))
-		{
-			OutKeyHandles.Add(It.Key());
-		}
-	}
-
-	for (auto It(BlueCurve.GetKeyHandleIterator()); It; ++It)
-	{
-		float Time = BlueCurve.GetKeyTime(It.Key());
-		if (TimeRange.Contains(Time))
-		{
-			OutKeyHandles.Add(It.Key());
-		}
-	}
-
-	for (auto It(AlphaCurve.GetKeyHandleIterator()); It; ++It)
-	{
-		float Time = AlphaCurve.GetKeyTime(It.Key());
-		if (TimeRange.Contains(Time))
-		{
-			OutKeyHandles.Add(It.Key());
-		}
+		ExternalValues[0].OnGetExternalValue = ExtractChannelR;
+		ExternalValues[1].OnGetExternalValue = ExtractChannelG;
+		ExternalValues[2].OnGetExternalValue = ExtractChannelB;
+		ExternalValues[3].OnGetExternalValue = ExtractChannelA;
 	}
-}
-
 
-TSharedPtr<FStructOnScope> UMovieSceneColorSection::GetKeyStruct(const TArray<FKeyHandle>& KeyHandles)
-{
-	TSharedRef<FStructOnScope> KeyStruct = MakeShareable(new FStructOnScope(FMovieSceneColorKeyStruct::StaticStruct()));
-	auto Struct = (FMovieSceneColorKeyStruct*)KeyStruct->GetStructMemory();
+	static FLinearColor GetPropertyValue(UObject& InObject, FTrackInstancePropertyBindings& Bindings)
 	{
-		Struct->Curves[0] = &RedCurve;
-		Struct->Curves[1] = &GreenCurve;
-		Struct->Curves[2] = &BlueCurve;
-		Struct->Curves[3] = &AlphaCurve;
-
-		Struct->Keys[0] = RedCurve.GetFirstMatchingKey(KeyHandles);
-		Struct->Keys[1] = GreenCurve.GetFirstMatchingKey(KeyHandles);
-		Struct->Keys[2] = BlueCurve.GetFirstMatchingKey(KeyHandles);
-		Struct->Keys[3] = AlphaCurve.GetFirstMatchingKey(KeyHandles);
+		const FName SlateColorName("SlateColor");
 
-		float FirstValidKeyTime = 0.f;
-		for (int32 Index = 0; Index <= 3; ++Index)
+		UStructProperty* ColorStructProperty = Cast<UStructProperty>(Bindings.GetProperty(InObject));
+		if (ColorStructProperty != nullptr)
 		{
-			if (Struct->Keys[Index] != nullptr)
+			if (ColorStructProperty->Struct->GetFName() == SlateColorName)
 			{
-				FirstValidKeyTime = Struct->Keys[Index]->Time;
-				Struct->Time = FirstValidKeyTime;
+				return Bindings.GetCurrentValue<FSlateColor>(InObject).GetSpecifiedColor();
 			}
-		}
 
-		for (int32 Index = 0; Index <= 3; ++Index)
-		{
-			if (Struct->Keys[Index] == nullptr && Struct->Curves[Index] != nullptr)
+			if (ColorStructProperty->Struct->GetFName() == NAME_LinearColor)
 			{
-				Struct->Color.Component(Index) = Struct->Curves[Index]->Eval(FirstValidKeyTime);
+				return Bindings.GetCurrentValue<FLinearColor>(InObject);
 			}
-			else
+
+			if (ColorStructProperty->Struct->GetFName() == NAME_Color)
 			{
-				Struct->Color.Component(Index) = Struct->Keys[Index]->Value;
+				return Bindings.GetCurrentValue<FColor>(InObject);
 			}
 		}
+		return FLinearColor(0.f,0.f,0.f,0.f);
 	}
 
-	return KeyStruct;
-}
-
-
-TOptional<float> UMovieSceneColorSection::GetKeyTime( FKeyHandle KeyHandle ) const
-{
-	if ( RedCurve.IsKeyHandleValid( KeyHandle ) )
+	static TOptional<float> ExtractChannelR(UObject& InObject, FTrackInstancePropertyBindings* Bindings)
 	{
-		return TOptional<float>( RedCurve.GetKeyTime( KeyHandle ) );
+		return Bindings ? GetPropertyValue(InObject, *Bindings).R : TOptional<float>();
 	}
-	if ( GreenCurve.IsKeyHandleValid( KeyHandle ) )
+	static TOptional<float> ExtractChannelG(UObject& InObject, FTrackInstancePropertyBindings* Bindings)
 	{
-		return TOptional<float>( GreenCurve.GetKeyTime( KeyHandle ) );
+		return Bindings ? GetPropertyValue(InObject, *Bindings).G : TOptional<float>();
 	}
-	if ( BlueCurve.IsKeyHandleValid( KeyHandle ) )
+	static TOptional<float> ExtractChannelB(UObject& InObject, FTrackInstancePropertyBindings* Bindings)
 	{
-		return TOptional<float>( BlueCurve.GetKeyTime( KeyHandle ) );
+		return Bindings ? GetPropertyValue(InObject, *Bindings).B : TOptional<float>();
 	}
-	if ( AlphaCurve.IsKeyHandleValid( KeyHandle ) )
+	static TOptional<float> ExtractChannelA(UObject& InObject, FTrackInstancePropertyBindings* Bindings)
 	{
-		return TOptional<float>( AlphaCurve.GetKeyTime( KeyHandle ) );
+		return Bindings ? GetPropertyValue(InObject, *Bindings).A : TOptional<float>();
 	}
-	return TOptional<float>();
-}
 
+	FMovieSceneChannelEditorData    CommonData[4];
+	TMovieSceneExternalValue<float> ExternalValues[4];
+};
+#endif
 
-void UMovieSceneColorSection::SetKeyTime( FKeyHandle KeyHandle, float Time )
+/* FMovieSceneColorKeyStruct interface
+ *****************************************************************************/
+
+void FMovieSceneColorKeyStruct::PropagateChanges(const FPropertyChangedEvent& ChangeEvent)
 {
-	if ( RedCurve.IsKeyHandleValid( KeyHandle ) )
-	{
-		RedCurve.SetKeyTime( KeyHandle, Time );
-	}
-	else if ( GreenCurve.IsKeyHandleValid( KeyHandle ) )
-	{
-		GreenCurve.SetKeyTime( KeyHandle, Time );
-	}
-	else if ( BlueCurve.IsKeyHandleValid( KeyHandle ) )
-	{
-		BlueCurve.SetKeyTime( KeyHandle, Time );
-	}
-	else if ( AlphaCurve.IsKeyHandleValid( KeyHandle ) )
-	{
-		AlphaCurve.SetKeyTime( KeyHandle, Time );
-	}
+	KeyStructInterop.Apply(Time);
 }
 
 
-/* IKeyframeSection interface
+/* UMovieSceneColorSection structors
  *****************************************************************************/
 
-template<typename CurveType>
-CurveType* GetCurveForChannel(EKeyColorChannel Channel, CurveType* RedCurve, CurveType* GreenCurve, CurveType* BlueCurve, CurveType* AlphaCurve)
+UMovieSceneColorSection::UMovieSceneColorSection(const FObjectInitializer& ObjectInitializer)
+	: Super(ObjectInitializer)
 {
-	switch (Channel)
-	{
-	case EKeyColorChannel::Red:
-		return RedCurve;
-	case EKeyColorChannel::Green:
-		return GreenCurve;
-	case EKeyColorChannel::Blue:
-		return BlueCurve;
-	case EKeyColorChannel::Alpha:
-		return AlphaCurve;
-	default:
-		checkf(false, TEXT("Invalid key color channel"));
-		return nullptr;
-	}
-}
+	EvalOptions.EnableAndSetCompletionMode
+		(GetLinkerCustomVersion(FSequencerObjectVersion::GUID) < FSequencerObjectVersion::WhenFinishedDefaultsToRestoreState ? 
+			EMovieSceneCompletionMode::KeepState : 
+			GetLinkerCustomVersion(FSequencerObjectVersion::GUID) < FSequencerObjectVersion::WhenFinishedDefaultsToProjectDefault ? 
+			EMovieSceneCompletionMode::RestoreState : 
+			EMovieSceneCompletionMode::ProjectDefault);
+	BlendType = EMovieSceneBlendType::Absolute;
+	bSupportsInfiniteRange = true;
 
+	FMovieSceneChannelData Channels;
 
-void UMovieSceneColorSection::AddKey(float Time, const FColorKey& Key, EMovieSceneKeyInterpolation KeyInterpolation)
-{
-	FRichCurve* ChannelCurve = GetCurveForChannel(Key.Channel, &RedCurve, &GreenCurve, &BlueCurve, &AlphaCurve);
-	AddKeyToCurve(*ChannelCurve, Time, Key.ChannelValue, KeyInterpolation);
-}
+#if WITH_EDITOR
 
+	static FColorSectionEditorData EditorData;
+	Channels.Add(RedCurve,   EditorData.CommonData[0], EditorData.ExternalValues[0]);
+	Channels.Add(GreenCurve, EditorData.CommonData[1], EditorData.ExternalValues[1]);
+	Channels.Add(BlueCurve,  EditorData.CommonData[2], EditorData.ExternalValues[2]);
+	Channels.Add(AlphaCurve, EditorData.CommonData[3], EditorData.ExternalValues[3]);
 
-bool UMovieSceneColorSection::NewKeyIsNewData(float Time, const FColorKey& Key) const
-{
-	const FRichCurve* ChannelCurve = GetCurveForChannel(Key.Channel, &RedCurve, &GreenCurve, &BlueCurve, &AlphaCurve);
-	return FMath::IsNearlyEqual(ChannelCurve->Eval(Time), Key.ChannelValue) == false;
-}
+#else
 
+	Channels.Add(RedCurve);
+	Channels.Add(GreenCurve);
+	Channels.Add(BlueCurve);
+	Channels.Add(AlphaCurve);
 
-bool UMovieSceneColorSection::HasKeys(const FColorKey& Key) const
-{
-	const FRichCurve* ChannelCurve = GetCurveForChannel(Key.Channel, &RedCurve, &GreenCurve, &BlueCurve, &AlphaCurve);
-	return ChannelCurve->GetNumKeys() != 0;
+#endif
+
+	ChannelProxy = MakeShared<FMovieSceneChannelProxy>(MoveTemp(Channels));
 }
 
+/* UMovieSceneSection interface
+ *****************************************************************************/
 
-void UMovieSceneColorSection::SetDefault(const FColorKey& Key)
+TSharedPtr<FStructOnScope> UMovieSceneColorSection::GetKeyStruct(TArrayView<const FKeyHandle> KeyHandles)
 {
-	FRichCurve* ChannelCurve = GetCurveForChannel(Key.Channel, &RedCurve, &GreenCurve, &BlueCurve, &AlphaCurve);
-	SetCurveDefault(*ChannelCurve, Key.ChannelValue);
-}
+	TSharedRef<FStructOnScope> KeyStruct = MakeShareable(new FStructOnScope(FMovieSceneColorKeyStruct::StaticStruct()));
+	auto Struct = (FMovieSceneColorKeyStruct*)KeyStruct->GetStructMemory();
 
+	TArrayView<FMovieSceneFloatChannel*> FloatChannels = ChannelProxy->GetChannels<FMovieSceneFloatChannel>();
 
-void UMovieSceneColorSection::ClearDefaults()
-{
-	RedCurve.ClearDefaultValue();
-	GreenCurve.ClearDefaultValue();
-	BlueCurve.ClearDefaultValue();
-	AlphaCurve.ClearDefaultValue();
+	Struct->KeyStructInterop.Add(FMovieSceneChannelValueHelper(ChannelProxy->MakeHandle(FloatChannels[0]), &Struct->Color.R, KeyHandles));
+	Struct->KeyStructInterop.Add(FMovieSceneChannelValueHelper(ChannelProxy->MakeHandle(FloatChannels[1]), &Struct->Color.G, KeyHandles));
+	Struct->KeyStructInterop.Add(FMovieSceneChannelValueHelper(ChannelProxy->MakeHandle(FloatChannels[2]), &Struct->Color.B, KeyHandles));
+	Struct->KeyStructInterop.Add(FMovieSceneChannelValueHelper(ChannelProxy->MakeHandle(FloatChannels[3]), &Struct->Color.A, KeyHandles));
+
+	Struct->KeyStructInterop.SetStartingValues();
+	Struct->Time = Struct->KeyStructInterop.GetUnifiedKeyTime().Get(0);
+
+	return KeyStruct;
 }
diff --git a/Engine/Source/Runtime/MovieSceneTracks/Private/Sections/MovieSceneEnumSection.cpp b/Engine/Source/Runtime/MovieSceneTracks/Private/Sections/MovieSceneEnumSection.cpp
--- a/Engine/Source/Runtime/MovieSceneTracks/Private/Sections/MovieSceneEnumSection.cpp
+++ b/Engine/Source/Runtime/MovieSceneTracks/Private/Sections/MovieSceneEnumSection.cpp
@@ -1,105 +1,28 @@
 // Copyright 1998-2018 Epic Games, Inc. All Rights Reserved.
 
 #include "Sections/MovieSceneEnumSection.h"
-#include "Evaluation/MovieScenePropertyTemplates.h"
 #include "UObject/SequencerObjectVersion.h"
+#include "UObject/SequencerObjectVersion.h"
+#include "UObject/SequencerObjectVersion.h"
+#include "Channels/MovieSceneChannelProxy.h"
 
 UMovieSceneEnumSection::UMovieSceneEnumSection( const FObjectInitializer& ObjectInitializer )
 	: Super( ObjectInitializer )
-{ 
+{
 	EvalOptions.EnableAndSetCompletionMode
 		(GetLinkerCustomVersion(FSequencerObjectVersion::GUID) < FSequencerObjectVersion::WhenFinishedDefaultsToRestoreState ? 
 			EMovieSceneCompletionMode::KeepState : 
 			GetLinkerCustomVersion(FSequencerObjectVersion::GUID) < FSequencerObjectVersion::WhenFinishedDefaultsToProjectDefault ? 
 			EMovieSceneCompletionMode::RestoreState : 
 			EMovieSceneCompletionMode::ProjectDefault);
-}
-
-
-void UMovieSceneEnumSection::MoveSection( float DeltaPosition, TSet<FKeyHandle>& KeyHandles )
-{
-	Super::MoveSection( DeltaPosition, KeyHandles );
-
-	EnumCurve.ShiftCurve(DeltaPosition, KeyHandles);
-}
-
+	bSupportsInfiniteRange = true;
+#if WITH_EDITOR
 
-void UMovieSceneEnumSection::DilateSection( float DilationFactor, float Origin, TSet<FKeyHandle>& KeyHandles )
-{
-	Super::DilateSection(DilationFactor, Origin, KeyHandles);
-	
-	EnumCurve.ScaleCurve(Origin, DilationFactor, KeyHandles);
-}
-
-
-void UMovieSceneEnumSection::GetKeyHandles(TSet<FKeyHandle>& OutKeyHandles, TRange<float> TimeRange) const
-{
-	if (!TimeRange.Overlaps(GetRange()))
-	{
-		return;
-	}
+	ChannelProxy = MakeShared<FMovieSceneChannelProxy>(EnumCurve, FMovieSceneChannelEditorData(), TMovieSceneExternalValue<uint8>::Make());
 
-	for (auto It(EnumCurve.GetKeyHandleIterator()); It; ++It)
-	{
-		float Time = EnumCurve.GetKeyTime(It.Key());
-		if (TimeRange.Contains(Time))
-		{
-			OutKeyHandles.Add(It.Key());
-		}
-	}
-}
+#else
 
+	ChannelProxy = MakeShared<FMovieSceneChannelProxy>(EnumCurve);
 
-TOptional<float> UMovieSceneEnumSection::GetKeyTime( FKeyHandle KeyHandle ) const
-{
-	if ( EnumCurve.IsKeyHandleValid( KeyHandle ) )
-	{
-		return TOptional<float>( EnumCurve.GetKeyTime( KeyHandle ) );
-	}
-	return TOptional<float>();
-}
-
-
-void UMovieSceneEnumSection::SetKeyTime( FKeyHandle KeyHandle, float Time )
-{
-	if ( EnumCurve.IsKeyHandleValid( KeyHandle ) )
-	{
-		EnumCurve.SetKeyTime( KeyHandle, Time );
-	}
-}
-
-
-void UMovieSceneEnumSection::AddKey( float Time, const int64& Value, EMovieSceneKeyInterpolation KeyInterpolation )
-{
-	if (TryModify())
-	{
-		EnumCurve.UpdateOrAddKey(Time, Value);
-	}
-}
-
-
-bool UMovieSceneEnumSection::NewKeyIsNewData( float Time, const int64& Value ) const
-{
-	return EnumCurve.Evaluate( Time, Value ) != Value;
-}
-
-
-bool UMovieSceneEnumSection::HasKeys( const int64& Value ) const
-{
-	return EnumCurve.GetNumKeys() > 0;
-}
-
-
-void UMovieSceneEnumSection::SetDefault( const int64& Value )
-{
-	if (EnumCurve.GetDefaultValue() != Value && TryModify())
-	{
-		EnumCurve.SetDefaultValue(Value);
-	}
-}
-
-
-void UMovieSceneEnumSection::ClearDefaults()
-{
-	EnumCurve.ClearDefaultValue();
-}
+#endif
+}
\ No newline at end of file
diff --git a/Engine/Source/Runtime/MovieSceneTracks/Private/Sections/MovieSceneEventSection.cpp b/Engine/Source/Runtime/MovieSceneTracks/Private/Sections/MovieSceneEventSection.cpp
--- a/Engine/Source/Runtime/MovieSceneTracks/Private/Sections/MovieSceneEventSection.cpp
+++ b/Engine/Source/Runtime/MovieSceneTracks/Private/Sections/MovieSceneEventSection.cpp
@@ -5,10 +5,11 @@
 #include "IMovieScenePlayer.h"
 #include "UObject/ReleaseObjectVersion.h"
 #include "UObject/LinkerLoad.h"
+#include "MovieSceneFwd.h"
+#include "Channels/MovieSceneChannelProxy.h"
 #include "Serialization/MemoryArchive.h"
 #include "Engine/UserDefinedStruct.h"
-
-#include "Curves/KeyFrameAlgorithms.h"
+#include "MovieSceneFrameMigration.h"
 
 /* Custom version specifically for event parameter struct serialization (serialized into FMovieSceneEventParameters::StructBytes) */
 namespace EEventParameterVersion
@@ -84,7 +85,7 @@ public:
 	 * @param StructPtr 	The struct representing the type of Source
 	 * @param Source 		Source pointer to the instance of the payload to write
 	 */
-	void Write(UStruct* StructPtr, uint8* Source)
+	void Write(UScriptStruct* StructPtr, uint8* Source)
 	{
 		FArchive& Ar = *this;
 
@@ -154,7 +155,7 @@ public:
 	 * @param StructPtr 	The struct representing the type of Dest
 	 * @param Dest 			Destination instance to receive the deserialized data
 	 */
-	void Read(UStruct* StructPtr, uint8* Dest)
+	void Read(UScriptStruct* StructPtr, uint8* Dest)
 	{
 		bool bHasCustomVersion = false;
 		// Optionally deserialize the custom version header, provided it was serialized
@@ -221,11 +222,42 @@ private:
 	const TArray<uint8>& Bytes;
 };
 
+bool operator==(const FMovieSceneEventParameters& A, const FMovieSceneEventParameters& B)
+{
+	UScriptStruct* StructA = A.GetStructType();
+	UScriptStruct* StructB = B.GetStructType();
+
+	if (StructA != StructB)
+	{
+		return false;
+	}
+
+	if (!StructA)
+	{
+		return true;
+	}
+
+	FStructOnScope StructContainerA(StructA);
+	A.GetInstance(StructContainerA);
+	uint8* InstA = StructContainerA.GetStructMemory();
+
+	FStructOnScope StructContainerB(StructB);
+	B.GetInstance(StructContainerB);
+	uint8* InstB = StructContainerB.GetStructMemory();
+
+	return InstA ? StructA->CompareScriptStruct(InstA, InstB, 0) : InstB == nullptr;
+}
+
+bool operator!=(const FMovieSceneEventParameters& A, const FMovieSceneEventParameters& B)
+{
+	return !(A == B);
+}
+
 void FMovieSceneEventParameters::OverwriteWith(uint8* InstancePtr)
 {
 	check(InstancePtr);
 
-	if (UStruct* StructPtr = GetStructType())
+	if (UScriptStruct* StructPtr = GetStructType())
 	{
 		FEventParameterWriter(StructBytes).Write(StructPtr, InstancePtr);
 	}
@@ -237,7 +269,7 @@ void FMovieSceneEventParameters::OverwriteWith(uint8* InstancePtr)
 
 void FMovieSceneEventParameters::GetInstance(FStructOnScope& OutStruct) const
 {
-	UStruct* StructPtr = GetStructType();
+	UScriptStruct* StructPtr = GetStructType();
 	OutStruct.Initialize(StructPtr);
 
 #if WITH_EDITOR
@@ -261,7 +293,7 @@ bool FMovieSceneEventParameters::Serialize(FArchive& Ar)
 
 	if (Ar.IsLoading() && Ar.CustomVer(FReleaseObjectVersion::GUID) < FReleaseObjectVersion::EventSectionParameterStringAssetRef)
 	{
-		UStruct* StructPtr = nullptr;
+		UScriptStruct* StructPtr = nullptr;
 		Ar << StructPtr;
 		StructType = StructPtr;
 	}
@@ -275,73 +307,71 @@ bool FMovieSceneEventParameters::Serialize(FArchive& Ar)
 	return true;
 }
 
-/* UMovieSceneSection structors
- *****************************************************************************/
-
-UMovieSceneEventSection::UMovieSceneEventSection()
-#if WITH_EDITORONLY_DATA
-	: CurveInterface(TCurveInterface<FEventPayload, float>(&EventData.KeyTimes, &EventData.KeyValues, &EventData.KeyHandles))
-#else
-	: CurveInterface(TCurveInterface<FEventPayload, float>(&EventData.KeyTimes, &EventData.KeyValues))
-#endif
+uint32 FMovieSceneEventSectionData::GetChannelID()
 {
-	SetIsInfinite(true);
+	static uint32 ID = FMovieSceneChannelEntry::RegisterNewID();
+	return ID;
 }
 
-void UMovieSceneEventSection::PostLoad()
+void FMovieSceneEventSectionData::PostSerialize(const FArchive& Ar)
 {
-	for (FNameCurveKey EventKey : Events_DEPRECATED.GetKeys())
+#if WITH_EDITORONLY_DATA
+	
+	if (KeyTimes_DEPRECATED.Num())
 	{
-		EventData.KeyTimes.Add(EventKey.Time);
-		EventData.KeyValues.Add(FEventPayload(EventKey.Value));
-	}
+		FFrameRate LegacyFrameRate = GetLegacyConversionFrameRate();
 
-	if (Events_DEPRECATED.GetKeys().Num())
-	{
-		MarkAsChanged();
+		TArray<FEventPayload> OldValues = KeyValues;
+		Times.Reset(KeyTimes_DEPRECATED.Num());
+		KeyValues.Reset(KeyTimes_DEPRECATED.Num());
+		for (int32 Index = 0; Index < KeyTimes_DEPRECATED.Num(); ++Index)
+		{
+			FFrameNumber KeyTime = UpgradeLegacyMovieSceneTime(nullptr, LegacyFrameRate, KeyTimes_DEPRECATED[Index]);
+			ConvertInsertAndSort<FEventPayload>(Index, KeyTime, OldValues[Index], Times, KeyValues);
+		}
+		KeyTimes_DEPRECATED.Empty();
 	}
-
-	Super::PostLoad();
+#endif
 }
 
-/* UMovieSceneSection overrides
+/* UMovieSceneSection structors
  *****************************************************************************/
 
-void UMovieSceneEventSection::DilateSection(float DilationFactor, float Origin, TSet<FKeyHandle>& KeyHandles)
+UMovieSceneEventSection::UMovieSceneEventSection()
 {
-	Super::DilateSection(DilationFactor, Origin, KeyHandles);
-
-	KeyFrameAlgorithms::Scale(CurveInterface.GetValue(), Origin, DilationFactor, KeyHandles);
-}
+#if WITH_EDITORONLY_DATA
+	bIsInfinite_DEPRECATED = true;
+#endif
+	bSupportsInfiniteRange = true;
+	SetRange(TRange<FFrameNumber>::All());
 
+#if WITH_EDITOR
 
-void UMovieSceneEventSection::GetKeyHandles(TSet<FKeyHandle>& KeyHandles, TRange<float> TimeRange) const
-{
-	for (auto It(CurveInterface->IterateKeys()); It; ++It)
-	{
-		if (TimeRange.Contains(*It))
-		{
-			KeyHandles.Add(It.GetKeyHandle());
-		}
-	}
-}
+	ChannelProxy = MakeShared<FMovieSceneChannelProxy>(EventData, FMovieSceneChannelEditorData());
 
+#else
 
-void UMovieSceneEventSection::MoveSection(float DeltaPosition, TSet<FKeyHandle>& KeyHandles)
-{
-	Super::MoveSection(DeltaPosition, KeyHandles);
+	ChannelProxy = MakeShared<FMovieSceneChannelProxy>(EventData);
 
-	KeyFrameAlgorithms::Translate(CurveInterface.GetValue(), DeltaPosition, KeyHandles);
+#endif
 }
 
-
-TOptional<float> UMovieSceneEventSection::GetKeyTime( FKeyHandle KeyHandle ) const
+void UMovieSceneEventSection::PostLoad()
 {
-	return CurveInterface->GetKeyTime(KeyHandle);
-}
+	if (Events_DEPRECATED.GetKeys().Num())
+	{
+		TMovieSceneChannel<FEventPayload> Channel = EventData.GetInterface();
 
+		FFrameRate LegacyFrameRate = GetLegacyConversionFrameRate();
 
-void UMovieSceneEventSection::SetKeyTime( FKeyHandle KeyHandle, float Time )
-{
-	CurveInterface->SetKeyTime(KeyHandle, Time);
-}
+		for (FNameCurveKey EventKey : Events_DEPRECATED.GetKeys())
+		{
+			FFrameNumber KeyTime = UpgradeLegacyMovieSceneTime(this, LegacyFrameRate, EventKey.Time);
+			Channel.AddKey(KeyTime, FEventPayload(EventKey.Value));
+		}
+
+		MarkAsChanged();
+	}
+
+	Super::PostLoad();
+}
\ No newline at end of file
diff --git a/Engine/Source/Runtime/MovieSceneTracks/Private/Sections/MovieSceneFadeSection.cpp b/Engine/Source/Runtime/MovieSceneTracks/Private/Sections/MovieSceneFadeSection.cpp
--- a/Engine/Source/Runtime/MovieSceneTracks/Private/Sections/MovieSceneFadeSection.cpp
+++ b/Engine/Source/Runtime/MovieSceneTracks/Private/Sections/MovieSceneFadeSection.cpp
@@ -12,7 +12,11 @@ UMovieSceneFadeSection::UMovieSceneFadeSection()
 	, FadeColor(FLinearColor::Black)
 	, bFadeAudio(false)
 {
-	SetIsInfinite(true);
+#if WITH_EDITORONLY_DATA
+	bIsInfinite_DEPRECATED = true;
+#endif
+
+	SetRange(TRange<FFrameNumber>::All());
 
 	EvalOptions.EnableAndSetCompletionMode
 		(GetLinkerCustomVersion(FSequencerObjectVersion::GUID) < FSequencerObjectVersion::WhenFinishedDefaultsToProjectDefault ? 
diff --git a/Engine/Source/Runtime/MovieSceneTracks/Private/Sections/MovieSceneFloatSection.cpp b/Engine/Source/Runtime/MovieSceneTracks/Private/Sections/MovieSceneFloatSection.cpp
--- a/Engine/Source/Runtime/MovieSceneTracks/Private/Sections/MovieSceneFloatSection.cpp
+++ b/Engine/Source/Runtime/MovieSceneTracks/Private/Sections/MovieSceneFloatSection.cpp
@@ -2,10 +2,11 @@
 
 #include "Sections/MovieSceneFloatSection.h"
 #include "UObject/SequencerObjectVersion.h"
+#include "Channels/MovieSceneChannelProxy.h"
 
 UMovieSceneFloatSection::UMovieSceneFloatSection( const FObjectInitializer& ObjectInitializer )
 	: Super( ObjectInitializer )
-{ 
+{
 	EvalOptions.EnableAndSetCompletionMode
 		(GetLinkerCustomVersion(FSequencerObjectVersion::GUID) < FSequencerObjectVersion::WhenFinishedDefaultsToRestoreState ? 
 			EMovieSceneCompletionMode::KeepState : 
@@ -13,92 +14,14 @@ UMovieSceneFloatSection::UMovieSceneFloatSection( const FObjectInitializer& Obje
 			EMovieSceneCompletionMode::RestoreState : 
 			EMovieSceneCompletionMode::ProjectDefault);
 	BlendType = EMovieSceneBlendType::Absolute;
-}
-
-
-float UMovieSceneFloatSection::Eval( float Position, float DefaultValue ) const
-{
-	return FloatCurve.Eval( Position, DefaultValue );
-}
-
-
-void UMovieSceneFloatSection::MoveSection( float DeltaPosition, TSet<FKeyHandle>& KeyHandles )
-{
-	Super::MoveSection( DeltaPosition, KeyHandles );
-
-	// Move the curve
-	FloatCurve.ShiftCurve(DeltaPosition, KeyHandles);
-}
+	bSupportsInfiniteRange = true;
+#if WITH_EDITOR
 
+	ChannelProxy = MakeShared<FMovieSceneChannelProxy>(FloatCurve, FMovieSceneChannelEditorData(), TMovieSceneExternalValue<float>::Make());
 
-void UMovieSceneFloatSection::DilateSection( float DilationFactor, float Origin, TSet<FKeyHandle>& KeyHandles )
-{	
-	Super::DilateSection(DilationFactor, Origin, KeyHandles);
-	
-	FloatCurve.ScaleCurve(Origin, DilationFactor, KeyHandles);
-}
-
-
-void UMovieSceneFloatSection::GetKeyHandles(TSet<FKeyHandle>& OutKeyHandles, TRange<float> TimeRange) const
-{
-	if (!TimeRange.Overlaps(GetRange()))
-	{
-		return;
-	}
-
-	for (auto It(FloatCurve.GetKeyHandleIterator()); It; ++It)
-	{
-		float Time = FloatCurve.GetKeyTime(It.Key());
-		if (TimeRange.Contains(Time))
-		{
-			OutKeyHandles.Add(It.Key());
-		}
-	}
-}
-
-
-TOptional<float> UMovieSceneFloatSection::GetKeyTime( FKeyHandle KeyHandle ) const
-{
-	if ( FloatCurve.IsKeyHandleValid( KeyHandle ) )
-	{
-		return TOptional<float>( FloatCurve.GetKeyTime( KeyHandle ) );
-	}
-	return TOptional<float>();
-}
+#else
 
+	ChannelProxy = MakeShared<FMovieSceneChannelProxy>(FloatCurve);
 
-void UMovieSceneFloatSection::SetKeyTime( FKeyHandle KeyHandle, float Time )
-{
-	if ( FloatCurve.IsKeyHandleValid( KeyHandle ) )
-	{
-		FloatCurve.SetKeyTime( KeyHandle, Time );
-	}
-}
-
-
-void UMovieSceneFloatSection::AddKey( float Time, const float& Value, EMovieSceneKeyInterpolation KeyInterpolation )
-{
-	AddKeyToCurve( FloatCurve, Time, Value, KeyInterpolation );
-}
-
-
-bool UMovieSceneFloatSection::NewKeyIsNewData(float Time, const float& Value) const
-{
-	return FMath::IsNearlyEqual( FloatCurve.Eval( Time ), Value ) == false;
-}
-
-bool UMovieSceneFloatSection::HasKeys( const float& Value ) const
-{
-	return FloatCurve.GetNumKeys() > 0;
-}
-
-void UMovieSceneFloatSection::SetDefault( const float& Value )
-{
-	SetCurveDefault( FloatCurve, Value );
-}
-
-
-void UMovieSceneFloatSection::ClearDefaults()
-{
-	FloatCurve.ClearDefaultValue();
-}
+#endif
+}
\ No newline at end of file
diff --git a/Engine/Source/Runtime/MovieSceneTracks/Private/Sections/MovieSceneIntegerSection.cpp b/Engine/Source/Runtime/MovieSceneTracks/Private/Sections/MovieSceneIntegerSection.cpp
--- a/Engine/Source/Runtime/MovieSceneTracks/Private/Sections/MovieSceneIntegerSection.cpp
+++ b/Engine/Source/Runtime/MovieSceneTracks/Private/Sections/MovieSceneIntegerSection.cpp
@@ -1,12 +1,12 @@
 // Copyright 1998-2018 Epic Games, Inc. All Rights Reserved.
 
 #include "Sections/MovieSceneIntegerSection.h"
-#include "Evaluation/MovieScenePropertyTemplates.h"
 #include "UObject/SequencerObjectVersion.h"
+#include "Channels/MovieSceneChannelProxy.h"
 
 UMovieSceneIntegerSection::UMovieSceneIntegerSection( const FObjectInitializer& ObjectInitializer )
 	: Super( ObjectInitializer )
-{ 
+{
 	EvalOptions.EnableAndSetCompletionMode
 		(GetLinkerCustomVersion(FSequencerObjectVersion::GUID) < FSequencerObjectVersion::WhenFinishedDefaultsToRestoreState ? 
 			EMovieSceneCompletionMode::KeepState : 
@@ -14,93 +14,14 @@ UMovieSceneIntegerSection::UMovieSceneIntegerSection( const FObjectInitializer&
 			EMovieSceneCompletionMode::RestoreState : 
 			EMovieSceneCompletionMode::ProjectDefault);
 	BlendType = EMovieSceneBlendType::Absolute;
-}
-
-
-void UMovieSceneIntegerSection::MoveSection( float DeltaPosition, TSet<FKeyHandle>& KeyHandles )
-{
-	Super::MoveSection( DeltaPosition, KeyHandles );
+	bSupportsInfiniteRange = true;
+#if WITH_EDITOR
 
-	IntegerCurve.ShiftCurve(DeltaPosition, KeyHandles);
-}
+	ChannelProxy = MakeShared<FMovieSceneChannelProxy>(IntegerCurve, FMovieSceneChannelEditorData(), TMovieSceneExternalValue<int32>::Make());
 
+#else
 
-void UMovieSceneIntegerSection::DilateSection( float DilationFactor, float Origin, TSet<FKeyHandle>& KeyHandles )
-{
-	Super::DilateSection(DilationFactor, Origin, KeyHandles);
-	
-	IntegerCurve.ScaleCurve(Origin, DilationFactor, KeyHandles);
-}
-
-
-void UMovieSceneIntegerSection::GetKeyHandles(TSet<FKeyHandle>& OutKeyHandles, TRange<float> TimeRange) const
-{
-	if (!TimeRange.Overlaps(GetRange()))
-	{
-		return;
-	}
-
-	for (auto It(IntegerCurve.GetKeyHandleIterator()); It; ++It)
-	{
-		float Time = IntegerCurve.GetKeyTime(It.Key());
-		if (TimeRange.Contains(Time))
-		{
-			OutKeyHandles.Add(It.Key());
-		}
-	}
-}
-
-
-TOptional<float> UMovieSceneIntegerSection::GetKeyTime( FKeyHandle KeyHandle ) const
-{
-	if ( IntegerCurve.IsKeyHandleValid( KeyHandle ) )
-	{
-		return TOptional<float>( IntegerCurve.GetKeyTime( KeyHandle ) );
-	}
-	return TOptional<float>();
-}
-
-
-void UMovieSceneIntegerSection::SetKeyTime( FKeyHandle KeyHandle, float Time )
-{
-	if ( IntegerCurve.IsKeyHandleValid( KeyHandle ) )
-	{
-		IntegerCurve.SetKeyTime( KeyHandle, Time );
-	}
-}
+	ChannelProxy = MakeShared<FMovieSceneChannelProxy>(IntegerCurve);
 
-
-void UMovieSceneIntegerSection::AddKey( float Time, const int32& Value, EMovieSceneKeyInterpolation KeyInterpolation )
-{
-	if (TryModify())
-	{
-		IntegerCurve.UpdateOrAddKey(Time, Value);
-	}
-}
-
-
-bool UMovieSceneIntegerSection::NewKeyIsNewData( float Time, const int32& Value ) const
-{
-	return IntegerCurve.Evaluate( Time, Value ) != Value;
-}
-
-
-bool UMovieSceneIntegerSection::HasKeys( const int32& Value ) const
-{
-	return IntegerCurve.GetNumKeys() > 0;
-}
-
-
-void UMovieSceneIntegerSection::SetDefault( const int32& Value )
-{
-	if (IntegerCurve.GetDefaultValue() != Value && TryModify())
-	{
-		IntegerCurve.SetDefaultValue(Value);
-	}
-}
-
-
-void UMovieSceneIntegerSection::ClearDefaults()
-{
-	IntegerCurve.ClearDefaultValue();
-}
+#endif
+}
\ No newline at end of file
diff --git a/Engine/Source/Runtime/MovieSceneTracks/Private/Sections/MovieSceneLevelVisibilitySection.cpp b/Engine/Source/Runtime/MovieSceneTracks/Private/Sections/MovieSceneLevelVisibilitySection.cpp
--- a/Engine/Source/Runtime/MovieSceneTracks/Private/Sections/MovieSceneLevelVisibilitySection.cpp
+++ b/Engine/Source/Runtime/MovieSceneTracks/Private/Sections/MovieSceneLevelVisibilitySection.cpp
@@ -32,16 +32,4 @@ TArray<FName>* UMovieSceneLevelVisibilitySection::GetLevelNames()
 FMovieSceneEvalTemplatePtr UMovieSceneLevelVisibilitySection::GenerateTemplate() const
 {
 	return FMovieSceneLevelVisibilitySectionTemplate(*this);
-}
-
-
-TOptional<float> UMovieSceneLevelVisibilitySection::GetKeyTime(FKeyHandle KeyHandle) const
-{
-	return TOptional<float>();
-}
-
-
-void UMovieSceneLevelVisibilitySection::SetKeyTime(FKeyHandle KeyHandle, float Time)
-{
-	// do nothing
-}
+}
\ No newline at end of file
