Copying //UE4/Dev-Sequencer to //UE4/Dev-Main (Source: //UE4/Dev-Sequencer @ 3946692)

diff --git a/Engine/Source/Runtime/MovieSceneTracks/Private/Sections/MovieSceneParameterSection.cpp b/Engine/Source/Runtime/MovieSceneTracks/Private/Sections/MovieSceneParameterSection.cpp
--- a/Engine/Source/Runtime/MovieSceneTracks/Private/Sections/MovieSceneParameterSection.cpp
+++ b/Engine/Source/Runtime/MovieSceneTracks/Private/Sections/MovieSceneParameterSection.cpp
@@ -2,6 +2,7 @@
 
 #include "Sections/MovieSceneParameterSection.h"
 #include "UObject/SequencerObjectVersion.h"
+#include "Channels/MovieSceneChannelProxy.h"
 
 FScalarParameterNameAndCurve::FScalarParameterNameAndCurve( FName InParameterName )
 {
@@ -21,6 +22,7 @@ FColorParameterNameAndCurves::FColorParameterNameAndCurves( FName InParameterNam
 UMovieSceneParameterSection::UMovieSceneParameterSection( const FObjectInitializer& ObjectInitializer )
 	: Super( ObjectInitializer )
 {
+	bSupportsInfiniteRange = true;
 	EvalOptions.EnableAndSetCompletionMode
 		(GetLinkerCustomVersion(FSequencerObjectVersion::GUID) < FSequencerObjectVersion::WhenFinishedDefaultsToRestoreState ? 
 			EMovieSceneCompletionMode::KeepState : 
@@ -29,36 +31,117 @@ UMovieSceneParameterSection::UMovieSceneParameterSection( const FObjectInitializ
 			EMovieSceneCompletionMode::ProjectDefault);
 }
 
-void UMovieSceneParameterSection::AddScalarParameterKey( FName InParameterName, float InTime, float InValue )
+void UMovieSceneParameterSection::Serialize(FArchive& Ar)
 {
-	FRichCurve* ExistingCurve = nullptr;
+	Super::Serialize(Ar);
+
+	if (Ar.IsLoading())
+	{
+		ReconstructChannelProxy();
+	}
+}
+
+void UMovieSceneParameterSection::ReconstructChannelProxy()
+{
+	FMovieSceneChannelData Channels;
+
+#if WITH_EDITOR
+
+	for ( FScalarParameterNameAndCurve& Scalar : GetScalarParameterNamesAndCurves() )
+	{
+		FMovieSceneChannelEditorData EditorData(Scalar.ParameterName, FText::FromName(Scalar.ParameterName));
+		// Prevent single channels from collapsing to the track node
+		EditorData.bCanCollapseToTrack = false;
+		Channels.Add(Scalar.ParameterCurve, EditorData, TMovieSceneExternalValue<float>());
+	}
+	for ( FVectorParameterNameAndCurves& Vector : GetVectorParameterNamesAndCurves() )
+	{
+		FString ParameterString = Vector.ParameterName.ToString();
+		FText Group = FText::FromString(ParameterString);
+
+		Channels.Add(Vector.XCurve, FMovieSceneChannelEditorData(*(ParameterString + TEXT(".X")), FCommonChannelData::ChannelX, Group), TMovieSceneExternalValue<float>());
+		Channels.Add(Vector.YCurve, FMovieSceneChannelEditorData(*(ParameterString + TEXT(".Y")), FCommonChannelData::ChannelY, Group), TMovieSceneExternalValue<float>());
+		Channels.Add(Vector.ZCurve, FMovieSceneChannelEditorData(*(ParameterString + TEXT(".Z")), FCommonChannelData::ChannelZ, Group), TMovieSceneExternalValue<float>());
+	}
+	for ( FColorParameterNameAndCurves& Color : GetColorParameterNamesAndCurves() )
+	{
+		FString ParameterString = Color.ParameterName.ToString();
+		FText Group = FText::FromString(ParameterString);
+
+		FMovieSceneChannelEditorData RChannelData(*(ParameterString + TEXT("R")), FCommonChannelData::ChannelR, Group);
+		RChannelData.SortOrder = 0;
+		RChannelData.Color = FCommonChannelData::RedChannelColor;
+
+		FMovieSceneChannelEditorData GChannelData(*(ParameterString + TEXT("G")), FCommonChannelData::ChannelG, Group);
+		GChannelData.SortOrder = 1;
+		GChannelData.Color = FCommonChannelData::GreenChannelColor;
+
+		FMovieSceneChannelEditorData BChannelData(*(ParameterString + TEXT("B")), FCommonChannelData::ChannelB, Group);
+		BChannelData.SortOrder = 2;
+		BChannelData.Color = FCommonChannelData::BlueChannelColor;
+
+		FMovieSceneChannelEditorData AChannelData(*(ParameterString + TEXT("A")), FCommonChannelData::ChannelA, Group);
+		AChannelData.SortOrder = 3;
+
+		Channels.Add(Color.RedCurve,   RChannelData, TMovieSceneExternalValue<float>());
+		Channels.Add(Color.GreenCurve, GChannelData, TMovieSceneExternalValue<float>());
+		Channels.Add(Color.BlueCurve,  BChannelData, TMovieSceneExternalValue<float>());
+		Channels.Add(Color.AlphaCurve, AChannelData, TMovieSceneExternalValue<float>());
+	}
+
+#else
+
+	for ( FScalarParameterNameAndCurve& Scalar : GetScalarParameterNamesAndCurves() )
+	{
+		Channels.Add(Scalar.ParameterCurve);
+	}
+	for ( FVectorParameterNameAndCurves& Vector : GetVectorParameterNamesAndCurves() )
+	{
+		Channels.Add(Vector.XCurve);
+		Channels.Add(Vector.YCurve);
+		Channels.Add(Vector.ZCurve);
+	}
+	for ( FColorParameterNameAndCurves& Color : GetColorParameterNamesAndCurves() )
+	{
+		Channels.Add(Color.RedCurve);
+		Channels.Add(Color.GreenCurve);
+		Channels.Add(Color.BlueCurve);
+		Channels.Add(Color.AlphaCurve);
+	}
+
+#endif
+
+	ChannelProxy = MakeShared<FMovieSceneChannelProxy>(MoveTemp(Channels));
+}
+
+void UMovieSceneParameterSection::AddScalarParameterKey( FName InParameterName, FFrameNumber InTime, float InValue )
+{
+	FMovieSceneFloatChannel* ExistingChannel = nullptr;
 	for ( FScalarParameterNameAndCurve& ScalarParameterNameAndCurve : ScalarParameterNamesAndCurves )
 	{
 		if ( ScalarParameterNameAndCurve.ParameterName == InParameterName )
 		{
-			ExistingCurve = &ScalarParameterNameAndCurve.ParameterCurve;
+			ExistingChannel = &ScalarParameterNameAndCurve.ParameterCurve;
 			break;
 		}
 	}
-	if ( ExistingCurve == nullptr )
+	if ( ExistingChannel == nullptr )
 	{
-		int32 NewIndex = ScalarParameterNamesAndCurves.Add( FScalarParameterNameAndCurve( InParameterName ) );
-		ScalarParameterNamesAndCurves[NewIndex].Index = ScalarParameterNamesAndCurves.Num() + VectorParameterNamesAndCurves.Num() - 1;
-		ExistingCurve = &ScalarParameterNamesAndCurves[NewIndex].ParameterCurve;
-	}
-	ExistingCurve->AddKey(InTime, InValue);
+		const int32 NewIndex = ScalarParameterNamesAndCurves.Add( FScalarParameterNameAndCurve( InParameterName ) );
+		ExistingChannel = &ScalarParameterNamesAndCurves[NewIndex].ParameterCurve;
 
-	if (GetStartTime() > InTime)
-	{
-		SetStartTime(InTime);
+		ReconstructChannelProxy();
 	}
-	if (GetEndTime() < InTime)
+
+	ExistingChannel->AddCubicKey(InTime, InValue);
+
+	if (TryModify())
 	{
-		SetEndTime(InTime);
+		SetRange(TRange<FFrameNumber>::Hull(TRange<FFrameNumber>(InTime), GetRange()));
 	}
 }
 
-void UMovieSceneParameterSection::AddVectorParameterKey( FName InParameterName, float InTime, FVector InValue )
+void UMovieSceneParameterSection::AddVectorParameterKey( FName InParameterName, FFrameNumber InTime, FVector InValue )
 {
 	FVectorParameterNameAndCurves* ExistingCurves = nullptr;
 	for ( FVectorParameterNameAndCurves& VectorParameterNameAndCurve : VectorParameterNamesAndCurves )
@@ -72,24 +155,22 @@ void UMovieSceneParameterSection::AddVectorParameterKey( FName InParameterName,
 	if ( ExistingCurves == nullptr )
 	{
 		int32 NewIndex = VectorParameterNamesAndCurves.Add( FVectorParameterNameAndCurves( InParameterName ) );
-		VectorParameterNamesAndCurves[NewIndex].Index = VectorParameterNamesAndCurves.Num() + VectorParameterNamesAndCurves.Num() - 1;
 		ExistingCurves = &VectorParameterNamesAndCurves[NewIndex];
-	}
-	ExistingCurves->XCurve.AddKey( InTime, InValue.X );
-	ExistingCurves->YCurve.AddKey( InTime, InValue.Y );
-	ExistingCurves->ZCurve.AddKey( InTime, InValue.Z );
 
-	if (GetStartTime() > InTime)
-	{
-		SetStartTime(InTime);
+		ReconstructChannelProxy();
 	}
-	if (GetEndTime() < InTime)
+
+	ExistingCurves->XCurve.AddCubicKey(InTime, InValue.X);
+	ExistingCurves->YCurve.AddCubicKey(InTime, InValue.Y);
+	ExistingCurves->ZCurve.AddCubicKey(InTime, InValue.Z);
+
+	if (TryModify())
 	{
-		SetEndTime(InTime);
+		SetRange(TRange<FFrameNumber>::Hull(TRange<FFrameNumber>(InTime), GetRange()));
 	}
 }
 
-void UMovieSceneParameterSection::AddColorParameterKey( FName InParameterName, float InTime, FLinearColor InValue )
+void UMovieSceneParameterSection::AddColorParameterKey( FName InParameterName, FFrameNumber InTime, FLinearColor InValue )
 {
 	FColorParameterNameAndCurves* ExistingCurves = nullptr;
 	for ( FColorParameterNameAndCurves& ColorParameterNameAndCurve : ColorParameterNamesAndCurves )
@@ -103,21 +184,19 @@ void UMovieSceneParameterSection::AddColorParameterKey( FName InParameterName, f
 	if ( ExistingCurves == nullptr )
 	{
 		int32 NewIndex = ColorParameterNamesAndCurves.Add( FColorParameterNameAndCurves( InParameterName ) );
-		ColorParameterNamesAndCurves[NewIndex].Index = ColorParameterNamesAndCurves.Num() + ColorParameterNamesAndCurves.Num() - 1;
 		ExistingCurves = &ColorParameterNamesAndCurves[NewIndex];
-	}
-	ExistingCurves->RedCurve.AddKey( InTime, InValue.R );
-	ExistingCurves->GreenCurve.AddKey( InTime, InValue.G );
-	ExistingCurves->BlueCurve.AddKey( InTime, InValue.B );
-	ExistingCurves->AlphaCurve.AddKey( InTime, InValue.A );
 
-	if (GetStartTime() > InTime)
-	{
-		SetStartTime(InTime);
+		ReconstructChannelProxy();
 	}
-	if (GetEndTime() < InTime)
+
+	ExistingCurves->RedCurve.AddCubicKey(   InTime, InValue.R );
+	ExistingCurves->GreenCurve.AddCubicKey( InTime, InValue.G );
+	ExistingCurves->BlueCurve.AddCubicKey(  InTime, InValue.B );
+	ExistingCurves->AlphaCurve.AddCubicKey( InTime, InValue.A );
+
+	if (TryModify())
 	{
-		SetEndTime(InTime);
+		SetRange(TRange<FFrameNumber>::Hull(TRange<FFrameNumber>(InTime), GetRange()));
 	}
 }
 
@@ -128,7 +207,7 @@ bool UMovieSceneParameterSection::RemoveScalarParameter( FName InParameterName )
 		if ( ScalarParameterNamesAndCurves[i].ParameterName == InParameterName )
 		{
 			ScalarParameterNamesAndCurves.RemoveAt(i);
-			UpdateParameterIndicesFromRemoval(i);
+			ReconstructChannelProxy();
 			return true;
 		}
 	}
@@ -142,7 +221,7 @@ bool UMovieSceneParameterSection::RemoveVectorParameter( FName InParameterName )
 		if ( VectorParameterNamesAndCurves[i].ParameterName == InParameterName )
 		{
 			VectorParameterNamesAndCurves.RemoveAt( i );
-			UpdateParameterIndicesFromRemoval( i );
+			ReconstructChannelProxy();
 			return true;
 		}
 	}
@@ -156,16 +235,16 @@ bool UMovieSceneParameterSection::RemoveColorParameter( FName InParameterName )
 		if ( ColorParameterNamesAndCurves[i].ParameterName == InParameterName )
 		{
 			ColorParameterNamesAndCurves.RemoveAt( i );
-			UpdateParameterIndicesFromRemoval( i );
+			ReconstructChannelProxy();
 			return true;
 		}
 	}
 	return false;
 }
 
-TArray<FScalarParameterNameAndCurve>* UMovieSceneParameterSection::GetScalarParameterNamesAndCurves()
+TArray<FScalarParameterNameAndCurve>& UMovieSceneParameterSection::GetScalarParameterNamesAndCurves()
 {
-	return &ScalarParameterNamesAndCurves;
+	return ScalarParameterNamesAndCurves;
 }
 
 const TArray<FScalarParameterNameAndCurve>& UMovieSceneParameterSection::GetScalarParameterNamesAndCurves() const
@@ -173,9 +252,9 @@ const TArray<FScalarParameterNameAndCurve>& UMovieSceneParameterSection::GetScal
 	return ScalarParameterNamesAndCurves;
 }
 
-TArray<FVectorParameterNameAndCurves>* UMovieSceneParameterSection::GetVectorParameterNamesAndCurves()
+TArray<FVectorParameterNameAndCurves>& UMovieSceneParameterSection::GetVectorParameterNamesAndCurves()
 {
-	return &VectorParameterNamesAndCurves;
+	return VectorParameterNamesAndCurves;
 }
 
 const TArray<FVectorParameterNameAndCurves>& UMovieSceneParameterSection::GetVectorParameterNamesAndCurves() const
@@ -183,9 +262,9 @@ const TArray<FVectorParameterNameAndCurves>& UMovieSceneParameterSection::GetVec
 	return VectorParameterNamesAndCurves;
 }
 
-TArray<FColorParameterNameAndCurves>* UMovieSceneParameterSection::GetColorParameterNamesAndCurves()
+TArray<FColorParameterNameAndCurves>& UMovieSceneParameterSection::GetColorParameterNamesAndCurves()
 {
-	return &ColorParameterNamesAndCurves;
+	return ColorParameterNamesAndCurves;
 }
 
 const TArray<FColorParameterNameAndCurves>& UMovieSceneParameterSection::GetColorParameterNamesAndCurves() const
@@ -207,162 +286,4 @@ void UMovieSceneParameterSection::GetParameterNames( TSet<FName>& ParameterNames
 	{
 		ParameterNames.Add( ColorParameterNameAndCurves.ParameterName );
 	}
-}
-
-void UMovieSceneParameterSection::UpdateParameterIndicesFromRemoval( int32 RemovedIndex )
-{
-	for ( FScalarParameterNameAndCurve& ScalarParameterAndCurve : ScalarParameterNamesAndCurves )
-	{
-		if ( ScalarParameterAndCurve.Index > RemovedIndex )
-		{
-			ScalarParameterAndCurve.Index--;
-		}
-	}
-	for ( FVectorParameterNameAndCurves& VectorParameterAndCurves : VectorParameterNamesAndCurves )
-	{
-		if ( VectorParameterAndCurves.Index > RemovedIndex )
-		{
-			VectorParameterAndCurves.Index--;
-		}
-	}
-	for ( FColorParameterNameAndCurves& ColorParameterAndCurves : ColorParameterNamesAndCurves )
-	{
-		if ( ColorParameterAndCurves.Index > RemovedIndex )
-		{
-			ColorParameterAndCurves.Index--;
-		}
-	}
-}
-
-void UMovieSceneParameterSection::GatherCurves(TArray<const FRichCurve*>& OutCurves) const
-{
-	for ( const FScalarParameterNameAndCurve& ScalarParameterNameAndCurve : ScalarParameterNamesAndCurves )
-	{
-		OutCurves.Add(&ScalarParameterNameAndCurve.ParameterCurve);
-	}
-
-	for ( const FVectorParameterNameAndCurves& VectorParameterNameAndCurves : VectorParameterNamesAndCurves )
-	{
-		OutCurves.Add(&VectorParameterNameAndCurves.XCurve);
-		OutCurves.Add(&VectorParameterNameAndCurves.YCurve);
-		OutCurves.Add(&VectorParameterNameAndCurves.ZCurve);
-	}
-
-	for ( const FColorParameterNameAndCurves& ColorParameterNameAndCurve : ColorParameterNamesAndCurves )
-	{
-		OutCurves.Add(&ColorParameterNameAndCurve.RedCurve);
-		OutCurves.Add(&ColorParameterNameAndCurve.BlueCurve);
-		OutCurves.Add(&ColorParameterNameAndCurve.GreenCurve);
-		OutCurves.Add(&ColorParameterNameAndCurve.AlphaCurve);
-	}
-}
-
-void UMovieSceneParameterSection::GatherCurves(TArray<FRichCurve*>& OutCurves)
-{
-	for ( FScalarParameterNameAndCurve& ScalarParameterNameAndCurve : ScalarParameterNamesAndCurves )
-	{
-		OutCurves.Add(&ScalarParameterNameAndCurve.ParameterCurve);
-	}
-
-	for ( FVectorParameterNameAndCurves& VectorParameterNameAndCurves : VectorParameterNamesAndCurves )
-	{
-		OutCurves.Add(&VectorParameterNameAndCurves.XCurve);
-		OutCurves.Add(&VectorParameterNameAndCurves.YCurve);
-		OutCurves.Add(&VectorParameterNameAndCurves.ZCurve);
-	}
-
-	for ( FColorParameterNameAndCurves& ColorParameterNameAndCurve : ColorParameterNamesAndCurves )
-	{
-		OutCurves.Add(&ColorParameterNameAndCurve.RedCurve);
-		OutCurves.Add(&ColorParameterNameAndCurve.BlueCurve);
-		OutCurves.Add(&ColorParameterNameAndCurve.GreenCurve);
-		OutCurves.Add(&ColorParameterNameAndCurve.AlphaCurve);
-	}
-}
-
-
-
-/* UMovieSceneSection overrides
- *****************************************************************************/
-
-void UMovieSceneParameterSection::DilateSection(float DilationFactor, float Origin, TSet<FKeyHandle>& KeyHandles)
-{
-	Super::DilateSection(DilationFactor, Origin, KeyHandles);
-
-	TArray<FRichCurve*> AllCurves;
-	GatherCurves(AllCurves);
-
-	for (auto Curve : AllCurves)
-	{
-		Curve->ScaleCurve(Origin, DilationFactor, KeyHandles);
-	}
-}
-
-void UMovieSceneParameterSection::GetKeyHandles(TSet<FKeyHandle>& OutKeyHandles, TRange<float> TimeRange) const
-{
-	if (!TimeRange.Overlaps(GetRange()))
-	{
-		return;
-	}
-
-	TArray<const FRichCurve*> AllCurves;
-	GatherCurves(AllCurves);
-
-	for (auto Curve : AllCurves)
-	{
-		for (auto It(Curve->GetKeyHandleIterator()); It; ++It)
-		{
-			float Time = Curve->GetKeyTime(It.Key());
-			if (TimeRange.Contains(Time))
-			{
-				OutKeyHandles.Add(It.Key());
-			}
-		}
-	}
-}
-
-
-void UMovieSceneParameterSection::MoveSection(float DeltaPosition, TSet<FKeyHandle>& KeyHandles)
-{
-	Super::MoveSection(DeltaPosition, KeyHandles);
-
-	TArray<FRichCurve*> AllCurves;
-	GatherCurves(AllCurves);
-
-	for (auto Curve : AllCurves)
-	{
-		Curve->ShiftCurve(DeltaPosition, KeyHandles);
-	}
-}
-
-
-TOptional<float> UMovieSceneParameterSection::GetKeyTime( FKeyHandle KeyHandle ) const
-{
-	TArray<const FRichCurve*> AllCurves;
-	GatherCurves( AllCurves );
-
-	for ( auto Curve : AllCurves )
-	{
-		if ( Curve->IsKeyHandleValid( KeyHandle ) )
-		{
-			return TOptional<float>( Curve->GetKeyTime( KeyHandle ) );
-		}
-	}
-	return TOptional<float>();
-}
-
-
-void UMovieSceneParameterSection::SetKeyTime( FKeyHandle KeyHandle, float Time )
-{
-	TArray<FRichCurve*> AllCurves;
-	GatherCurves( AllCurves );
-
-	for ( auto Curve : AllCurves )
-	{
-		if ( Curve->IsKeyHandleValid( KeyHandle ) )
-		{
-			Curve->SetKeyTime( KeyHandle, Time );
-			break;
-		}
-	}
-}
+}
\ No newline at end of file
diff --git a/Engine/Source/Runtime/MovieSceneTracks/Private/Sections/MovieSceneParticleSection.cpp b/Engine/Source/Runtime/MovieSceneTracks/Private/Sections/MovieSceneParticleSection.cpp
--- a/Engine/Source/Runtime/MovieSceneTracks/Private/Sections/MovieSceneParticleSection.cpp
+++ b/Engine/Source/Runtime/MovieSceneTracks/Private/Sections/MovieSceneParticleSection.cpp
@@ -3,75 +3,45 @@
 #include "Sections/MovieSceneParticleSection.h"
 #include "Evaluation/MovieSceneParticleTemplate.h"
 #include "UObject/SequencerObjectVersion.h"
+#include "Channels/MovieSceneChannelProxy.h"
+#include "UObject/Package.h"
 
-UMovieSceneParticleSection::UMovieSceneParticleSection( const FObjectInitializer& ObjectInitializer )
-	: Super( ObjectInitializer )
+FMovieSceneParticleChannel::FMovieSceneParticleChannel()
 {
-	ParticleKeys.SetDefaultValue((int32)EParticleKey::Deactivate);
-	ParticleKeys.SetUseDefaultValueBeforeFirstKey(true);
-	SetIsInfinite(true);
-		
-	EvalOptions.EnableAndSetCompletionMode
-		(GetLinkerCustomVersion(FSequencerObjectVersion::GUID) < FSequencerObjectVersion::WhenFinishedDefaultsToRestoreState ? 
-			EMovieSceneCompletionMode::KeepState : 
-			GetLinkerCustomVersion(FSequencerObjectVersion::GUID) < FSequencerObjectVersion::WhenFinishedDefaultsToProjectDefault ? 
-			EMovieSceneCompletionMode::RestoreState : 
-			EMovieSceneCompletionMode::ProjectDefault);
-}
-
-void UMovieSceneParticleSection::AddKey( float Time, EParticleKey::Type KeyType )
-{
-	ParticleKeys.AddKey( Time, (int32)KeyType );
+	SetEnum(FindObject<UEnum>(ANY_PACKAGE, TEXT("EParticleKey")));
 }
 
-void UMovieSceneParticleSection::MoveSection( float DeltaPosition, TSet<FKeyHandle>& KeyHandles )
+uint32 FMovieSceneParticleChannel::GetChannelID()
 {
-	Super::MoveSection( DeltaPosition, KeyHandles );
-
-	ParticleKeys.ShiftCurve( DeltaPosition, KeyHandles );
+	static const uint32 ID = FMovieSceneChannelEntry::RegisterNewID();
+	return ID;
 }
 
-void UMovieSceneParticleSection::DilateSection( float DilationFactor, float Origin, TSet<FKeyHandle>& KeyHandles )
+UMovieSceneParticleSection::UMovieSceneParticleSection( const FObjectInitializer& ObjectInitializer )
+	: Super( ObjectInitializer )
 {
-	Super::DilateSection( DilationFactor, Origin, KeyHandles );
+#if WITH_EDITORONLY_DATA
+	bIsInfinite_DEPRECATED = true;
+#endif
 
-	ParticleKeys.ScaleCurve( Origin, DilationFactor, KeyHandles );
-}
+	SetRange(TRange<FFrameNumber>::All());
 
-void UMovieSceneParticleSection::GetKeyHandles( TSet<FKeyHandle>& OutKeyHandles, TRange<float> TimeRange ) const
-{
-	if (!TimeRange.Overlaps(GetRange()))
-	{
-		return;
-	}
+	EvalOptions.EnableAndSetCompletionMode
+		(GetLinkerCustomVersion(FSequencerObjectVersion::GUID) < FSequencerObjectVersion::WhenFinishedDefaultsToRestoreState ? 
+			EMovieSceneCompletionMode::KeepState : 
+			GetLinkerCustomVersion(FSequencerObjectVersion::GUID) < FSequencerObjectVersion::WhenFinishedDefaultsToProjectDefault ? 
+			EMovieSceneCompletionMode::RestoreState : 
+			EMovieSceneCompletionMode::ProjectDefault);
 
-	for ( auto It( ParticleKeys.GetKeyHandleIterator() ); It; ++It )
-	{
-		float Time = ParticleKeys.GetKeyTime( It.Key() );
-		if (TimeRange.Contains(Time))
-		{
-			OutKeyHandles.Add( It.Key() );
-		}
-	}
-}
+#if WITH_EDITOR
 
+	ChannelProxy = MakeShared<FMovieSceneChannelProxy>(ParticleKeys, FMovieSceneChannelEditorData(), TMovieSceneExternalValue<uint8>());
 
-TOptional<float> UMovieSceneParticleSection::GetKeyTime( FKeyHandle KeyHandle ) const
-{
-	if ( ParticleKeys.IsKeyHandleValid( KeyHandle ) )
-	{
-		return TOptional<float>( ParticleKeys.GetKeyTime( KeyHandle ) );
-	}
-	return TOptional<float>();
-}
+#else
 
+	ChannelProxy = MakeShared<FMovieSceneChannelProxy>(ParticleKeys);
 
-void UMovieSceneParticleSection::SetKeyTime( FKeyHandle KeyHandle, float Time )
-{
-	if ( ParticleKeys.IsKeyHandleValid( KeyHandle ) )
-	{
-		ParticleKeys.SetKeyTime( KeyHandle, Time );
-	}
+#endif
 }
 
 FMovieSceneEvalTemplatePtr UMovieSceneParticleSection::GenerateTemplate() const
diff --git a/Engine/Source/Runtime/MovieSceneTracks/Private/Sections/MovieSceneSkeletalAnimationSection.cpp b/Engine/Source/Runtime/MovieSceneTracks/Private/Sections/MovieSceneSkeletalAnimationSection.cpp
--- a/Engine/Source/Runtime/MovieSceneTracks/Private/Sections/MovieSceneSkeletalAnimationSection.cpp
+++ b/Engine/Source/Runtime/MovieSceneTracks/Private/Sections/MovieSceneSkeletalAnimationSection.cpp
@@ -1,11 +1,13 @@
 // Copyright 1998-2018 Epic Games, Inc. All Rights Reserved.
 
 #include "Sections/MovieSceneSkeletalAnimationSection.h"
+#include "Channels/MovieSceneChannelProxy.h"
 #include "Animation/AnimSequence.h"
 #include "Evaluation/MovieSceneSkeletalAnimationTemplate.h"
 #include "Logging/MessageLog.h"
 #include "MovieScene.h"
 #include "UObject/SequencerObjectVersion.h"
+#include "MovieSceneTimeHelpers.h"
 
 #define LOCTEXT_NAMESPACE "MovieSceneSkeletalAnimationSection"
 
@@ -22,7 +24,7 @@ FMovieSceneSkeletalAnimationParams::FMovieSceneSkeletalAnimationParams()
 	PlayRate = 1.f;
 	bReverse = false;
 	SlotName = DefaultSlotName;
-	Weight.SetDefaultValue(1.f);
+	Weight.SetDefault(1.f);
 }
 
 UMovieSceneSkeletalAnimationSection::UMovieSceneSkeletalAnimationSection( const FObjectInitializer& ObjectInitializer )
@@ -43,10 +45,25 @@ UMovieSceneSkeletalAnimationSection::UMovieSceneSkeletalAnimationSection( const
 			EMovieSceneCompletionMode::ProjectDefault);
 
 #if WITH_EDITOR
+
 	PreviousPlayRate = Params.PlayRate;
+
+	static FMovieSceneChannelEditorData EditorData("Weight", LOCTEXT("WeightChannelName", "Weight"));
+	EditorData.bCanCollapseToTrack = false;
+	ChannelProxy = MakeShared<FMovieSceneChannelProxy>(Params.Weight, EditorData, TMovieSceneExternalValue<float>());
+
+#else
+
+	ChannelProxy = MakeShared<FMovieSceneChannelProxy>(Params.Weight);
+
 #endif
 }
 
+TOptional<FFrameTime> UMovieSceneSkeletalAnimationSection::GetOffsetTime() const
+{
+	return TOptional<FFrameTime>(Params.StartOffset * GetTypedOuter<UMovieScene>()->GetFrameResolution());
+}
+
 void UMovieSceneSkeletalAnimationSection::Serialize(FArchive& Ar)
 {
 	Ar.UsingCustomVersion(FSequencerObjectVersion::GUID);
@@ -131,27 +148,10 @@ FMovieSceneEvalTemplatePtr UMovieSceneSkeletalAnimationSection::GenerateTemplate
 	return FMovieSceneSkeletalAnimationSectionTemplate(*this);
 }
 
-void UMovieSceneSkeletalAnimationSection::MoveSection( float DeltaTime, TSet<FKeyHandle>& KeyHandles )
-{
-	Super::MoveSection(DeltaTime, KeyHandles);
-
-	Params.Weight.ShiftCurve(DeltaTime, KeyHandles);
-}
-
-
-void UMovieSceneSkeletalAnimationSection::DilateSection( float DilationFactor, float Origin, TSet<FKeyHandle>& KeyHandles )
-{
-	Params.PlayRate /= DilationFactor;
-
-	Super::DilateSection(DilationFactor, Origin, KeyHandles);
-	
-	Params.Weight.ScaleCurve(Origin, DilationFactor, KeyHandles);
-}
-
-UMovieSceneSection* UMovieSceneSkeletalAnimationSection::SplitSection(float SplitTime)
+UMovieSceneSection* UMovieSceneSkeletalAnimationSection::SplitSection(FFrameNumber SplitTime)
 {
 	float AnimPlayRate = FMath::IsNearlyZero(Params.PlayRate) ? 1.0f : Params.PlayRate;
-	float AnimPosition = (SplitTime - GetStartTime()) * AnimPlayRate;
+	float AnimPosition = (SplitTime - GetRange().GetLowerBoundValue()) / GetTypedOuter<UMovieScene>()->GetFrameResolution() * AnimPlayRate;
 	float SeqLength = Params.GetSequenceLength() - (Params.StartOffset + Params.EndOffset);
 
 	float NewOffset = FMath::Fmod(AnimPosition, SeqLength);
@@ -167,41 +167,27 @@ UMovieSceneSection* UMovieSceneSkeletalAnimationSection::SplitSection(float Spli
 }
 
 
-void UMovieSceneSkeletalAnimationSection::GetKeyHandles(TSet<FKeyHandle>& OutKeyHandles, TRange<float> TimeRange) const
-{
-	if (!TimeRange.Overlaps(GetRange()))
-	{
-		return;
-	}
-
-	for (auto It(Params.Weight.GetKeyHandleIterator()); It; ++It)
-	{
-		float Time = Params.Weight.GetKeyTime(It.Key());
-		if (TimeRange.Contains(Time))
-		{
-			OutKeyHandles.Add(It.Key());
-		}
-	}
-}
-
-
-void UMovieSceneSkeletalAnimationSection::GetSnapTimes(TArray<float>& OutSnapTimes, bool bGetSectionBorders) const
+void UMovieSceneSkeletalAnimationSection::GetSnapTimes(TArray<FFrameNumber>& OutSnapTimes, bool bGetSectionBorders) const
 {
 	Super::GetSnapTimes(OutSnapTimes, bGetSectionBorders);
 
-	float CurrentTime = GetStartTime();
-	float AnimPlayRate = FMath::IsNearlyZero(Params.PlayRate) ? 1.0f : Params.PlayRate;
-	float SeqLength = (Params.GetSequenceLength() - (Params.StartOffset + Params.EndOffset)) / AnimPlayRate;
+	const FFrameRate   FrameRate  = GetTypedOuter<UMovieScene>()->GetFrameResolution();
+	const FFrameNumber StartFrame = GetInclusiveStartFrame();
+	const FFrameNumber EndFrame   = GetExclusiveEndFrame() - 1; // -1 because we don't need to add the end frame twice
+
+	const float AnimPlayRate     = FMath::IsNearlyZero(Params.PlayRate) ? 1.0f : Params.PlayRate;
+	const float SeqLengthSeconds = (Params.GetSequenceLength() - (Params.StartOffset + Params.EndOffset)) / AnimPlayRate;
 
-	// Snap to the repeat times
-	while (CurrentTime <= GetEndTime() && !FMath::IsNearlyZero(SeqLength, KINDA_SMALL_NUMBER) && SeqLength > 0)
+	FFrameTime SequenceFrameLength = SeqLengthSeconds * FrameRate;
+	if (SequenceFrameLength.FrameNumber > 1)
 	{
-		if (CurrentTime >= GetStartTime())
+		// Snap to the repeat times
+		FFrameTime CurrentTime = StartFrame;
+		while (CurrentTime < EndFrame)
 		{
-			OutSnapTimes.Add(CurrentTime);
+			OutSnapTimes.Add(CurrentTime.FrameNumber);
+			CurrentTime += SequenceFrameLength;
 		}
-
-		CurrentTime += SeqLength;
 	}
 }
 
@@ -224,10 +210,9 @@ void UMovieSceneSkeletalAnimationSection::PostEditChangeProperty(FPropertyChange
 
 		if (!FMath::IsNearlyZero(NewPlayRate))
 		{
-			float CurrentDuration = GetEndTime() - GetStartTime();
+			float CurrentDuration = MovieScene::DiscreteSize(GetRange()) * PreviousPlayRate;
 			float NewDuration = CurrentDuration * (PreviousPlayRate / NewPlayRate);
-			float NewEndTime = GetStartTime() + NewDuration;
-			SetEndTime(NewEndTime);
+			SetEndFrame( GetInclusiveStartFrame() + FMath::FloorToInt(NewDuration * NewPlayRate) );
 
 			PreviousPlayRate = NewPlayRate;
 		}
diff --git a/Engine/Source/Runtime/MovieSceneTracks/Private/Sections/MovieSceneSlomoSection.cpp b/Engine/Source/Runtime/MovieSceneTracks/Private/Sections/MovieSceneSlomoSection.cpp
--- a/Engine/Source/Runtime/MovieSceneTracks/Private/Sections/MovieSceneSlomoSection.cpp
+++ b/Engine/Source/Runtime/MovieSceneTracks/Private/Sections/MovieSceneSlomoSection.cpp
@@ -10,9 +10,13 @@
 UMovieSceneSlomoSection::UMovieSceneSlomoSection()
 	: UMovieSceneFloatSection()
 {
-	SetIsInfinite(true);
-	GetFloatCurve().SetDefaultValue(1.0f);
+#if WITH_EDITORONLY_DATA
+	bIsInfinite_DEPRECATED = true;
+#endif
 
+	bSupportsInfiniteRange = true;
+	SetRange(TRange<FFrameNumber>::All());
+	FloatCurve.SetDefault(1.f);
 	EvalOptions.EnableAndSetCompletionMode
 		(GetLinkerCustomVersion(FSequencerObjectVersion::GUID) < FSequencerObjectVersion::WhenFinishedDefaultsToProjectDefault ? 
 			EMovieSceneCompletionMode::RestoreState : 
diff --git a/Engine/Source/Runtime/MovieSceneTracks/Private/Sections/MovieSceneSpawnSection.cpp b/Engine/Source/Runtime/MovieSceneTracks/Private/Sections/MovieSceneSpawnSection.cpp
--- a/Engine/Source/Runtime/MovieSceneTracks/Private/Sections/MovieSceneSpawnSection.cpp
+++ b/Engine/Source/Runtime/MovieSceneTracks/Private/Sections/MovieSceneSpawnSection.cpp
@@ -11,4 +11,6 @@ UMovieSceneSpawnSection::UMovieSceneSpawnSection(const FObjectInitializer& Init)
 		(GetLinkerCustomVersion(FSequencerObjectVersion::GUID) < FSequencerObjectVersion::WhenFinishedDefaultsToProjectDefault ? 
 			EMovieSceneCompletionMode::RestoreState : 
 			EMovieSceneCompletionMode::ProjectDefault);
+
+	BoolCurve.SetDefault(true);
 }
diff --git a/Engine/Source/Runtime/MovieSceneTracks/Private/Sections/MovieSceneStringChannel.cpp b/Engine/Source/Runtime/MovieSceneTracks/Private/Sections/MovieSceneStringChannel.cpp
new file mode 100644
--- /dev/null
+++ b/Engine/Source/Runtime/MovieSceneTracks/Private/Sections/MovieSceneStringChannel.cpp
@@ -0,0 +1,57 @@
+// Copyright 1998-2018 Epic Games, Inc. All Rights Reserved.
+
+#include "Sections/MovieSceneStringChannel.h"
+#include "Curves/StringCurve.h"
+#include "MovieSceneFwd.h"
+#include "Channels/MovieSceneChannelProxy.h"
+#include "MovieSceneFrameMigration.h"
+
+uint32 FMovieSceneStringChannel::GetChannelID()
+{
+	static uint32 ID = FMovieSceneChannelEntry::RegisterNewID();
+	return ID;
+}
+
+const FString* FMovieSceneStringChannel::Evaluate(FFrameTime InTime) const
+{
+	if (Times.Num())
+	{
+		const int32 Index = FMath::Max(0, Algo::UpperBound(Times, InTime.FrameNumber)-1);
+		return &Values[Index];
+	}
+	
+	return bHasDefaultValue ? &DefaultValue : nullptr;
+}
+
+bool FMovieSceneStringChannel::SerializeFromMismatchedTag(const FPropertyTag& Tag, FArchive& Ar)
+{
+	static const FName StringCurveName("StringCurve");
+	if (Tag.Type == NAME_StructProperty && Tag.StructName == StringCurveName)
+	{
+		FStringCurve StringCurve;
+		FStringCurve::StaticStruct()->SerializeItem(Ar, &StringCurve, nullptr);
+
+		FString NewDefault = StringCurve.GetDefaultValue();
+		if (NewDefault.Len())
+		{
+			bHasDefaultValue = true;
+			DefaultValue = MoveTemp(NewDefault);
+		}
+
+		FFrameRate LegacyFrameRate = GetLegacyConversionFrameRate();
+
+		Times.Reserve(StringCurve.GetNumKeys());
+		Values.Reserve(StringCurve.GetNumKeys());
+		int32 Index = 0;
+		for (const FStringCurveKey& Key : StringCurve.GetKeys())
+		{
+			FFrameNumber KeyTime = UpgradeLegacyMovieSceneTime(nullptr, LegacyFrameRate, Key.Time);
+
+			FString Val(Key.Value);
+			ConvertInsertAndSort<FString>(Index++, KeyTime, Val, Times, Values);
+		}
+		return true;
+	}
+
+	return false;
+}
diff --git a/Engine/Source/Runtime/MovieSceneTracks/Private/Sections/MovieSceneStringSection.cpp b/Engine/Source/Runtime/MovieSceneTracks/Private/Sections/MovieSceneStringSection.cpp
--- a/Engine/Source/Runtime/MovieSceneTracks/Private/Sections/MovieSceneStringSection.cpp
+++ b/Engine/Source/Runtime/MovieSceneTracks/Private/Sections/MovieSceneStringSection.cpp
@@ -2,116 +2,26 @@
 
 #include "Sections/MovieSceneStringSection.h"
 #include "UObject/SequencerObjectVersion.h"
+#include "Channels/MovieSceneChannelProxy.h"
 
 UMovieSceneStringSection::UMovieSceneStringSection( const FObjectInitializer& ObjectInitializer )
 	: Super( ObjectInitializer )
-{ 
+{
+	bSupportsInfiniteRange = true;
 	EvalOptions.EnableAndSetCompletionMode
 		(GetLinkerCustomVersion(FSequencerObjectVersion::GUID) < FSequencerObjectVersion::WhenFinishedDefaultsToRestoreState ? 
 			EMovieSceneCompletionMode::KeepState : 
 			GetLinkerCustomVersion(FSequencerObjectVersion::GUID) < FSequencerObjectVersion::WhenFinishedDefaultsToProjectDefault ? 
 			EMovieSceneCompletionMode::RestoreState : 
 			EMovieSceneCompletionMode::ProjectDefault);
-}
-
-/* IKeyframeSection interface
- *****************************************************************************/
 
-FString UMovieSceneStringSection::Eval(float Time, const FString& DefaultString) const
-{
-	return StringCurve.Eval(Time, DefaultString);
-}
+#if WITH_EDITOR
 
+	ChannelProxy = MakeShared<FMovieSceneChannelProxy>(StringCurve, FMovieSceneChannelEditorData(), TMovieSceneExternalValue<FString>::Make());
 
-/* IKeyframeSection interface
- *****************************************************************************/
-
-bool UMovieSceneStringSection::NewKeyIsNewData(float Time, const FString& Key) const
-{
-	return (StringCurve.Eval(Time, FString()) != Key);
-}
+#else
 
+	ChannelProxy = MakeShared<FMovieSceneChannelProxy>(StringCurve);
 
-bool UMovieSceneStringSection::HasKeys(const FString& Key) const
-{
-	return (StringCurve.GetNumKeys() > 0);
-}
-
-void UMovieSceneStringSection::AddKey( float Time, const FString& Key, EMovieSceneKeyInterpolation KeyInterpolation )
-{
-	if (TryModify())
-	{
-		StringCurve.UpdateOrAddKey(Time, Key);
-	}
-}
-
-void UMovieSceneStringSection::SetDefault(const FString& Value)
-{
-	if (StringCurve.DefaultValue.Compare(Value) != 0 && TryModify())
-	{
-		StringCurve.SetDefaultValue(Value);
-	}
-}
-
-void UMovieSceneStringSection::ClearDefaults()
-{
-	StringCurve.ClearDefaultValue();
-}
-
-
-/* UMovieSceneSection interface
- *****************************************************************************/
-
-void UMovieSceneStringSection::DilateSection(float DilationFactor, float Origin, TSet<FKeyHandle>& KeyHandles)
-{
-	Super::DilateSection(DilationFactor, Origin, KeyHandles);
-
-	StringCurve.ScaleCurve(Origin, DilationFactor, KeyHandles);
-}
-
-
-void UMovieSceneStringSection::GetKeyHandles(TSet<FKeyHandle>& OutKeyHandles, TRange<float> TimeRange) const
-{
-	if (!TimeRange.Overlaps(GetRange()))
-	{
-		return;
-	}
-
-	for (auto It(StringCurve.GetKeyHandleIterator()); It; ++It)
-	{
-		float Time = StringCurve.GetKeyTime(It.Key());
-
-		if (TimeRange.Contains(Time))
-		{
-			OutKeyHandles.Add(It.Key());
-		}
-	}
-}
-
-
-TOptional<float> UMovieSceneStringSection::GetKeyTime(FKeyHandle KeyHandle) const
-{
-	if (StringCurve.IsKeyHandleValid(KeyHandle))
-	{
-		return TOptional<float>(StringCurve.GetKeyTime(KeyHandle));
-	}
-
-	return TOptional<float>();
-}
-
-
-void UMovieSceneStringSection::MoveSection(float DeltaPosition, TSet<FKeyHandle>& KeyHandles)
-{
-	Super::MoveSection(DeltaPosition, KeyHandles);
-
-	StringCurve.ShiftCurve(DeltaPosition, KeyHandles);
-}
-
-
-void UMovieSceneStringSection::SetKeyTime(FKeyHandle KeyHandle, float Time)
-{
-	if (StringCurve.IsKeyHandleValid(KeyHandle))
-	{
-		StringCurve.SetKeyTime(KeyHandle, Time);
-	}
-}
+#endif
+}
\ No newline at end of file
diff --git a/Engine/Source/Runtime/MovieSceneTracks/Private/Sections/MovieSceneVectorSection.cpp b/Engine/Source/Runtime/MovieSceneTracks/Private/Sections/MovieSceneVectorSection.cpp
--- a/Engine/Source/Runtime/MovieSceneTracks/Private/Sections/MovieSceneVectorSection.cpp
+++ b/Engine/Source/Runtime/MovieSceneTracks/Private/Sections/MovieSceneVectorSection.cpp
@@ -3,27 +3,79 @@
 #include "Sections/MovieSceneVectorSection.h"
 #include "UObject/StructOnScope.h"
 #include "UObject/SequencerObjectVersion.h"
+#include "Channels/MovieSceneChannelProxy.h"
 
 /* FMovieSceneVectorKeyStruct interface
  *****************************************************************************/
 
-void FMovieSceneVectorKeyStructBase::PropagateChanges(const FPropertyChangedEvent& ChangeEvent)
+#if WITH_EDITOR
+struct FVectorSectionEditorData
 {
-	for (int32 Index = 0; Index < GetChannelsUsed(); ++Index)
+	FVectorSectionEditorData(int32 NumChannels)
 	{
-		if (Keys[Index] == nullptr)
+		CommonData[0].SetIdentifiers("Vector.X", FCommonChannelData::ChannelX);
+		CommonData[0].SortOrder = 0;
+		CommonData[0].Color = FCommonChannelData::RedChannelColor;
+
+		CommonData[1].SetIdentifiers("Vector.Y", FCommonChannelData::ChannelY);
+		CommonData[1].SortOrder = 1;
+		CommonData[1].Color = FCommonChannelData::GreenChannelColor;
+
+		CommonData[2].SetIdentifiers("Vector.Z", FCommonChannelData::ChannelZ);
+		CommonData[2].SortOrder = 2;
+		CommonData[2].Color = FCommonChannelData::BlueChannelColor;
+
+		CommonData[3].SetIdentifiers("Vector.W", FCommonChannelData::ChannelW);
+		CommonData[3].SortOrder = 3;
+
+		ExternalValues[0].OnGetExternalValue = [NumChannels](UObject& InObject, FTrackInstancePropertyBindings* Bindings) { return ExtractChannelX(InObject, Bindings, NumChannels); };
+		ExternalValues[1].OnGetExternalValue = [NumChannels](UObject& InObject, FTrackInstancePropertyBindings* Bindings) { return ExtractChannelY(InObject, Bindings, NumChannels); };
+		ExternalValues[2].OnGetExternalValue = [NumChannels](UObject& InObject, FTrackInstancePropertyBindings* Bindings) { return ExtractChannelZ(InObject, Bindings, NumChannels); };
+		ExternalValues[3].OnGetExternalValue = [NumChannels](UObject& InObject, FTrackInstancePropertyBindings* Bindings) { return ExtractChannelW(InObject, Bindings, NumChannels); };
+	}
+
+	static FVector4 GetPropertyValue(UObject& InObject, FTrackInstancePropertyBindings& Bindings, int32 NumChannels)
+	{
+		if (NumChannels == 2)
 		{
-			if (Curves[Index] != nullptr)
-			{
-				Curves[Index]->SetDefaultValue(GetPropertyChannelByIndex(Index));
-			}
+			FVector2D Vector = Bindings.GetCurrentValue<FVector2D>(InObject);
+			return FVector4(Vector.X, Vector.Y, 0.f, 0.f);
+		}
+		else if (NumChannels == 3)
+		{
+			FVector Vector = Bindings.GetCurrentValue<FVector>(InObject);
+			return FVector4(Vector.X, Vector.Y, Vector.Z, 0.f);
 		}
 		else
 		{
-			Keys[Index]->Value = GetPropertyChannelByIndex(Index);
-			Keys[Index]->Time = Time;
+			return Bindings.GetCurrentValue<FVector>(InObject);
 		}
 	}
+
+	static TOptional<float> ExtractChannelX(UObject& InObject, FTrackInstancePropertyBindings* Bindings, int32 NumChannels)
+	{
+		return Bindings ? GetPropertyValue(InObject, *Bindings, NumChannels).X : TOptional<float>();
+	}
+	static TOptional<float> ExtractChannelY(UObject& InObject, FTrackInstancePropertyBindings* Bindings, int32 NumChannels)
+	{
+		return Bindings ? GetPropertyValue(InObject, *Bindings, NumChannels).Y : TOptional<float>();
+	}
+	static TOptional<float> ExtractChannelZ(UObject& InObject, FTrackInstancePropertyBindings* Bindings, int32 NumChannels)
+	{
+		return Bindings ? GetPropertyValue(InObject, *Bindings, NumChannels).Z : TOptional<float>();
+	}
+	static TOptional<float> ExtractChannelW(UObject& InObject, FTrackInstancePropertyBindings* Bindings, int32 NumChannels)
+	{
+		return Bindings ? GetPropertyValue(InObject, *Bindings, NumChannels).W : TOptional<float>();
+	}
+	FMovieSceneChannelEditorData    CommonData[4];
+	TMovieSceneExternalValue<float> ExternalValues[4];
+};
+#endif
+
+void FMovieSceneVectorKeyStructBase::PropagateChanges(const FPropertyChangedEvent& ChangeEvent)
+{
+	KeyStructInterop.Apply(Time);
 }
 
 
@@ -34,6 +86,7 @@ UMovieSceneVectorSection::UMovieSceneVectorSection(const FObjectInitializer& Obj
 	: Super(ObjectInitializer)
 {
 	ChannelsUsed = 0;
+	bSupportsInfiniteRange = true;
 
 	EvalOptions.EnableAndSetCompletionMode
 		(GetLinkerCustomVersion(FSequencerObjectVersion::GUID) < FSequencerObjectVersion::WhenFinishedDefaultsToRestoreState ? 
@@ -44,55 +97,43 @@ UMovieSceneVectorSection::UMovieSceneVectorSection(const FObjectInitializer& Obj
 	BlendType = EMovieSceneBlendType::Absolute;
 }
 
-/* UMovieSceneSection interface
- *****************************************************************************/
-
-void UMovieSceneVectorSection::MoveSection(float DeltaTime, TSet<FKeyHandle>& KeyHandles)
+void UMovieSceneVectorSection::Serialize(FArchive& Ar)
 {
-	check(ChannelsUsed >= 2 && ChannelsUsed <= 4);
-	Super::MoveSection(DeltaTime, KeyHandles);
+	Super::Serialize(Ar);
 
-	for (int32 i = 0; i < ChannelsUsed; ++i)
+	if (Ar.IsLoading())
 	{
-		Curves[i].ShiftCurve(DeltaTime, KeyHandles);
+		RecreateChannelProxy();
 	}
 }
 
-
-void UMovieSceneVectorSection::DilateSection(float DilationFactor, float Origin, TSet<FKeyHandle>& KeyHandles)
+void UMovieSceneVectorSection::RecreateChannelProxy()
 {
-	check(ChannelsUsed >= 2 && ChannelsUsed <= 4);
-	Super::DilateSection(DilationFactor, Origin, KeyHandles);
+	FMovieSceneChannelData Channels;
 
-	for (int32 i = 0; i < ChannelsUsed; ++i)
-	{
-		Curves[i].ScaleCurve(Origin, DilationFactor, KeyHandles);
-	}
-}
+	check(ChannelsUsed <= ARRAY_COUNT(Curves));
 
+#if WITH_EDITOR
 
-void UMovieSceneVectorSection::GetKeyHandles(TSet<FKeyHandle>& OutKeyHandles, TRange<float> TimeRange) const
-{
-	if (!TimeRange.Overlaps(GetRange()))
+	const FVectorSectionEditorData EditorData(ChannelsUsed);
+	for (int32 Index = 0; Index < ChannelsUsed; ++Index)
 	{
-		return;
+		Channels.Add(Curves[Index], EditorData.CommonData[Index], EditorData.ExternalValues[Index]);
 	}
 
-	for (int32 i = 0; i < ChannelsUsed; ++i)
+#else
+
+	for (int32 Index = 0; Index < ChannelsUsed; ++Index)
 	{
-		for (auto It(Curves[i].GetKeyHandleIterator()); It; ++It)
-		{
-			float Time = Curves[i].GetKeyTime(It.Key());
-			if (TimeRange.Contains(Time))
-			{
-				OutKeyHandles.Add(It.Key());
-			}
-		}
+		Channels.Add(Curves[Index]);
 	}
-}
 
+#endif
+
+	ChannelProxy = MakeShared<FMovieSceneChannelProxy>(MoveTemp(Channels));
+}
 
-TSharedPtr<FStructOnScope> UMovieSceneVectorSection::GetKeyStruct(const TArray<FKeyHandle>& KeyHandles)
+TSharedPtr<FStructOnScope> UMovieSceneVectorSection::GetKeyStruct(TArrayView<const FKeyHandle> KeyHandles)
 {
 	TSharedPtr<FStructOnScope> KeyStruct;
 	if (ChannelsUsed == 2)
@@ -112,117 +153,14 @@ TSharedPtr<FStructOnScope> UMovieSceneVectorSection::GetKeyStruct(const TArray<F
 	{
 		FMovieSceneVectorKeyStructBase* Struct = (FMovieSceneVectorKeyStructBase*)KeyStruct->GetStructMemory();
 
-		float FirstValidKeyTime = 0.f;
-		for (int32 Index = 0; Index < Struct->GetChannelsUsed(); ++Index)
+		TArrayView<FMovieSceneFloatChannel*> Channels = ChannelProxy->GetChannels<FMovieSceneFloatChannel>();
+		for (int32 Index = 0; Index < ChannelsUsed; ++Index)
 		{
-			Struct->Curves[Index] = &Curves[Index];
-			Struct->Keys[Index] = Curves[Index].GetFirstMatchingKey(KeyHandles);
-			if (Struct->Keys[Index] != nullptr)
-			{
-				FirstValidKeyTime = Struct->Keys[Index]->Time;
-				Struct->Time = FirstValidKeyTime;
-			}
+			Struct->KeyStructInterop.Add(FMovieSceneChannelValueHelper(ChannelProxy->MakeHandle(Channels[Index]), Struct->GetPropertyChannelByIndex(Index), KeyHandles));
 		}
 
-		for (int32 Index = 0; Index < Struct->GetChannelsUsed(); ++Index)
-		{
-			if (Struct->Keys[Index] == nullptr)
-			{
-				Struct->SetPropertyChannelByIndex(Index, Struct->Curves[Index]->Eval(FirstValidKeyTime));
-			}
-			else
-			{
-				Struct->SetPropertyChannelByIndex(Index, Struct->Keys[Index]->Value);
-			}
-		}
+		Struct->Time = Struct->KeyStructInterop.GetUnifiedKeyTime().Get(0);
 	}
 
 	return KeyStruct;
 }
-
-
-TOptional<float> UMovieSceneVectorSection::GetKeyTime( FKeyHandle KeyHandle ) const
-{
-	for ( auto Curve : Curves )
-	{
-		if ( Curve.IsKeyHandleValid( KeyHandle ) )
-		{
-			return TOptional<float>( Curve.GetKeyTime( KeyHandle ) );
-		}
-	}
-	return TOptional<float>();
-}
-
-
-void UMovieSceneVectorSection::SetKeyTime( FKeyHandle KeyHandle, float Time )
-{
-	for ( auto Curve : Curves )
-	{
-		if ( Curve.IsKeyHandleValid( KeyHandle ) )
-		{
-			Curve.SetKeyTime( KeyHandle, Time );
-			break;
-		}
-	}
-}
-
-
-/* IKeyframeSection interface
- *****************************************************************************/
-
-template<typename CurveType>
-CurveType* GetCurveForChannel(EKeyVectorChannel Channel, CurveType* Curves, int32 ChannelsUsed)
-{
-	switch (Channel)
-	{
-		case EKeyVectorChannel::X:
-			return &Curves[0];
-		case EKeyVectorChannel::Y:
-			return &Curves[1];
-		case EKeyVectorChannel::Z:
-			checkf(ChannelsUsed >= 3, TEXT("Can not get Z channel, it is not in use on this section."));
-			return &Curves[2];
-		case EKeyVectorChannel::W:
-			checkf(ChannelsUsed >= 4, TEXT("Can not get W channel, it is not in use on this section."));
-			return &Curves[3];
-	}
-	checkf(false, TEXT("Invalid channel requested"));
-	return nullptr;
-}
-
-
-void UMovieSceneVectorSection::AddKey(float Time, const FVectorKey& Key, EMovieSceneKeyInterpolation KeyInterpolation)
-{
-	FRichCurve* ChannelCurve = GetCurveForChannel(Key.Channel, Curves, ChannelsUsed);
-	AddKeyToCurve(*ChannelCurve, Time, Key.Value, KeyInterpolation);
-}
-
-
-bool UMovieSceneVectorSection::NewKeyIsNewData(float Time, const FVectorKey& Key) const
-{
-	const FRichCurve* ChannelCurve = GetCurveForChannel(Key.Channel, Curves, ChannelsUsed);
-	return FMath::IsNearlyEqual(ChannelCurve->Eval(Time), Key.Value) == false;
-}
-
-
-bool UMovieSceneVectorSection::HasKeys(const FVectorKey& Key) const
-{
-	const FRichCurve* ChannelCurve = GetCurveForChannel(Key.Channel, Curves, ChannelsUsed);
-	return ChannelCurve->GetNumKeys() > 0;
-}
-
-
-void UMovieSceneVectorSection::SetDefault(const FVectorKey& Key)
-{
-	FRichCurve* ChannelCurve = GetCurveForChannel(Key.Channel, Curves, ChannelsUsed);
-	SetCurveDefault( *ChannelCurve, Key.Value );
-}
-
-
-void UMovieSceneVectorSection::ClearDefaults()
-{
-	for (auto& Curve : Curves)
-	{
-		Curve.ClearDefaultValue();
-	}
-}
diff --git a/Engine/Source/Runtime/MovieSceneTracks/Private/Tracks/MovieScene3DAttachTrack.cpp b/Engine/Source/Runtime/MovieSceneTracks/Private/Tracks/MovieScene3DAttachTrack.cpp
--- a/Engine/Source/Runtime/MovieSceneTracks/Private/Tracks/MovieScene3DAttachTrack.cpp
+++ b/Engine/Source/Runtime/MovieSceneTracks/Private/Tracks/MovieScene3DAttachTrack.cpp
@@ -19,13 +19,13 @@ FMovieSceneEvalTemplatePtr UMovieScene3DAttachTrack::CreateTemplateForSection(co
 	return FMovieScene3DAttachSectionTemplate(*CastChecked<UMovieScene3DAttachSection>(&InSection));
 }
 
-void UMovieScene3DAttachTrack::AddConstraint(float KeyTime, float ConstraintEndTime, const FName SocketName, const FName ComponentName, const FMovieSceneObjectBindingID& ConstraintBindingID)
+void UMovieScene3DAttachTrack::AddConstraint(FFrameNumber KeyTime, int32 Duration, const FName SocketName, const FName ComponentName, const FMovieSceneObjectBindingID& ConstraintBindingID)
 {
 	// add the section
 	UMovieScene3DAttachSection* NewSection = NewObject<UMovieScene3DAttachSection>(this);
 	NewSection->SetFlags(RF_Transactional);
-	NewSection->AddAttach(KeyTime, ConstraintEndTime, ConstraintBindingID);
-	NewSection->InitialPlacement( ConstraintSections, KeyTime, ConstraintEndTime, SupportsMultipleRows() );
+	NewSection->SetAttachTargetID(ConstraintBindingID);
+	NewSection->InitialPlacement(ConstraintSections, KeyTime, Duration, SupportsMultipleRows());
 	NewSection->AttachSocketName = SocketName;
 	NewSection->AttachComponentName = ComponentName;
 
diff --git a/Engine/Source/Runtime/MovieSceneTracks/Private/Tracks/MovieScene3DConstraintTrack.cpp b/Engine/Source/Runtime/MovieSceneTracks/Private/Tracks/MovieScene3DConstraintTrack.cpp
--- a/Engine/Source/Runtime/MovieSceneTracks/Private/Tracks/MovieScene3DConstraintTrack.cpp
+++ b/Engine/Source/Runtime/MovieSceneTracks/Private/Tracks/MovieScene3DConstraintTrack.cpp
@@ -22,11 +22,6 @@ const TArray<UMovieSceneSection*>& UMovieScene3DConstraintTrack::GetAllSections(
 	return ConstraintSections;
 }
 
-void UMovieScene3DConstraintTrack::AddConstraint(float Time, float ConstraintEndTime, const FName SocketName, const FName ComponentName, const FGuid& ConstraintId)
-{
-	AddConstraint(Time, ConstraintEndTime, SocketName, ComponentName, FMovieSceneObjectBindingID(ConstraintId, MovieSceneSequenceID::Root));
-}
-
 
 void UMovieScene3DConstraintTrack::RemoveAllAnimationData()
 {
@@ -57,17 +52,4 @@ bool UMovieScene3DConstraintTrack::IsEmpty() const
 	return ConstraintSections.Num() == 0;
 }
 
-
-TRange<float> UMovieScene3DConstraintTrack::GetSectionBoundaries() const
-{
-	TArray< TRange<float> > Bounds;
-
-	for (int32 i = 0; i < ConstraintSections.Num(); ++i)
-	{
-		Bounds.Add(ConstraintSections[i]->GetRange());
-	}
-
-	return TRange<float>::Hull(Bounds);
-}
-
 #undef LOCTEXT_NAMESPACE
diff --git a/Engine/Source/Runtime/MovieSceneTracks/Private/Tracks/MovieScene3DPathTrack.cpp b/Engine/Source/Runtime/MovieSceneTracks/Private/Tracks/MovieScene3DPathTrack.cpp
--- a/Engine/Source/Runtime/MovieSceneTracks/Private/Tracks/MovieScene3DPathTrack.cpp
+++ b/Engine/Source/Runtime/MovieSceneTracks/Private/Tracks/MovieScene3DPathTrack.cpp
@@ -21,12 +21,12 @@ FMovieSceneEvalTemplatePtr UMovieScene3DPathTrack::CreateTemplateForSection(cons
 }
 
 
-void UMovieScene3DPathTrack::AddConstraint(float KeyTime, float ConstraintEndTime, const FName SocketName, const FName ComponentName, const FMovieSceneObjectBindingID& ConstraintBindingID)
+void UMovieScene3DPathTrack::AddConstraint(FFrameNumber KeyTime, int32 Duration, const FName SocketName, const FName ComponentName, const FMovieSceneObjectBindingID& ConstraintBindingID)
 {
 	UMovieScene3DPathSection* NewSection = NewObject<UMovieScene3DPathSection>(this);
 	{
-		NewSection->AddPath(KeyTime, ConstraintEndTime, ConstraintBindingID);
-		NewSection->InitialPlacement( ConstraintSections, KeyTime, ConstraintEndTime, SupportsMultipleRows() );
+		NewSection->SetPathBindingID( ConstraintBindingID );
+		NewSection->InitialPlacement( ConstraintSections, KeyTime, Duration, SupportsMultipleRows() );
 	}
 
 	ConstraintSections.Add(NewSection);
diff --git a/Engine/Source/Runtime/MovieSceneTracks/Private/Tracks/MovieScene3DTransformTrack.cpp b/Engine/Source/Runtime/MovieSceneTracks/Private/Tracks/MovieScene3DTransformTrack.cpp
--- a/Engine/Source/Runtime/MovieSceneTracks/Private/Tracks/MovieScene3DTransformTrack.cpp
+++ b/Engine/Source/Runtime/MovieSceneTracks/Private/Tracks/MovieScene3DTransformTrack.cpp
@@ -1,11 +1,11 @@
 // Copyright 1998-2018 Epic Games, Inc. All Rights Reserved.
 
 #include "Tracks/MovieScene3DTransformTrack.h"
-#include "MovieSceneCommonHelpers.h"
 #include "Sections/MovieScene3DTransformSection.h"
 #include "Compilation/MovieSceneTemplateInterrogation.h"
 #include "Compilation/MovieSceneSegmentCompiler.h"
 #include "Evaluation/MovieSceneEvaluationTrack.h"
+#include "Channels/MovieSceneChannelProxy.h"
 #include "Algo/BinarySearch.h"
 
 UMovieScene3DTransformTrack::UMovieScene3DTransformTrack( const FObjectInitializer& ObjectInitializer )
@@ -23,97 +23,149 @@ UMovieScene3DTransformTrack::UMovieScene3DTransformTrack( const FObjectInitializ
 	EvalOptions.bEvaluateNearestSection_DEPRECATED = EvalOptions.bCanEvaluateNearestSection = true;
 }
 
-TArray<FTrajectoryKey> UMovieScene3DTransformTrack::GetTrajectoryData(float Time, int32 MaxNumDataPoints) const
+UMovieSceneSection* UMovieScene3DTransformTrack::CreateNewSection()
+{
+	return NewObject<UMovieSceneSection>(this, UMovieScene3DTransformSection::StaticClass(), NAME_None, RF_Transactional);
+}
+
+FMovieSceneInterrogationKey UMovieScene3DTransformTrack::GetInterrogationKey()
+{
+	static FMovieSceneAnimTypeID TypeID = FMovieSceneAnimTypeID::Unique();
+	return TypeID;
+}
+
+
+#if WITH_EDITOR
+
+uint32 GetDistanceFromTo(FFrameNumber TestValue, FFrameNumber TargetValue)
+{
+	static const int32 MaxInt = TNumericLimits<int32>::Max();
+	if (TestValue == TargetValue)
+	{
+		return 0;
+	}
+	else if (TestValue > TargetValue)
+	{
+		return (uint32(TestValue.Value + MaxInt) - uint32(TargetValue.Value + MaxInt));
+	}
+	else
+	{
+		return (uint32(TargetValue.Value + MaxInt) - uint32(TestValue.Value + MaxInt));
+	}
+}
+
+TArray<FTrajectoryKey> UMovieScene3DTransformTrack::GetTrajectoryData(FFrameNumber Time, int32 MaxNumDataPoints) const
 {
 	struct FCurveKeyIterator
 	{
-		FCurveKeyIterator(UMovieScene3DTransformSection* InSection, EMovieSceneTransformChannel InType, const FRichCurve& InCurve, float StartTime)
-			: Section(InSection), Type(InType), SortedKeys(InCurve.GetConstRefOfKeys()), SectionRange(InSection->IsInfinite() ? TRange<float>::All() : InSection->GetRange())
+		FCurveKeyIterator(UMovieScene3DTransformSection* InSection, FMovieSceneFloatChannel* InChannel, FName InChannelName, FFrameNumber StartTime)
+			: Section(InSection), Channel(InChannel->GetInterface()), ChannelName(InChannelName), SectionRange(InSection->GetRange()), CurrentIndex(INDEX_NONE)
 		{
-			CurrentIndex = Algo::LowerBoundBy(SortedKeys, StartTime, &FRichCurveKey::Time, TLess<>());
+			TArrayView<const FFrameNumber> Times = Channel.GetTimes();
+			CurrentIndex = Algo::LowerBound(Times, StartTime);
 
 			bIsLowerBound = false;
-			bIsUpperBound = !SortedKeys.IsValidIndex(CurrentIndex) && SectionRange.GetUpperBound().IsClosed();
+			bIsUpperBound = CurrentIndex == Times.Num() && SectionRange.GetUpperBound().IsClosed();
+		}
+
+		bool IsValid() const
+		{
+			return Channel.GetTimes().IsValidIndex(CurrentIndex);
+		}
+
+		FFrameNumber GetCurrentKeyTime() const
+		{
+			return Channel.GetTimes()[CurrentIndex];
 		}
 
 		FCurveKeyIterator& operator--()
 		{
+			TArrayView<const FFrameNumber> Times = Channel.GetTimes();
+
 			if (bIsLowerBound)
 			{
 				bIsLowerBound = false;
-				CurrentIndex = -1;
+				CurrentIndex = INDEX_NONE;
 			}
 			else
 			{
 				if (bIsUpperBound)
 				{
 					bIsUpperBound = false;
-					CurrentIndex = Algo::LowerBoundBy(SortedKeys, SectionRange.GetUpperBoundValue(), &FRichCurveKey::Time, TLess<>()) - 1;
+					CurrentIndex = Algo::LowerBound(Times, SectionRange.GetUpperBoundValue()) - 1;
 				}
 				else
 				{
 					--CurrentIndex;
 				}
 
-				bIsLowerBound = SectionRange.GetLowerBound().IsClosed() && (!SortedKeys.IsValidIndex(CurrentIndex) || !SectionRange.Contains(SortedKeys[CurrentIndex].Time) );
+				bIsLowerBound = SectionRange.GetLowerBound().IsClosed() && (!Times.IsValidIndex(CurrentIndex) || !SectionRange.Contains(Channel.GetTimes()[CurrentIndex]) );
 			}
 			return *this;
 		}
 
 		FCurveKeyIterator& operator++()
 		{
+			TArrayView<const FFrameNumber> Times = Channel.GetTimes();
+
 			if (bIsUpperBound)
 			{
 				bIsUpperBound = false;
-				CurrentIndex = -1;
+				CurrentIndex = INDEX_NONE;
 			}
 			else
 			{
 				if (bIsLowerBound)
 				{
 					bIsLowerBound = false;
-					CurrentIndex = Algo::UpperBoundBy(SortedKeys, SectionRange.GetLowerBoundValue(), &FRichCurveKey::Time, TLess<>());
+					CurrentIndex = Algo::UpperBound(Times, SectionRange.GetLowerBoundValue());
 				}
 				else 
 				{
 					++CurrentIndex;
 				}
 
-				bIsUpperBound = SectionRange.GetUpperBound().IsClosed() && (!SortedKeys.IsValidIndex(CurrentIndex) || !SectionRange.Contains(SortedKeys[CurrentIndex].Time) );
+				bIsUpperBound = SectionRange.GetUpperBound().IsClosed() && (!Times.IsValidIndex(CurrentIndex) || !SectionRange.Contains(Times[CurrentIndex]) );
 			}
 			return *this;
 		}
 
 		explicit operator bool() const
 		{
-			return ( bIsLowerBound || bIsUpperBound ) || ( SortedKeys.IsValidIndex(CurrentIndex) && SectionRange.Contains(SortedKeys[CurrentIndex].Time) );
+			TArrayView<const FFrameNumber> Times = Channel.GetTimes();
+			return ( bIsLowerBound || bIsUpperBound ) || ( Times.IsValidIndex(CurrentIndex) && SectionRange.Contains(Times[CurrentIndex]) );
 		}
 
-		float GetTime() const
+		FFrameNumber GetTime() const
 		{
-			return bIsLowerBound ? SectionRange.GetLowerBoundValue() : bIsUpperBound ? SectionRange.GetUpperBoundValue() : SortedKeys[CurrentIndex].Time;
+			return bIsLowerBound ? SectionRange.GetLowerBoundValue() : bIsUpperBound ? SectionRange.GetUpperBoundValue() : Channel.GetTimes()[CurrentIndex];
 		}
 
 		ERichCurveInterpMode GetInterpMode() const
 		{
-			return ( bIsLowerBound || bIsUpperBound ) ? RCIM_None : SortedKeys[CurrentIndex].InterpMode.GetValue();
+			return ( bIsLowerBound || bIsUpperBound ) ? RCIM_None : Channel.GetValues()[CurrentIndex].InterpMode.GetValue();
 		}
 
-		EMovieSceneTransformChannel GetType() const
+		UMovieScene3DTransformSection* GetSection() const
 		{
-			return Type;
+			return Section;
 		}
 
-		UMovieScene3DTransformSection* GetSection() const
+		FName GetChannelName() const
 		{
-			return Section;
+			return ChannelName;
+		}
+
+		TOptional<FKeyHandle> GetKeyHandle()
+		{
+			return CurrentIndex == INDEX_NONE ? TOptional<FKeyHandle>() : Channel.GetHandle(CurrentIndex);
 		}
 
 	private:
 		UMovieScene3DTransformSection* Section;
-		EMovieSceneTransformChannel Type;
-		const TArray<FRichCurveKey>& SortedKeys;
-		TRange<float> SectionRange;
+		TMovieSceneChannel<FMovieSceneFloatValue> Channel;
+		FName ChannelName;
+		TRange<FFrameNumber> SectionRange;
 		int32 CurrentIndex;
 		bool bIsUpperBound, bIsLowerBound;
 	};
@@ -126,50 +178,48 @@ TArray<FTrajectoryKey> UMovieScene3DTransformTrack::GetTrajectoryData(float Time
 		UMovieScene3DTransformSection* TransformSection = Cast<UMovieScene3DTransformSection>(Section);
 		if (TransformSection)
 		{
+			TArrayView<FMovieSceneFloatChannel*>           FloatChannels = TransformSection->GetChannelProxy().GetChannels<FMovieSceneFloatChannel>();
+			TArrayView<const FMovieSceneChannelEditorData> EditorData    = TransformSection->GetChannelProxy().GetEditorData<FMovieSceneFloatChannel>();
+
 			EMovieSceneTransformChannel Mask = TransformSection->GetMask().GetChannels();
 			if (EnumHasAnyFlags(Mask, EMovieSceneTransformChannel::TranslationX))
 			{
-				ForwardIters.Emplace(TransformSection, EMovieSceneTransformChannel::TranslationX, TransformSection->GetTranslationCurve(EAxis::X), Time);
-				BackwardIters.Emplace(TransformSection, EMovieSceneTransformChannel::TranslationX, TransformSection->GetTranslationCurve(EAxis::X), Time);
+				ForwardIters.Emplace(TransformSection, FloatChannels[0], EditorData[0].Name, Time);
+				BackwardIters.Emplace(TransformSection, FloatChannels[0], EditorData[0].Name, Time);
 			}
 			if (EnumHasAnyFlags(Mask, EMovieSceneTransformChannel::TranslationY))
 			{
-				ForwardIters.Emplace(TransformSection, EMovieSceneTransformChannel::TranslationY, TransformSection->GetTranslationCurve(EAxis::Y), Time);
-				BackwardIters.Emplace(TransformSection, EMovieSceneTransformChannel::TranslationY, TransformSection->GetTranslationCurve(EAxis::Y), Time);
+				ForwardIters.Emplace(TransformSection, FloatChannels[1], EditorData[1].Name, Time);
+				BackwardIters.Emplace(TransformSection, FloatChannels[1], EditorData[1].Name, Time);
 			}
 			if (EnumHasAnyFlags(Mask, EMovieSceneTransformChannel::TranslationZ))
 			{
-				ForwardIters.Emplace(TransformSection, EMovieSceneTransformChannel::TranslationZ, TransformSection->GetTranslationCurve(EAxis::Z), Time);
-				BackwardIters.Emplace(TransformSection, EMovieSceneTransformChannel::TranslationZ, TransformSection->GetTranslationCurve(EAxis::Z), Time);
+				ForwardIters.Emplace(TransformSection, FloatChannels[2], EditorData[2].Name, Time);
+				BackwardIters.Emplace(TransformSection, FloatChannels[2], EditorData[2].Name, Time);
 			}
 			if (EnumHasAnyFlags(Mask, EMovieSceneTransformChannel::RotationX))
 			{
-				ForwardIters.Emplace(TransformSection, EMovieSceneTransformChannel::RotationX, TransformSection->GetRotationCurve(EAxis::X), Time);
-				BackwardIters.Emplace(TransformSection, EMovieSceneTransformChannel::RotationX, TransformSection->GetRotationCurve(EAxis::X), Time);
+				ForwardIters.Emplace(TransformSection, FloatChannels[3], EditorData[3].Name, Time);
+				BackwardIters.Emplace(TransformSection, FloatChannels[3], EditorData[3].Name, Time);
 			}
 			if (EnumHasAnyFlags(Mask, EMovieSceneTransformChannel::RotationY))
 			{
-				ForwardIters.Emplace(TransformSection, EMovieSceneTransformChannel::RotationY, TransformSection->GetRotationCurve(EAxis::Y), Time);
-				BackwardIters.Emplace(TransformSection, EMovieSceneTransformChannel::RotationY, TransformSection->GetRotationCurve(EAxis::Y), Time);
+				ForwardIters.Emplace(TransformSection, FloatChannels[4], EditorData[4].Name, Time);
+				BackwardIters.Emplace(TransformSection, FloatChannels[4], EditorData[4].Name, Time);
 			}
 			if (EnumHasAnyFlags(Mask, EMovieSceneTransformChannel::RotationZ))
 			{
-				ForwardIters.Emplace(TransformSection, EMovieSceneTransformChannel::RotationZ, TransformSection->GetRotationCurve(EAxis::Z), Time);
-				BackwardIters.Emplace(TransformSection, EMovieSceneTransformChannel::RotationZ, TransformSection->GetRotationCurve(EAxis::Z), Time);
+				ForwardIters.Emplace(TransformSection, FloatChannels[5], EditorData[5].Name, Time);
+				BackwardIters.Emplace(TransformSection, FloatChannels[5], EditorData[5].Name, Time);
 			}
 		}
 	}
 
 	auto HasAnyValidIterators = [](const FCurveKeyIterator& It)
 	{
-		return It;
+		return bool(It);
 	};
 
-	for (FCurveKeyIterator& Bck : BackwardIters)
-	{
-		--Bck;
-	}
-
 	TArray<FTrajectoryKey> Result;
 	while (ForwardIters.ContainsByPredicate(HasAnyValidIterators) || BackwardIters.ContainsByPredicate(HasAnyValidIterators))
 	{
@@ -178,41 +228,52 @@ TArray<FTrajectoryKey> UMovieScene3DTransformTrack::GetTrajectoryData(float Time
 			break;
 		}
 
-		FTrajectoryKey& NewKey = Result[Result.Emplace(TNumericLimits<float>::Max())];
+		uint32 ClosestDistance = -1;
+		TOptional<FFrameNumber> ClosestTime;
 
 		// Find the closest key time
 		for (const FCurveKeyIterator& Fwd : ForwardIters)
 		{
-			if (Fwd && ( FMath::Abs(Time - Fwd.GetTime()) < FMath::Abs(Time - NewKey.Time) ) )
+			if (Fwd && ( !ClosestTime.IsSet() || GetDistanceFromTo(Fwd.GetTime(), Time) < ClosestDistance ) )
 			{
-				NewKey.Time = Fwd.GetTime();
+				ClosestTime = Fwd.GetTime();
+				ClosestDistance = GetDistanceFromTo(Fwd.GetTime(), Time);
 			}
 		}
 		for (const FCurveKeyIterator& Bck : BackwardIters)
 		{
-			if (Bck && ( FMath::Abs(Time - Bck.GetTime()) < FMath::Abs(Time - NewKey.Time) ) )
+			if (Bck && ( !ClosestTime.IsSet() || GetDistanceFromTo(Bck.GetTime(), Time) < ClosestDistance ) )
 			{
-				NewKey.Time = Bck.GetTime();
+				ClosestTime = Bck.GetTime();
+				ClosestDistance = GetDistanceFromTo(Bck.GetTime(), Time);
 			}
 		}
 
+		
+		FTrajectoryKey& NewKey = Result[Result.Emplace(ClosestTime.GetValue())];
 		for (FCurveKeyIterator& Fwd : ForwardIters)
 		{
-			if (Fwd && FMath::IsNearlyEqual(Fwd.GetTime(), NewKey.Time))
+			if (Fwd && Fwd.GetTime() == NewKey.Time)
 			{
-				// Add this key to the trajectory key
-				NewKey.KeyData.Emplace(Fwd.GetSection(), Fwd.GetInterpMode(), Fwd.GetType());
+				if (Fwd.IsValid())
+				{
+					// Add this key to the trajectory key
+					NewKey.KeyData.Emplace(Fwd.GetSection(), Fwd.GetKeyHandle(), Fwd.GetInterpMode(), Fwd.GetChannelName());
+				}
 				// Move onto the next key in this curve
 				++Fwd;
 			}
 		}
-
+		
 		for (FCurveKeyIterator& Bck : BackwardIters)
 		{
-			if (Bck && FMath::IsNearlyEqual(Bck.GetTime(), NewKey.Time))
+			if (Bck && Bck.GetTime() == NewKey.Time)
 			{
-				// Add this key to the trajectory key
-				NewKey.KeyData.Emplace(Bck.GetSection(), Bck.GetInterpMode(), Bck.GetType());
+				if (Bck.IsValid())
+				{
+					// Add this key to the trajectory key
+					NewKey.KeyData.Emplace(Bck.GetSection(), Bck.GetKeyHandle(), Bck.GetInterpMode(), Bck.GetChannelName());
+				}
 				// Move onto the next key in this curve
 				--Bck;
 			}
@@ -229,13 +290,4 @@ TArray<FTrajectoryKey> UMovieScene3DTransformTrack::GetTrajectoryData(float Time
 	return Result;
 }
 
-UMovieSceneSection* UMovieScene3DTransformTrack::CreateNewSection()
-{
-	return NewObject<UMovieSceneSection>(this, UMovieScene3DTransformSection::StaticClass(), NAME_None, RF_Transactional);
-}
-
-FMovieSceneInterrogationKey UMovieScene3DTransformTrack::GetInterrogationKey()
-{
-	static FMovieSceneAnimTypeID TypeID = FMovieSceneAnimTypeID::Unique();
-	return TypeID;
-}
+#endif	// WITH_EDITOR
diff --git a/Engine/Source/Runtime/MovieSceneTracks/Private/Tracks/MovieSceneAudioTrack.cpp b/Engine/Source/Runtime/MovieSceneTracks/Private/Tracks/MovieSceneAudioTrack.cpp
--- a/Engine/Source/Runtime/MovieSceneTracks/Private/Tracks/MovieSceneAudioTrack.cpp
+++ b/Engine/Source/Runtime/MovieSceneTracks/Private/Tracks/MovieSceneAudioTrack.cpp
@@ -68,17 +68,6 @@ bool UMovieSceneAudioTrack::IsEmpty() const
 	return AudioSections.Num() == 0;
 }
 
-
-TRange<float> UMovieSceneAudioTrack::GetSectionBoundaries() const
-{
-	TArray< TRange<float> > Bounds;
-	for (int32 i = 0; i < AudioSections.Num(); ++i)
-	{
-		Bounds.Add(AudioSections[i]->GetRange());
-	}
-	return TRange<float>::Hull(Bounds);
-}
-
 float UMovieSceneAudioTrack::GetSoundDuration(USoundBase* Sound)
 {
 	USoundWave* SoundWave = nullptr;
@@ -109,26 +98,31 @@ float UMovieSceneAudioTrack::GetSoundDuration(USoundBase* Sound)
 	return Duration == INDEFINITELY_LOOPING_DURATION ? SoundWave->Duration : Duration;
 }
 
-void UMovieSceneAudioTrack::AddNewSound(USoundBase* Sound, float Time)
+UMovieSceneSection* UMovieSceneAudioTrack::AddNewSoundOnRow(USoundBase* Sound, FFrameNumber Time, int32 RowIndex)
 {
 	check(Sound);
 	
+	FFrameRate FrameRate = GetTypedOuter<UMovieScene>()->GetFrameResolution();
+
 	// determine initial duration
 	// @todo Once we have infinite sections, we can remove this
-	float DurationToUse = 1.f; // if all else fails, use 1 second duration
+	// @todo ^^ Why? Infinte sections would mean there's no starting time?
+	FFrameTime DurationToUse = 1.f * FrameRate; // if all else fails, use 1 second duration
 
 	float SoundDuration = GetSoundDuration(Sound);
 	if (SoundDuration != INDEFINITELY_LOOPING_DURATION)
 	{
-		DurationToUse = SoundDuration;
+		DurationToUse = SoundDuration * FrameRate;
 	}
 
 	// add the section
 	UMovieSceneAudioSection* NewSection = NewObject<UMovieSceneAudioSection>(this);
-	NewSection->InitialPlacement( AudioSections, Time, Time + DurationToUse, SupportsMultipleRows() );
+	NewSection->InitialPlacementOnRow( AudioSections, Time, DurationToUse.FrameNumber.Value, RowIndex );
 	NewSection->SetSound(Sound);
 
 	AudioSections.Add(NewSection);
+
+	return NewSection;
 }
 
 
diff --git a/Engine/Source/Runtime/MovieSceneTracks/Private/Tracks/MovieSceneByteTrack.cpp b/Engine/Source/Runtime/MovieSceneTracks/Private/Tracks/MovieSceneByteTrack.cpp
--- a/Engine/Source/Runtime/MovieSceneTracks/Private/Tracks/MovieSceneByteTrack.cpp
+++ b/Engine/Source/Runtime/MovieSceneTracks/Private/Tracks/MovieSceneByteTrack.cpp
@@ -8,10 +8,18 @@ UMovieSceneByteTrack::UMovieSceneByteTrack( const FObjectInitializer& ObjectInit
 	: Super( ObjectInitializer )
 { }
 
+void UMovieSceneByteTrack::PostLoad()
+{
+	Super::PostLoad();
+
+	SetEnum(Enum);
+}
 
 UMovieSceneSection* UMovieSceneByteTrack::CreateNewSection()
 {
-	return NewObject<UMovieSceneSection>(this, UMovieSceneByteSection::StaticClass(), NAME_None, RF_Transactional);
+	UMovieSceneByteSection* NewByteSection = NewObject<UMovieSceneByteSection>(this, UMovieSceneByteSection::StaticClass(), NAME_None, RF_Transactional);
+	NewByteSection->ByteCurve.SetEnum(Enum);
+	return NewByteSection;
 }
 
 FMovieSceneEvalTemplatePtr UMovieSceneByteTrack::CreateTemplateForSection(const UMovieSceneSection& InSection) const
@@ -22,10 +30,17 @@ FMovieSceneEvalTemplatePtr UMovieSceneByteTrack::CreateTemplateForSection(const
 void UMovieSceneByteTrack::SetEnum(UEnum* InEnum)
 {
 	Enum = InEnum;
+	for (UMovieSceneSection* Section : Sections)
+	{
+		if (UMovieSceneByteSection* ByteSection = Cast<UMovieSceneByteSection>(Section))
+		{
+			ByteSection->ByteCurve.SetEnum(Enum);
+		}
+	}
 }
 
 
-class UEnum* UMovieSceneByteTrack::GetEnum() const
+UEnum* UMovieSceneByteTrack::GetEnum() const
 {
 	return Enum;
 }
diff --git a/Engine/Source/Runtime/MovieSceneTracks/Private/Tracks/MovieSceneCameraAnimTrack.cpp b/Engine/Source/Runtime/MovieSceneTracks/Private/Tracks/MovieSceneCameraAnimTrack.cpp
--- a/Engine/Source/Runtime/MovieSceneTracks/Private/Tracks/MovieSceneCameraAnimTrack.cpp
+++ b/Engine/Source/Runtime/MovieSceneTracks/Private/Tracks/MovieSceneCameraAnimTrack.cpp
@@ -9,20 +9,24 @@
 #include "Evaluation/MovieSceneEvaluationTrack.h"
 #include "Evaluation/MovieSceneEvaluationTemplate.h"
 #include "Compilation/MovieSceneCompilerRules.h"
+#include "MovieScene.h"
 
 #define LOCTEXT_NAMESPACE "MovieSceneCameraAnimTrack"
 
 
-void UMovieSceneCameraAnimTrack::AddNewCameraAnim(float KeyTime, UCameraAnim* CameraAnim)
+UMovieSceneSection* UMovieSceneCameraAnimTrack::AddNewCameraAnim(FFrameNumber KeyTime, UCameraAnim* CameraAnim)
 {
 	UMovieSceneCameraAnimSection* const NewSection = Cast<UMovieSceneCameraAnimSection>(CreateNewSection());
 	if (NewSection)
 	{
-		NewSection->InitialPlacement(CameraAnimSections, KeyTime, KeyTime + CameraAnim->AnimLength, SupportsMultipleRows());
+		FFrameTime AnimDurationFrames = CameraAnim->AnimLength * GetTypedOuter<UMovieScene>()->GetFrameResolution();
+		NewSection->InitialPlacement(CameraAnimSections, KeyTime, AnimDurationFrames.FrameNumber.Value, SupportsMultipleRows());
 		NewSection->AnimData.CameraAnim = CameraAnim;
 
 		AddSection(*NewSection);
 	}
+
+	return NewSection;
 }
 
 /* UMovieSceneTrack interface
@@ -77,19 +81,6 @@ bool UMovieSceneCameraAnimTrack::IsEmpty() const
 }
 
 
-TRange<float> UMovieSceneCameraAnimTrack::GetSectionBoundaries() const
-{
-	TArray<TRange<float>> Bounds;
-
-	for (auto Section : CameraAnimSections)
-	{
-		Bounds.Add(Section->GetRange());
-	}
-
-	return TRange<float>::Hull(Bounds);
-}
-
-
 
 
 #if WITH_EDITORONLY_DATA
@@ -99,14 +90,14 @@ FText UMovieSceneCameraAnimTrack::GetDisplayName() const
 }
 #endif
 
-void UMovieSceneCameraAnimTrack::GetCameraAnimSectionsAtTime(float Time, TArray<UMovieSceneCameraAnimSection*>& OutSections)
+void UMovieSceneCameraAnimTrack::GetCameraAnimSectionsAtTime(FFrameNumber Time, TArray<UMovieSceneCameraAnimSection*>& OutSections)
 {
 	OutSections.Empty();
 
 	for (auto Section : CameraAnimSections)
 	{
 		UMovieSceneCameraAnimSection* const CASection = dynamic_cast<UMovieSceneCameraAnimSection*>(Section);
-		if (CASection && CASection->IsTimeWithinSection(Time))
+		if (CASection && CASection->GetRange().Contains(Time))
 		{
 			OutSections.Add(CASection);
 		}
diff --git a/Engine/Source/Runtime/MovieSceneTracks/Private/Tracks/MovieSceneCameraCutTrack.cpp b/Engine/Source/Runtime/MovieSceneTracks/Private/Tracks/MovieSceneCameraCutTrack.cpp
--- a/Engine/Source/Runtime/MovieSceneTracks/Private/Tracks/MovieSceneCameraCutTrack.cpp
+++ b/Engine/Source/Runtime/MovieSceneTracks/Private/Tracks/MovieSceneCameraCutTrack.cpp
@@ -3,6 +3,7 @@
 #include "Tracks/MovieSceneCameraCutTrack.h"
 #include "MovieScene.h"
 #include "MovieSceneCommonHelpers.h"
+#include "MovieSceneTimeHelpers.h"
 #include "Sections/MovieSceneCameraCutSection.h"
 #include "Evaluation/MovieSceneCameraCutTemplate.h"
 
@@ -22,23 +23,18 @@ UMovieSceneCameraCutTrack::UMovieSceneCameraCutTrack( const FObjectInitializer&
 }
 
 
-void UMovieSceneCameraCutTrack::AddNewCameraCut(FGuid CameraHandle, float StartTime)
-{
-	AddNewCameraCut(FMovieSceneObjectBindingID(CameraHandle, MovieSceneSequenceID::Root), StartTime);
-}
-
-UMovieSceneCameraCutSection* UMovieSceneCameraCutTrack::AddNewCameraCut(const FMovieSceneObjectBindingID& CameraBindingID, float StartTime)
+UMovieSceneCameraCutSection* UMovieSceneCameraCutTrack::AddNewCameraCut(const FMovieSceneObjectBindingID& CameraBindingID, FFrameNumber StartTime)
 {
 	Modify();
 
 
-	float NewSectionEndTime = FindEndTimeForCameraCut(StartTime);
+	FFrameNumber NewSectionEndTime = FindEndTimeForCameraCut(StartTime);
 
 	// If there's an existing section, just swap the camera guid
 	UMovieSceneCameraCutSection* ExistingSection = nullptr;
 	for (auto Section : Sections)
 	{
-		if (Section->GetStartTime() == StartTime && Section->GetEndTime() == NewSectionEndTime)
+		if (Section->HasStartFrame() && Section->HasEndFrame() && Section->GetInclusiveStartFrame() == StartTime && Section->GetExclusiveEndFrame() == NewSectionEndTime)
 		{
 			ExistingSection = Cast<UMovieSceneCameraCutSection>(Section);
 			break;
@@ -53,11 +49,8 @@ UMovieSceneCameraCutSection* UMovieSceneCameraCutTrack::AddNewCameraCut(const FM
 	else
 	{
 		NewSection = NewObject<UMovieSceneCameraCutSection>(this, NAME_None, RF_Transactional);
-		{
-			NewSection->SetStartTime(StartTime);
-			NewSection->SetEndTime(FindEndTimeForCameraCut(StartTime));
-			NewSection->SetCameraBindingID(CameraBindingID);
-		}
+		NewSection->SetRange(TRange<FFrameNumber>(StartTime, NewSectionEndTime));
+		NewSection->SetCameraBindingID(CameraBindingID);
 
 		Sections.Add(NewSection);
 	}
@@ -121,37 +114,30 @@ void UMovieSceneCameraCutTrack::OnSectionMoved(UMovieSceneSection& Section)
 }
 #endif
 
-float UMovieSceneCameraCutTrack::FindEndTimeForCameraCut( float StartTime )
+FFrameNumber UMovieSceneCameraCutTrack::FindEndTimeForCameraCut( FFrameNumber StartTime )
 {
-	float EndTime = 0;
-	bool bFoundEndTime = false;
+	UMovieScene* OwnerScene = GetTypedOuter<UMovieScene>();
+
+	// End time should default to end where the movie scene ends. Ensure it is at least the same as start time (this should only happen when the movie scene has an initial time range smaller than the start time)
+	FFrameNumber ExclusivePlayEnd = MovieScene::DiscreteExclusiveUpper(OwnerScene->GetPlaybackRange());
+	FFrameNumber ExclusiveEndTime = FMath::Max( ExclusivePlayEnd, StartTime );
 
 	for( UMovieSceneSection* Section : Sections )
 	{
-		if( Section->GetStartTime() >= StartTime )
+		if( Section->HasStartFrame() && Section->GetInclusiveStartFrame() >= StartTime )
 		{
-			EndTime = Section->GetStartTime();
-			bFoundEndTime = true;
+			ExclusiveEndTime = Section->GetInclusiveStartFrame();
 			break;
 		}
 	}
 
-	if( !bFoundEndTime )
-	{
-		UMovieScene* OwnerScene = GetTypedOuter<UMovieScene>();
-
-		// End time should just end where the movie scene ends.  Ensure it is at least the same as start time (this should only happen when the movie scene has an initial time range smaller than the start time
-		EndTime = FMath::Max( OwnerScene->GetPlaybackRange().GetUpperBoundValue(), StartTime );
-	}
-
-			
-	if( StartTime == EndTime )
+	if( StartTime == ExclusiveEndTime )
 	{
 		// Give the CameraCut a reasonable length of time to start out with.  A 0 time CameraCut is not usable
-		EndTime = StartTime + .5f;
+		ExclusiveEndTime = (StartTime + .5f * OwnerScene->GetFrameResolution()).FrameNumber;
 	}
 
-	return EndTime;
+	return ExclusiveEndTime;
 }
 
 #undef LOCTEXT_NAMESPACE
diff --git a/Engine/Source/Runtime/MovieSceneTracks/Private/Tracks/MovieSceneCameraShakeTrack.cpp b/Engine/Source/Runtime/MovieSceneTracks/Private/Tracks/MovieSceneCameraShakeTrack.cpp
--- a/Engine/Source/Runtime/MovieSceneTracks/Private/Tracks/MovieSceneCameraShakeTrack.cpp
+++ b/Engine/Source/Runtime/MovieSceneTracks/Private/Tracks/MovieSceneCameraShakeTrack.cpp
@@ -8,20 +8,24 @@
 #include "Evaluation/MovieSceneEvaluationTemplate.h"
 #include "Compilation/MovieSceneSegmentCompiler.h"
 #include "Compilation/MovieSceneCompilerRules.h"
+#include "MovieScene.h"
 
 #define LOCTEXT_NAMESPACE "MovieSceneCameraShakeTrack"
 
-void UMovieSceneCameraShakeTrack::AddNewCameraShake(float KeyTime, TSubclassOf<UCameraShake> ShakeClass)
+UMovieSceneSection* UMovieSceneCameraShakeTrack::AddNewCameraShake(FFrameNumber KeyTime, TSubclassOf<UCameraShake> ShakeClass)
 {
 	UMovieSceneCameraShakeSection* const NewSection = Cast<UMovieSceneCameraShakeSection>(CreateNewSection());
 	if (NewSection)
 	{
 		// #fixme get length
-		NewSection->InitialPlacement(CameraShakeSections, KeyTime, KeyTime + 5.f /*AnimSequence->SequenceLength*/, SupportsMultipleRows());
+		FFrameTime Duration = 5.0 * GetTypedOuter<UMovieScene>()->GetFrameResolution();
+		NewSection->InitialPlacement(CameraShakeSections, KeyTime, Duration.FrameNumber.Value, SupportsMultipleRows());
 		NewSection->ShakeData.ShakeClass = ShakeClass;
 		
 		AddSection(*NewSection);
 	}
+
+	return NewSection;
 }
 
 FMovieSceneTrackSegmentBlenderPtr UMovieSceneCameraShakeTrack::GetTrackSegmentBlender() const
@@ -85,17 +89,4 @@ bool UMovieSceneCameraShakeTrack::IsEmpty() const
 }
 
 
-TRange<float> UMovieSceneCameraShakeTrack::GetSectionBoundaries() const
-{
-	TArray<TRange<float>> Bounds;
-
-	for (auto Section : CameraShakeSections)
-	{
-		Bounds.Add(Section->GetRange());
-	}
-
-	return TRange<float>::Hull(Bounds);
-}
-
-
 #undef LOCTEXT_NAMESPACE
diff --git a/Engine/Source/Runtime/MovieSceneTracks/Private/Tracks/MovieSceneCinematicShotTrack.cpp b/Engine/Source/Runtime/MovieSceneTracks/Private/Tracks/MovieSceneCinematicShotTrack.cpp
--- a/Engine/Source/Runtime/MovieSceneTracks/Private/Tracks/MovieSceneCinematicShotTrack.cpp
+++ b/Engine/Source/Runtime/MovieSceneTracks/Private/Tracks/MovieSceneCinematicShotTrack.cpp
@@ -20,9 +20,9 @@ UMovieSceneCinematicShotTrack::UMovieSceneCinematicShotTrack(const FObjectInitia
 #endif
 }
 
-UMovieSceneSubSection* UMovieSceneCinematicShotTrack::AddSequence(UMovieSceneSequence* Sequence, float StartTime, float Duration, const bool& bInsertSequence)
+UMovieSceneSubSection* UMovieSceneCinematicShotTrack::AddSequenceOnRow(UMovieSceneSequence* Sequence, FFrameNumber StartTime, int32 Duration, int32 RowIndex)
 {
-	UMovieSceneSubSection* NewSection = UMovieSceneSubTrack::AddSequence(Sequence, StartTime, Duration, bInsertSequence);
+	UMovieSceneSubSection* NewSection = UMovieSceneSubTrack::AddSequenceOnRow(Sequence, StartTime, Duration, RowIndex);
 
 	UMovieSceneCinematicShotSection* NewShotSection = Cast<UMovieSceneCinematicShotSection>(NewSection);
 
@@ -135,10 +135,8 @@ FMovieSceneTrackRowSegmentBlenderPtr UMovieSceneCinematicShotTrack::GetRowSegmen
 			}
 			else if (A.Section->GetOverlapPriority() == B.Section->GetOverlapPriority())
 			{
-				TRangeBound<float> StartBoundA = A.Section->IsInfinite() ? TRangeBound<float>::Open() : TRangeBound<float>::Inclusive(A.Section->GetStartTime());
-				TRangeBound<float> StartBoundB = B.Section->IsInfinite() ? TRangeBound<float>::Open() : TRangeBound<float>::Inclusive(B.Section->GetStartTime());
-
-				return TRangeBound<float>::MaxLower(StartBoundA, StartBoundB) == StartBoundA;
+				TRangeBound<FFrameNumber> StartBoundA = A.Section->GetRange().GetLowerBound();
+				return TRangeBound<FFrameNumber>::MaxLower(StartBoundA, B.Section->GetRange().GetLowerBound()) == StartBoundA;
 			}
 			return A.Section->GetOverlapPriority() > B.Section->GetOverlapPriority();
 		}
diff --git a/Engine/Source/Runtime/MovieSceneTracks/Private/Tracks/MovieSceneEnumTrack.cpp b/Engine/Source/Runtime/MovieSceneTracks/Private/Tracks/MovieSceneEnumTrack.cpp
--- a/Engine/Source/Runtime/MovieSceneTracks/Private/Tracks/MovieSceneEnumTrack.cpp
+++ b/Engine/Source/Runtime/MovieSceneTracks/Private/Tracks/MovieSceneEnumTrack.cpp
@@ -9,10 +9,17 @@ UMovieSceneEnumTrack::UMovieSceneEnumTrack( const FObjectInitializer& ObjectInit
 	: Super( ObjectInitializer )
 { }
 
+void UMovieSceneEnumTrack::PostLoad()
+{
+	Super::PostLoad();
+	SetEnum(Enum);
+}
 
 UMovieSceneSection* UMovieSceneEnumTrack::CreateNewSection()
 {
-	return NewObject<UMovieSceneSection>(this, UMovieSceneEnumSection::StaticClass(), NAME_None, RF_Transactional);
+	UMovieSceneEnumSection* NewEnumSection = NewObject<UMovieSceneEnumSection>(this, UMovieSceneEnumSection::StaticClass(), NAME_None, RF_Transactional);
+	NewEnumSection->EnumCurve.SetEnum(Enum);
+	return NewEnumSection;
 }
 
 FMovieSceneEvalTemplatePtr UMovieSceneEnumTrack::CreateTemplateForSection(const UMovieSceneSection& InSection) const
@@ -23,10 +30,18 @@ FMovieSceneEvalTemplatePtr UMovieSceneEnumTrack::CreateTemplateForSection(const
 void UMovieSceneEnumTrack::SetEnum(UEnum* InEnum)
 {
 	Enum = InEnum;
+
+	for (UMovieSceneSection* Section : Sections)
+	{
+		if (UMovieSceneEnumSection* EnumSection = Cast<UMovieSceneEnumSection>(Section))
+		{
+			EnumSection->EnumCurve.SetEnum(Enum);
+		}
+	}
 }
 
 
-class UEnum* UMovieSceneEnumTrack::GetEnum() const
+UEnum* UMovieSceneEnumTrack::GetEnum() const
 {
 	return Enum;
 }
diff --git a/Engine/Source/Runtime/MovieSceneTracks/Private/Tracks/MovieSceneEventTrack.cpp b/Engine/Source/Runtime/MovieSceneTracks/Private/Tracks/MovieSceneEventTrack.cpp
--- a/Engine/Source/Runtime/MovieSceneTracks/Private/Tracks/MovieSceneEventTrack.cpp
+++ b/Engine/Source/Runtime/MovieSceneTracks/Private/Tracks/MovieSceneEventTrack.cpp
@@ -32,19 +32,6 @@ const TArray<UMovieSceneSection*>& UMovieSceneEventTrack::GetAllSections() const
 }
 
 
-TRange<float> UMovieSceneEventTrack::GetSectionBoundaries() const
-{
-	TRange<float> SectionBoundaries = TRange<float>::Empty();
-
-	for (auto& Section : Sections)
-	{
-		SectionBoundaries = TRange<float>::Hull(SectionBoundaries, Section->GetRange());
-	}
-
-	return SectionBoundaries;
-}
-
-
 bool UMovieSceneEventTrack::HasSection(const UMovieSceneSection& Section) const
 {
 	return Sections.Contains(&Section);
diff --git a/Engine/Source/Runtime/MovieSceneTracks/Private/Tracks/MovieSceneFloatTrack.cpp b/Engine/Source/Runtime/MovieSceneTracks/Private/Tracks/MovieSceneFloatTrack.cpp
--- a/Engine/Source/Runtime/MovieSceneTracks/Private/Tracks/MovieSceneFloatTrack.cpp
+++ b/Engine/Source/Runtime/MovieSceneTracks/Private/Tracks/MovieSceneFloatTrack.cpp
@@ -21,21 +21,4 @@ UMovieSceneSection* UMovieSceneFloatTrack::CreateNewSection()
 FMovieSceneEvalTemplatePtr UMovieSceneFloatTrack::CreateTemplateForSection(const UMovieSceneSection& InSection) const
 {
 	return FMovieSceneFloatPropertySectionTemplate(*CastChecked<const UMovieSceneFloatSection>(&InSection), *this);
-}
-
-bool UMovieSceneFloatTrack::Eval( float Position, float LastPosition, float& InOutFloat ) const
-{
-	const UMovieSceneSection* Section = MovieSceneHelpers::FindNearestSectionAtTime( Sections, Position );
-
-	if( Section )
-	{
-		if (!Section->IsInfinite())
-		{
-			Position = FMath::Clamp(Position, Section->GetStartTime(), Section->GetEndTime());
-		}
-
-		InOutFloat = CastChecked<UMovieSceneFloatSection>( Section )->Eval( Position, InOutFloat );
-	}
-
-	return Section != nullptr;
-}
+}
\ No newline at end of file
diff --git a/Engine/Source/Runtime/MovieSceneTracks/Private/Tracks/MovieSceneLevelVisibilityTrack.cpp b/Engine/Source/Runtime/MovieSceneTracks/Private/Tracks/MovieSceneLevelVisibilityTrack.cpp
--- a/Engine/Source/Runtime/MovieSceneTracks/Private/Tracks/MovieSceneLevelVisibilityTrack.cpp
+++ b/Engine/Source/Runtime/MovieSceneTracks/Private/Tracks/MovieSceneLevelVisibilityTrack.cpp
@@ -53,19 +53,6 @@ const TArray<UMovieSceneSection*>& UMovieSceneLevelVisibilityTrack::GetAllSectio
 }
 
 
-TRange<float> UMovieSceneLevelVisibilityTrack::GetSectionBoundaries() const
-{
-	TArray< TRange<float> > Bounds;
-
-	for (int32 SectionIndex = 0; SectionIndex < Sections.Num(); ++SectionIndex)
-	{
-		Bounds.Add(Sections[SectionIndex]->GetRange());
-	}
-
-	return TRange<float>::Hull(Bounds);
-}
-
-
 bool UMovieSceneLevelVisibilityTrack::HasSection(const UMovieSceneSection& Section) const
 {
 	return Sections.Contains(&Section);
diff --git a/Engine/Source/Runtime/MovieSceneTracks/Private/Tracks/MovieSceneMaterialParameterCollectionTrack.cpp b/Engine/Source/Runtime/MovieSceneTracks/Private/Tracks/MovieSceneMaterialParameterCollectionTrack.cpp
--- a/Engine/Source/Runtime/MovieSceneTracks/Private/Tracks/MovieSceneMaterialParameterCollectionTrack.cpp
+++ b/Engine/Source/Runtime/MovieSceneTracks/Private/Tracks/MovieSceneMaterialParameterCollectionTrack.cpp
@@ -17,7 +17,7 @@ UMovieSceneMaterialParameterCollectionTrack::UMovieSceneMaterialParameterCollect
 UMovieSceneSection* UMovieSceneMaterialParameterCollectionTrack::CreateNewSection()
 {
 	UMovieSceneSection* NewSection =  NewObject<UMovieSceneParameterSection>(this, UMovieSceneParameterSection::StaticClass(), NAME_None, RF_Transactional);
-	NewSection->SetIsInfinite(true);
+	NewSection->SetRange(TRange<FFrameNumber>::All());
 	return NewSection;
 }
 
diff --git a/Engine/Source/Runtime/MovieSceneTracks/Private/Tracks/MovieSceneMaterialTrack.cpp b/Engine/Source/Runtime/MovieSceneTracks/Private/Tracks/MovieSceneMaterialTrack.cpp
--- a/Engine/Source/Runtime/MovieSceneTracks/Private/Tracks/MovieSceneMaterialTrack.cpp
+++ b/Engine/Source/Runtime/MovieSceneTracks/Private/Tracks/MovieSceneMaterialTrack.cpp
@@ -50,50 +50,35 @@ bool UMovieSceneMaterialTrack::IsEmpty() const
 }
 
 
-TRange<float> UMovieSceneMaterialTrack::GetSectionBoundaries() const
-{
-	TArray< TRange<float> > Bounds;
-
-	for (int32 SectionIndex = 0; SectionIndex < Sections.Num(); ++SectionIndex)
-	{
-		Bounds.Add(Sections[SectionIndex]->GetRange());
-	}
-
-	return TRange<float>::Hull(Bounds);
-}
-
-
 const TArray<UMovieSceneSection*>& UMovieSceneMaterialTrack::GetAllSections() const
 {
 	return Sections;
 }
 
 
-void UMovieSceneMaterialTrack::AddScalarParameterKey(FName ParameterName, float Time, float Value)
+void UMovieSceneMaterialTrack::AddScalarParameterKey(FName ParameterName, FFrameNumber Time, float Value)
 {
 	UMovieSceneParameterSection* NearestSection = Cast<UMovieSceneParameterSection>(MovieSceneHelpers::FindNearestSectionAtTime(Sections, Time));
 	if (NearestSection == nullptr)
 	{
 		NearestSection = Cast<UMovieSceneParameterSection>(CreateNewSection());
-		NearestSection->SetStartTime(Time);
-		NearestSection->SetEndTime(Time);
+		NearestSection->SetRange(TRange<FFrameNumber>::Inclusive(Time, Time));
 		Sections.Add(NearestSection);
 	}
 	if (NearestSection->TryModify())
 	{
-		NearestSection->AddScalarParameterKey(ParameterName, Time, Value);	
+		NearestSection->AddScalarParameterKey(ParameterName, Time, Value);
 	}
 }
 
 
-void UMovieSceneMaterialTrack::AddColorParameterKey(FName ParameterName, float Time, FLinearColor Value)
+void UMovieSceneMaterialTrack::AddColorParameterKey(FName ParameterName, FFrameNumber Time, FLinearColor Value)
 {
 	UMovieSceneParameterSection* NearestSection = Cast<UMovieSceneParameterSection>(MovieSceneHelpers::FindNearestSectionAtTime(Sections, Time));
 	if (NearestSection == nullptr)
 	{
 		NearestSection = Cast<UMovieSceneParameterSection>(CreateNewSection());
-		NearestSection->SetStartTime(Time);
-		NearestSection->SetEndTime(Time);
+		NearestSection->SetRange(TRange<FFrameNumber>::Inclusive(Time, Time));
 		Sections.Add(NearestSection);
 	}
 	if (NearestSection->TryModify())
diff --git a/Engine/Source/Runtime/MovieSceneTracks/Private/Tracks/MovieSceneParticleParameterTrack.cpp b/Engine/Source/Runtime/MovieSceneTracks/Private/Tracks/MovieSceneParticleParameterTrack.cpp
--- a/Engine/Source/Runtime/MovieSceneTracks/Private/Tracks/MovieSceneParticleParameterTrack.cpp
+++ b/Engine/Source/Runtime/MovieSceneTracks/Private/Tracks/MovieSceneParticleParameterTrack.cpp
@@ -49,18 +49,6 @@ bool UMovieSceneParticleParameterTrack::IsEmpty() const
 	return Sections.Num() == 0;
 }
 
-TRange<float> UMovieSceneParticleParameterTrack::GetSectionBoundaries() const
-{
-	TArray< TRange<float> > Bounds;
-
-	for ( int32 SectionIndex = 0; SectionIndex < Sections.Num(); ++SectionIndex )
-	{
-		Bounds.Add( Sections[SectionIndex]->GetRange() );
-	}
-
-	return TRange<float>::Hull( Bounds );
-}
-
 const TArray<UMovieSceneSection*>& UMovieSceneParticleParameterTrack::GetAllSections() const
 {
 	return Sections;
@@ -75,40 +63,37 @@ FText UMovieSceneParticleParameterTrack::GetDefaultDisplayName() const
 #endif
 
 
-void UMovieSceneParticleParameterTrack::AddScalarParameterKey( FName ParameterName, float Time, float Value )
+void UMovieSceneParticleParameterTrack::AddScalarParameterKey( FName ParameterName, FFrameNumber Time, float Value )
 {
 	UMovieSceneParameterSection* NearestSection = Cast<UMovieSceneParameterSection>(MovieSceneHelpers::FindNearestSectionAtTime(Sections, Time));
 	if ( NearestSection == nullptr )
 	{
 		NearestSection = Cast<UMovieSceneParameterSection>(CreateNewSection());
-		NearestSection->SetStartTime( Time );
-		NearestSection->SetEndTime( Time );
+		NearestSection->SetRange(TRange<FFrameNumber>::Inclusive(Time, Time));
 		Sections.Add( NearestSection );
 	}
 	NearestSection->AddScalarParameterKey(ParameterName, Time, Value);
 }
 
-void UMovieSceneParticleParameterTrack::AddVectorParameterKey( FName ParameterName, float Time, FVector Value )
+void UMovieSceneParticleParameterTrack::AddVectorParameterKey( FName ParameterName, FFrameNumber Time, FVector Value )
 {
 	UMovieSceneParameterSection* NearestSection = Cast<UMovieSceneParameterSection>( MovieSceneHelpers::FindNearestSectionAtTime( Sections, Time ) );
 	if ( NearestSection == nullptr )
 	{
 		NearestSection = Cast<UMovieSceneParameterSection>( CreateNewSection() );
-		NearestSection->SetStartTime( Time );
-		NearestSection->SetEndTime( Time );
+		NearestSection->SetRange(TRange<FFrameNumber>::Inclusive(Time, Time));
 		Sections.Add( NearestSection );
 	}
 	NearestSection->AddVectorParameterKey( ParameterName, Time, Value );
 }
 
-void UMovieSceneParticleParameterTrack::AddColorParameterKey( FName ParameterName, float Time, FLinearColor Value )
+void UMovieSceneParticleParameterTrack::AddColorParameterKey( FName ParameterName, FFrameNumber Time, FLinearColor Value )
 {
 	UMovieSceneParameterSection* NearestSection = Cast<UMovieSceneParameterSection>( MovieSceneHelpers::FindNearestSectionAtTime( Sections, Time ) );
 	if ( NearestSection == nullptr )
 	{
 		NearestSection = Cast<UMovieSceneParameterSection>( CreateNewSection() );
-		NearestSection->SetStartTime( Time );
-		NearestSection->SetEndTime( Time );
+		NearestSection->SetRange(TRange<FFrameNumber>::Inclusive(Time, Time));
 		Sections.Add( NearestSection );
 	}
 	NearestSection->AddColorParameterKey( ParameterName, Time, Value );
diff --git a/Engine/Source/Runtime/MovieSceneTracks/Private/Tracks/MovieSceneParticleTrack.cpp b/Engine/Source/Runtime/MovieSceneTracks/Private/Tracks/MovieSceneParticleTrack.cpp
--- a/Engine/Source/Runtime/MovieSceneTracks/Private/Tracks/MovieSceneParticleTrack.cpp
+++ b/Engine/Source/Runtime/MovieSceneTracks/Private/Tracks/MovieSceneParticleTrack.cpp
@@ -53,26 +53,14 @@ bool UMovieSceneParticleTrack::IsEmpty() const
 }
 
 
-TRange<float> UMovieSceneParticleTrack::GetSectionBoundaries() const
-{
-	TArray< TRange<float> > Bounds;
-	for (int32 i = 0; i < ParticleSections.Num(); ++i)
-	{
-		Bounds.Add(ParticleSections[i]->GetRange());
-	}
-	return TRange<float>::Hull(Bounds);
-}
-
-
-void UMovieSceneParticleTrack::AddNewSection( float SectionTime )
+void UMovieSceneParticleTrack::AddNewSection( FFrameNumber SectionTime )
 {
 	if ( MovieSceneHelpers::FindSectionAtTime( ParticleSections, SectionTime ) == nullptr )
 	{
 		UMovieSceneParticleSection* NewSection = Cast<UMovieSceneParticleSection>( CreateNewSection() );
-		NewSection->SetStartTime( SectionTime );
-		NewSection->SetEndTime( SectionTime );
-		NewSection->SetStartTime( SectionTime );
-		NewSection->SetEndTime( SectionTime );
+
+		NewSection->SetRange(TRange<FFrameNumber>::Inclusive(SectionTime, SectionTime));
+
 		ParticleSections.Add(NewSection);
 	}
 }
diff --git a/Engine/Source/Runtime/MovieSceneTracks/Private/Tracks/MovieScenePropertyTrack.cpp b/Engine/Source/Runtime/MovieSceneTracks/Private/Tracks/MovieScenePropertyTrack.cpp
--- a/Engine/Source/Runtime/MovieSceneTracks/Private/Tracks/MovieScenePropertyTrack.cpp
+++ b/Engine/Source/Runtime/MovieSceneTracks/Private/Tracks/MovieScenePropertyTrack.cpp
@@ -1,7 +1,8 @@
 // Copyright 1998-2018 Epic Games, Inc. All Rights Reserved.
 
 #include "Tracks/MovieScenePropertyTrack.h"
-
+#include "MovieSceneCommonHelpers.h"
+#include "Algo/Sort.h"
 
 UMovieScenePropertyTrack::UMovieScenePropertyTrack(const FObjectInitializer& ObjectInitializer)
 	: Super(ObjectInitializer)
@@ -85,23 +86,44 @@ bool UMovieScenePropertyTrack::IsEmpty() const
 	return Sections.Num() == 0;
 }
 
+UMovieSceneSection* UMovieScenePropertyTrack::FindSection(FFrameNumber Time)
+{
+	TArray<UMovieSceneSection*, TInlineAllocator<4>> OverlappingSections;
+
+	for (UMovieSceneSection* Section : Sections)
+	{
+		if (Section->GetRange().Contains(Time))
+		{
+			OverlappingSections.Add(Section);
+		}
+	}
+
+	Algo::Sort(OverlappingSections, MovieSceneHelpers::SortOverlappingSections);
+
+	if (OverlappingSections.Num())
+	{
+		return OverlappingSections[0];
+	}
+
+	return nullptr;
+}
+
 
-UMovieSceneSection* UMovieScenePropertyTrack::FindOrAddSection(float Time, bool& bSectionAdded)
+UMovieSceneSection* UMovieScenePropertyTrack::FindOrExtendSection(FFrameNumber Time)
 {
-	bSectionAdded = false;
+	UMovieSceneSection* FoundSection = FindSection(Time);
+	if (FoundSection)
+	{
+		return FoundSection;
+	}
 
 	// Find a spot for the section so that they are sorted by start time
 	for(int32 SectionIndex = 0; SectionIndex < Sections.Num(); ++SectionIndex)
 	{
 		UMovieSceneSection* Section = Sections[SectionIndex];
 
-		if(Section->IsTimeWithinSection(Time))
-		{
-			return Section;
-		}
-
 		// Check if there are no more sections that would overlap the time 
-		if(!Sections.IsValidIndex(SectionIndex+1) || Sections[SectionIndex+1]->GetStartTime() > Time)
+		if(!Sections.IsValidIndex(SectionIndex+1) || (Sections[SectionIndex+1]->HasEndFrame() && Sections[SectionIndex+1]->GetExclusiveEndFrame() > Time))
 		{
 			// No sections overlap the time
 		
@@ -110,34 +132,46 @@ UMovieSceneSection* UMovieScenePropertyTrack::FindOrAddSection(float Time, bool&
 				// Append and grow the previous section
 				UMovieSceneSection* PreviousSection = Sections[ SectionIndex ? SectionIndex-1 : 0 ];
 		
-				PreviousSection->SetEndTime(Time);
+				PreviousSection->SetEndFrame(Time);
 				return PreviousSection;
 			}
 			else if(Sections.IsValidIndex(SectionIndex+1))
 			{
 				// Prepend and grow the next section because there are no sections before this one
 				UMovieSceneSection* NextSection = Sections[SectionIndex+1];
-				NextSection->SetStartTime(Time);
+				NextSection->SetStartFrame(Time);
 				return NextSection;
 			}	
 			else
 			{
 				// SectionIndex == 0 
 				UMovieSceneSection* PreviousSection = Sections[0];
-				if(PreviousSection->GetEndTime() < Time)
+				if(PreviousSection->HasEndFrame() && PreviousSection->GetExclusiveEndFrame() < Time)
 				{
 					// Append and grow the section
-					PreviousSection->SetEndTime(Time);
+					PreviousSection->SetEndFrame(Time);
 				}
 				else
 				{
 					// Prepend and grow the section
-					PreviousSection->SetStartTime(Time);
+					PreviousSection->SetStartFrame(Time);
 				}
 				return PreviousSection;
 			}
 		}
+	}
+
+	return nullptr;
+}
 
+UMovieSceneSection* UMovieScenePropertyTrack::FindOrAddSection(FFrameNumber Time, bool& bSectionAdded)
+{
+	bSectionAdded = false;
+
+	UMovieSceneSection* FoundSection = FindSection(Time);
+	if (FoundSection)
+	{
+		return FoundSection;
 	}
 
 	check(Sections.Num() == 0);
@@ -145,8 +179,7 @@ UMovieSceneSection* UMovieScenePropertyTrack::FindOrAddSection(float Time, bool&
 	// Add a new section that starts and ends at the same time
 	UMovieSceneSection* NewSection = CreateNewSection();
 	NewSection->SetFlags(RF_Transactional);
-	NewSection->SetStartTime(Time);
-	NewSection->SetEndTime(Time);
+	NewSection->SetRange(TRange<FFrameNumber>::Inclusive(Time, Time));
 
 	Sections.Add(NewSection);
 	
@@ -154,16 +187,3 @@ UMovieSceneSection* UMovieScenePropertyTrack::FindOrAddSection(float Time, bool&
 
 	return NewSection;
 }
-
-
-TRange<float> UMovieScenePropertyTrack::GetSectionBoundaries() const
-{
-	TArray< TRange<float> > Bounds;
-
-	for (int32 SectionIndex = 0; SectionIndex < Sections.Num(); ++SectionIndex)
-	{
-		Bounds.Add(Sections[SectionIndex]->GetRange());
-	}
-
-	return TRange<float>::Hull(Bounds);
-}
diff --git a/Engine/Source/Runtime/MovieSceneTracks/Private/Tracks/MovieSceneSkeletalAnimationTrack.cpp b/Engine/Source/Runtime/MovieSceneTracks/Private/Tracks/MovieSceneSkeletalAnimationTrack.cpp
--- a/Engine/Source/Runtime/MovieSceneTracks/Private/Tracks/MovieSceneSkeletalAnimationTrack.cpp
+++ b/Engine/Source/Runtime/MovieSceneTracks/Private/Tracks/MovieSceneSkeletalAnimationTrack.cpp
@@ -7,6 +7,7 @@
 #include "Evaluation/MovieSceneEvaluationTrack.h"
 #include "Evaluation/MovieSceneSkeletalAnimationTemplate.h"
 #include "Compilation/IMovieSceneTemplateGenerator.h"
+#include "MovieScene.h"
 
 #define LOCTEXT_NAMESPACE "MovieSceneSkeletalAnimationTrack"
 
@@ -31,19 +32,22 @@ UMovieSceneSkeletalAnimationTrack::UMovieSceneSkeletalAnimationTrack(const FObje
 /* UMovieSceneSkeletalAnimationTrack interface
  *****************************************************************************/
 
-void UMovieSceneSkeletalAnimationTrack::AddNewAnimation(float KeyTime, UAnimSequenceBase* AnimSequence)
+UMovieSceneSection* UMovieSceneSkeletalAnimationTrack::AddNewAnimationOnRow(FFrameNumber KeyTime, UAnimSequenceBase* AnimSequence, int32 RowIndex)
 {
 	UMovieSceneSkeletalAnimationSection* NewSection = Cast<UMovieSceneSkeletalAnimationSection>(CreateNewSection());
 	{
-		NewSection->InitialPlacement(AnimationSections, KeyTime, KeyTime + AnimSequence->SequenceLength, SupportsMultipleRows());
+		FFrameTime AnimationLength = AnimSequence->SequenceLength * GetTypedOuter<UMovieScene>()->GetFrameResolution();
+		NewSection->InitialPlacementOnRow(AnimationSections, KeyTime, AnimationLength.FrameNumber.Value, RowIndex);
 		NewSection->Params.Animation = AnimSequence;
 	}
 
 	AddSection(*NewSection);
+
+	return NewSection;
 }
 
 
-TArray<UMovieSceneSection*> UMovieSceneSkeletalAnimationTrack::GetAnimSectionsAtTime(float Time)
+TArray<UMovieSceneSection*> UMovieSceneSkeletalAnimationTrack::GetAnimSectionsAtTime(FFrameNumber Time)
 {
 	TArray<UMovieSceneSection*> Sections;
 	for (auto Section : AnimationSections)
@@ -118,19 +122,6 @@ bool UMovieSceneSkeletalAnimationTrack::IsEmpty() const
 	return AnimationSections.Num() == 0;
 }
 
-
-TRange<float> UMovieSceneSkeletalAnimationTrack::GetSectionBoundaries() const
-{
-	TArray<TRange<float>> Bounds;
-
-	for (auto Section : AnimationSections)
-	{
-		Bounds.Add(Section->GetRange());
-	}
-
-	return TRange<float>::Hull(Bounds);
-}
-
 #if WITH_EDITORONLY_DATA
 
 FText UMovieSceneSkeletalAnimationTrack::GetDefaultDisplayName() const
diff --git a/Engine/Source/Runtime/MovieSceneTracks/Private/Tracks/MovieSceneSpawnTrack.cpp b/Engine/Source/Runtime/MovieSceneTracks/Private/Tracks/MovieSceneSpawnTrack.cpp
--- a/Engine/Source/Runtime/MovieSceneTracks/Private/Tracks/MovieSceneSpawnTrack.cpp
+++ b/Engine/Source/Runtime/MovieSceneTracks/Private/Tracks/MovieSceneSpawnTrack.cpp
@@ -54,9 +54,7 @@ void UMovieSceneSpawnTrack::PostLoad()
 
 UMovieSceneSection* UMovieSceneSpawnTrack::CreateNewSection()
 {
-	UMovieSceneSpawnSection* Section = NewObject<UMovieSceneSpawnSection>(this, NAME_None, RF_Transactional);
-	Section->GetCurve().SetDefaultValue(true);
-	return Section;
+	return NewObject<UMovieSceneSpawnSection>(this, NAME_None, RF_Transactional);
 }
 
 
@@ -84,12 +82,6 @@ bool UMovieSceneSpawnTrack::IsEmpty() const
 }
 
 
-TRange<float> UMovieSceneSpawnTrack::GetSectionBoundaries() const
-{
-	return TRange<float>::All();
-}
-
-
 const TArray<UMovieSceneSection*>& UMovieSceneSpawnTrack::GetAllSections() const
 {
 	return Sections;
@@ -128,17 +120,22 @@ ECookOptimizationFlags UMovieSceneSpawnTrack::GetCookOptimizationFlags() const
 	for (UMovieSceneSection* Section : Sections)
 	{
 		UMovieSceneSpawnSection* BoolSection = CastChecked<UMovieSceneSpawnSection>(Section);
+		TMovieSceneChannel<bool> BoolChannel = BoolSection->GetChannel().GetInterface();
+
 		if (!BoolSection->IsActive())
 		{
 			continue;
 		}
-		else if (BoolSection->GetCurve().GetNumKeys() == 0 && BoolSection->GetCurve().GetDefaultValue())
+
+		// If this bool section doesn't have any keys but does have a default value, this shouldn't be cooked out
+		if (BoolChannel.GetTimes().Num() == 0 && BoolSection->GetChannel().GetDefault().Get(false))
 		{
 			return ECookOptimizationFlags::None;
 		}
-		else for (auto It = BoolSection->GetCurve().GetKeyIterator(); It; ++It)
+		// If there are any keys that will cause this object to be spawned, we can't cook the object out
+		else for (bool Key : BoolChannel.GetValues())
 		{
-			if (It->Value)
+			if (Key != false)
 			{
 				return ECookOptimizationFlags::None;
 			}
