Copying //UE4/Dev-Sequencer to //UE4/Dev-Main (Source: //UE4/Dev-Sequencer @ 3946692)

diff --git a/Engine/Source/Runtime/MovieSceneTracks/Private/Tracks/MovieSceneStringTrack.cpp b/Engine/Source/Runtime/MovieSceneTracks/Private/Tracks/MovieSceneStringTrack.cpp
--- a/Engine/Source/Runtime/MovieSceneTracks/Private/Tracks/MovieSceneStringTrack.cpp
+++ b/Engine/Source/Runtime/MovieSceneTracks/Private/Tracks/MovieSceneStringTrack.cpp
@@ -9,35 +9,6 @@
 #define LOCTEXT_NAMESPACE "MovieSceneStringTrack"
 
 
-/* UMovieSceneStringTrack interface
- *****************************************************************************/
-
-bool UMovieSceneStringTrack::AddKeyToSection(float Time, const FString& String)
-{
-	UMovieSceneSection* TargetSection = MovieSceneHelpers::FindNearestSectionAtTime(Sections, Time);
-
-	if (TargetSection == nullptr)
-	{
-		TargetSection = CreateNewSection();
-		TargetSection->SetStartTime(Time);
-		TargetSection->SetEndTime(Time);
-
-		Sections.Add(TargetSection);
-	}
-
-	UMovieSceneStringSection* StringSection = Cast<UMovieSceneStringSection>(TargetSection);
-
-	if (StringSection == nullptr)
-	{
-		return false;
-	}
-
-	StringSection->AddKey(Time, String, EMovieSceneKeyInterpolation::Auto);
-
-	return true;
-}
-
-
 /* UMovieSceneTrack interface
  *****************************************************************************/
 
@@ -64,19 +35,6 @@ const TArray<UMovieSceneSection*>& UMovieSceneStringTrack::GetAllSections() cons
 }
 
 
-TRange<float> UMovieSceneStringTrack::GetSectionBoundaries() const
-{
-	TRange<float> SectionBoundaries = TRange<float>::Empty();
-
-	for (auto& Section : Sections)
-	{
-		SectionBoundaries = TRange<float>::Hull(SectionBoundaries, Section->GetRange());
-	}
-
-	return SectionBoundaries;
-}
-
-
 bool UMovieSceneStringTrack::HasSection(const UMovieSceneSection& Section) const
 {
 	return Sections.Contains(&Section);
diff --git a/Engine/Source/Runtime/MovieSceneTracks/Public/Evaluation/MovieScene3DTransformTemplate.h b/Engine/Source/Runtime/MovieSceneTracks/Public/Evaluation/MovieScene3DTransformTemplate.h
--- a/Engine/Source/Runtime/MovieSceneTracks/Public/Evaluation/MovieScene3DTransformTemplate.h
+++ b/Engine/Source/Runtime/MovieSceneTracks/Public/Evaluation/MovieScene3DTransformTemplate.h
@@ -8,6 +8,7 @@
 #include "Evaluation/MovieSceneEvalTemplate.h"
 #include "Sections/MovieScene3DTransformSection.h"
 #include "Evaluation/Blending/BlendableToken.h"
+#include "Channels/MovieSceneFloatChannel.h"
 #include "Evaluation/Blending/MovieSceneMultiChannelBlending.h"
 #include "MovieScene3DTransformTemplate.generated.h"
 
@@ -21,19 +22,19 @@ struct FMovieScene3DTransformTemplateData
 	FMovieScene3DTransformTemplateData(){}
 	FMovieScene3DTransformTemplateData(const UMovieScene3DTransformSection& Section);
 
-	MovieScene::TMultiChannelValue<float, 9> Evaluate(float InTime) const;
+	MovieScene::TMultiChannelValue<float, 9> Evaluate(FFrameTime InTime) const;
 
 	UPROPERTY()
-	FRichCurve TranslationCurve[3];
+	FMovieSceneFloatChannel TranslationCurve[3];
 
 	UPROPERTY()
-	FRichCurve RotationCurve[3];
+	FMovieSceneFloatChannel RotationCurve[3];
 
 	UPROPERTY()
-	FRichCurve ScaleCurve[3];
+	FMovieSceneFloatChannel ScaleCurve[3];
 
 	UPROPERTY()
-	FRichCurve ManualWeight;
+	FMovieSceneFloatChannel ManualWeight;
 
 	UPROPERTY() 
 	EMovieSceneBlendType BlendType;
diff --git a/Engine/Source/Runtime/MovieSceneTracks/Public/Evaluation/MovieSceneSpawnTemplate.h b/Engine/Source/Runtime/MovieSceneTracks/Public/Evaluation/MovieSceneSpawnTemplate.h
--- a/Engine/Source/Runtime/MovieSceneTracks/Public/Evaluation/MovieSceneSpawnTemplate.h
+++ b/Engine/Source/Runtime/MovieSceneTracks/Public/Evaluation/MovieSceneSpawnTemplate.h
@@ -4,7 +4,7 @@
 
 #include "CoreMinimal.h"
 #include "UObject/ObjectMacros.h"
-#include "Curves/IntegralCurve.h"
+#include "Channels/MovieSceneBoolChannel.h"
 #include "Evaluation/MovieSceneAnimTypeID.h"
 #include "Evaluation/MovieSceneEvalTemplate.h"
 #include "MovieSceneSpawnTemplate.generated.h"
@@ -29,5 +29,5 @@ private:
 
 protected:
 	UPROPERTY()
-	FIntegralCurve Curve;
+	FMovieSceneBoolChannel Curve;
 };
diff --git a/Engine/Source/Runtime/MovieSceneTracks/Public/Sections/MovieScene3DAttachSection.h b/Engine/Source/Runtime/MovieSceneTracks/Public/Sections/MovieScene3DAttachSection.h
--- a/Engine/Source/Runtime/MovieSceneTracks/Public/Sections/MovieScene3DAttachSection.h
+++ b/Engine/Source/Runtime/MovieSceneTracks/Public/Sections/MovieScene3DAttachSection.h
@@ -23,14 +23,12 @@ class UMovieScene3DAttachSection
 public:
 
 	/** 
-	 * Adds an attach to the section
+	 * Sets the object to attach to
 	 *
-	 * @param Time	The location in time where the attach should be added
-	 * @param SequenceEndTime   The time at the end of the sequence, by default the attach is set to end at this time
 	 * @param InAttachBindingId The object binding id to the path
 	 */
-	void AddAttach(float Time, float SequenceEndTime, const FMovieSceneObjectBindingID& InAttachBindingID);
-	
+	void SetAttachTargetID(const FMovieSceneObjectBindingID& InAttachBindingID);
+
 public:
 
 	UPROPERTY(EditAnywhere, Category="Attach")
diff --git a/Engine/Source/Runtime/MovieSceneTracks/Public/Sections/MovieScene3DConstraintSection.h b/Engine/Source/Runtime/MovieSceneTracks/Public/Sections/MovieScene3DConstraintSection.h
--- a/Engine/Source/Runtime/MovieSceneTracks/Public/Sections/MovieScene3DConstraintSection.h
+++ b/Engine/Source/Runtime/MovieSceneTracks/Public/Sections/MovieScene3DConstraintSection.h
@@ -44,8 +44,6 @@ public:
 
 	//~ UMovieSceneSection interface
 
-	virtual TOptional<float> GetKeyTime(FKeyHandle KeyHandle) const override;
-	virtual void SetKeyTime(FKeyHandle KeyHandle, float Time) override;
 	virtual void OnBindingsUpdated(const TMap<FGuid, FGuid>& OldGuidToNewGuidMap) override;
 	
 	/** ~UObject interface */
diff --git a/Engine/Source/Runtime/MovieSceneTracks/Public/Sections/MovieScene3DPathSection.h b/Engine/Source/Runtime/MovieSceneTracks/Public/Sections/MovieScene3DPathSection.h
--- a/Engine/Source/Runtime/MovieSceneTracks/Public/Sections/MovieScene3DPathSection.h
+++ b/Engine/Source/Runtime/MovieSceneTracks/Public/Sections/MovieScene3DPathSection.h
@@ -5,7 +5,7 @@
 #include "CoreMinimal.h"
 #include "UObject/ObjectMacros.h"
 #include "Curves/KeyHandle.h"
-#include "Curves/RichCurve.h"
+#include "Channels/MovieSceneFloatChannel.h"
 #include "Sections/MovieScene3DConstraintSection.h"
 #include "MovieScene3DPathSection.generated.h"
 
@@ -35,30 +35,22 @@ class UMovieScene3DPathSection
 
 public:
 
+	virtual void InitialPlacement(const TArray<UMovieSceneSection*>& Sections, FFrameNumber InStartTime, int32 Duration, bool bAllowMultipleRows) override;
+
 	/**
 	 * Evaluates the path track.
 	 *
 	 * @param Time The position in time within the movie scene.
 	 */
-	void Eval(USceneComponent* SceneComponent, float Time, USplineComponent* SplineComponent, FVector& OutTranslation, FRotator& OutRotation) const;
+	void Eval(USceneComponent* SceneComponent, FFrameTime Time, USplineComponent* SplineComponent, FVector& OutTranslation, FRotator& OutRotation) const;
 
 	/** 
-	 * Adds a path to the section.
+	 * Sets the path binding ID
 	 *
-	 * @param Time The location in time where the path should be added.
-	 * @param SequenceEndTime The time at the end of the sequence, by default the path is set to end at this time.
 	 * @param InPathBindingID The object binding id to the path.
 	 */
-	void AddPath(float Time, float SequenceEndTime, const FMovieSceneObjectBindingID& InPathBindingID);
+	void SetPathBindingID(const FMovieSceneObjectBindingID& InPathBindingID);
 
-	/** 
-	 * Returns the timing curve.
-	 *
-	 * @return The timing curve.
-	 */
-	MOVIESCENETRACKS_API FRichCurve& GetTimingCurve() { return TimingCurve; }
-	MOVIESCENETRACKS_API const FRichCurve& GetTimingCurve() const { return TimingCurve; }
-	
 	MOVIESCENETRACKS_API MovieScene3DPathSection_Axis GetFrontAxisEnum() const { return FrontAxisEnum; }
 	MOVIESCENETRACKS_API MovieScene3DPathSection_Axis GetUpAxisEnum() const { return UpAxisEnum; }
 	MOVIESCENETRACKS_API bool GetFollow() const { return bFollow; }
@@ -67,17 +59,13 @@ public:
 
 public:
 
-	//~ MovieSceneSection interface
-
-	virtual void MoveSection(float DeltaPosition, TSet<FKeyHandle>& KeyHandles) override;
-	virtual void DilateSection(float DilationFactor, float Origin, TSet<FKeyHandle>& KeyHandles) override;
-	virtual void GetKeyHandles(TSet<FKeyHandle>& OutKeyHandles, TRange<float> TimeRange) const override;
+	const FMovieSceneFloatChannel& GetTimingChannel() const { return TimingCurve; }
 
 private:
 
 	/** Timing Curve */
 	UPROPERTY(EditAnywhere, Category="Path")
-	FRichCurve TimingCurve;
+	FMovieSceneFloatChannel TimingCurve;
 
 	/** Front Axis */
 	UPROPERTY(EditAnywhere, Category="Path")
diff --git a/Engine/Source/Runtime/MovieSceneTracks/Public/Sections/MovieScene3DTransformSection.h b/Engine/Source/Runtime/MovieSceneTracks/Public/Sections/MovieScene3DTransformSection.h
--- a/Engine/Source/Runtime/MovieSceneTracks/Public/Sections/MovieScene3DTransformSection.h
+++ b/Engine/Source/Runtime/MovieSceneTracks/Public/Sections/MovieScene3DTransformSection.h
@@ -6,24 +6,11 @@
 #include "UObject/ObjectMacros.h"
 #include "Components/SceneComponent.h"
 #include "Curves/KeyHandle.h"
-#include "Curves/RichCurve.h"
 #include "MovieSceneSection.h"
-#include "Sections/IKeyframeSection.h"
 #include "MovieSceneKeyStruct.h"
+#include "Channels/MovieSceneFloatChannel.h"
 #include "MovieScene3DTransformSection.generated.h"
 
-class FStructOnScope;
-
-namespace EKey3DTransformChannel
-{
-	enum Type
-	{
-		Translation = 0x00000001,
-		Rotation = 0x00000002,
-		Scale = 0x00000004,
-		All = Translation | Rotation | Scale
-	};
-}
 
 #if WITH_EDITORONLY_DATA
 /** Visibility options for 3d trajectory. */
@@ -36,6 +23,7 @@ enum class EShow3DTrajectory : uint8
 };
 #endif
 
+
 /**
 * Stores information about a transform for the purpose of adding keys to a transform section
 */
@@ -47,10 +35,12 @@ struct FTransformData
 	FRotator Rotation;
 	/** Scale component */
 	FVector Scale;
-	/** Whether or not the data is valid (any values set) */
-	bool bValid;
 
-	bool IsValid() const { return bValid; }
+	FTransformData()
+		: Translation( ForceInitToZero )
+		, Rotation( ForceInitToZero )
+		, Scale( ForceInitToZero )
+	{}
 
 	/**
 	* Constructor.  Builds the data from a scene component
@@ -62,31 +52,9 @@ struct FTransformData
 		: Translation( InComponent->RelativeLocation )
 		, Rotation( InComponent->RelativeRotation )
 		, Scale( InComponent->RelativeScale3D )
-		, bValid( true )
 	{}
-
-	FTransformData()
-		: bValid( false )
-	{}
-};
-
-
-struct FTransformKey
-{
-	FTransformKey( EKey3DTransformChannel::Type InChannel, EAxis::Type InAxis, float InValue, bool InbUnwindRotation )
-	{
-		Channel = InChannel;
-		Axis = InAxis;
-		Value = InValue;
-		bUnwindRotation = InbUnwindRotation;
-	}
-	EKey3DTransformChannel::Type Channel;
-	EAxis::Type Axis;
-	float Value;
-	bool bUnwindRotation;
 };
 
-
 /**
  * Proxy structure for translation keys in 3D transform sections.
  */
@@ -102,12 +70,13 @@ struct FMovieScene3DLocationKeyStruct
 
 	/** The key's time. */
 	UPROPERTY(EditAnywhere, Category=Key)
-	float Time;
+	FFrameNumber Time;
 
-	FRichCurveKey* LocationKeys[3];
+	FMovieSceneKeyStructHelper KeyStructInterop;
 
 	virtual void PropagateChanges(const FPropertyChangedEvent& ChangeEvent) override;
 };
+template<> struct TStructOpsTypeTraits<FMovieScene3DLocationKeyStruct> : public TStructOpsTypeTraitsBase2<FMovieScene3DLocationKeyStruct> { enum { WithCopy = false }; };
 
 
 /**
@@ -125,13 +94,13 @@ struct FMovieScene3DRotationKeyStruct
 
 	/** The key's time. */
 	UPROPERTY(EditAnywhere, Category=Key)
-	float Time;
+	FFrameNumber Time;
 
-	FRichCurveKey* RotationKeys[3];
+	FMovieSceneKeyStructHelper KeyStructInterop;
 
 	virtual void PropagateChanges(const FPropertyChangedEvent& ChangeEvent) override;
 };
-
+template<> struct TStructOpsTypeTraits<FMovieScene3DRotationKeyStruct> : public TStructOpsTypeTraitsBase2<FMovieScene3DRotationKeyStruct> { enum { WithCopy = false }; };
 
 /**
  * Proxy structure for translation keys in 3D transform sections.
@@ -148,12 +117,13 @@ struct FMovieScene3DScaleKeyStruct
 
 	/** The key's time. */
 	UPROPERTY(EditAnywhere, Category=Key)
-	float Time;
+	FFrameNumber Time;
 
-	FRichCurveKey* ScaleKeys[3];
+	FMovieSceneKeyStructHelper KeyStructInterop;
 
 	virtual void PropagateChanges(const FPropertyChangedEvent& ChangeEvent) override;
 };
+template<> struct TStructOpsTypeTraits<FMovieScene3DScaleKeyStruct> : public TStructOpsTypeTraitsBase2<FMovieScene3DScaleKeyStruct> { enum { WithCopy = false }; };
 
 
 /**
@@ -179,14 +149,13 @@ struct FMovieScene3DTransformKeyStruct
 
 	/** The key's time. */
 	UPROPERTY(EditAnywhere, Category=Key)
-	float Time;
+	FFrameNumber Time;
 
-	FRichCurveKey* LocationKeys[3];
-	FRichCurveKey* RotationKeys[3];
-	FRichCurveKey* ScaleKeys[3];
+	FMovieSceneKeyStructHelper KeyStructInterop;
 
 	virtual void PropagateChanges(const FPropertyChangedEvent& ChangeEvent) override;
 };
+template<> struct TStructOpsTypeTraits<FMovieScene3DTransformKeyStruct> : public TStructOpsTypeTraitsBase2<FMovieScene3DTransformKeyStruct> { enum { WithCopy = false }; };
 
 enum class EMovieSceneTransformChannel : uint32
 {
@@ -272,109 +241,29 @@ private:
 UCLASS(MinimalAPI)
 class UMovieScene3DTransformSection
 	: public UMovieSceneSection
-	, public IKeyframeSection<FTransformKey>
 {
 	GENERATED_UCLASS_BODY()
 
 public:
 
 	/**
-	 * Evaluates the translation component of the transform
-	 *
-	 * @param Time				The position in time within the movie scene
-	 * @param OutTranslation	The evaluated translation.  Note: will remain unchanged if there were no keys to evaluate
-	 */
-	MOVIESCENETRACKS_API void EvalTranslation( float Time, FVector& InOutTranslation ) const;
-
-	/**
-	 * Evaluates the rotation component of the transform
-	 *
-	 * @param Time				The position in time within the movie scene
-	 * @param OutRotation		The evaluated rotation.  Note: will remain unchanged if there were no keys to evaluate
-	 */
-	MOVIESCENETRACKS_API void EvalRotation( float Time, FRotator& InOutRotation ) const;
-
-	/**
-	 * Evaluates the scale component of the transform
-	 *
-	 * @param Time				The position in time within the movie scene
-	 * @param OutScale			The evaluated scale.  Note: will remain unchanged if there were no keys to evaluate
-	 */
-	MOVIESCENETRACKS_API void EvalScale( float Time, FVector& InOutScale ) const;
-
-	/** 
-	 * Returns the translation curve for a specific axis
-	 *
-	 * @param Axis	The axis of the translation curve to get
-	 * @return The curve on the axis
-	 */
-	MOVIESCENETRACKS_API FRichCurve& GetTranslationCurve( EAxis::Type Axis );
-	MOVIESCENETRACKS_API const FRichCurve& GetTranslationCurve( EAxis::Type Axis ) const;
-
-	/** 
-	 * Returns the rotation curve for a specific axis
-	 *
-	 * @param Axis	The axis of the rotation curve to get
-	 * @return The curve on the axis
-	 */
-	MOVIESCENETRACKS_API FRichCurve& GetRotationCurve( EAxis::Type Axis );
-	MOVIESCENETRACKS_API const FRichCurve& GetRotationCurve( EAxis::Type Axis ) const;
-
-	/** 
-	 * Returns the scale curve for a specific axis
-	 *
-	 * @param Axis	The axis of the scale curve to get
-	 * @return The curve on the axis
-	 */
-	MOVIESCENETRACKS_API FRichCurve& GetScaleCurve( EAxis::Type Axis );
-	MOVIESCENETRACKS_API const FRichCurve& GetScaleCurve( EAxis::Type Axis ) const;
-
-	/**
-	 * Returns the manual weight curve for this section
-	 * @return The manual weight curve
+	 * Access the mask that defines which channels this track should animate
 	 */
-	MOVIESCENETRACKS_API FRichCurve& GetManualWeightCurve();
-	MOVIESCENETRACKS_API const FRichCurve& GetManualWeightCurve() const;
-
-	FMovieSceneTransformMask GetMask() const
-	{
-		return TransformMask;
-	}
-
-	void SetMask(FMovieSceneTransformMask NewMask)
-	{
-		TransformMask = NewMask;
-	}
+	MOVIESCENETRACKS_API FMovieSceneTransformMask GetMask() const;
 
 	/**
-	 * Return the trajectory visibility
+	 * Set the mask that defines which channels this track should animate
 	 */
-#if WITH_EDITORONLY_DATA
-	MOVIESCENETRACKS_API EShow3DTrajectory GetShow3DTrajectory() { return Show3DTrajectory; }
-#endif
-
-public:
+	MOVIESCENETRACKS_API void SetMask(FMovieSceneTransformMask NewMask);
 
-	// UMovieSceneSection interface
+protected:
 
-	virtual void MoveSection( float DeltaPosition, TSet<FKeyHandle>& KeyHandles ) override;
-	virtual void DilateSection( float DilationFactor, float Origin, TSet<FKeyHandle>& KeyHandles ) override;
-	virtual void GetKeyHandles(TSet<FKeyHandle>& OutKeyHandles, TRange<float> TimeRange) const override;
-	virtual TSharedPtr<FStructOnScope> GetKeyStruct(const TArray<FKeyHandle>& KeyHandles) override;
-	virtual TOptional<float> GetKeyTime( FKeyHandle KeyHandle ) const override;
-	virtual void SetKeyTime( FKeyHandle KeyHandle, float Time ) override;
-
-public:
+	virtual void Serialize(FArchive& Ar) override;
+	virtual TSharedPtr<FStructOnScope> GetKeyStruct(TArrayView<const FKeyHandle> KeyHandles) override;
+	virtual FMovieSceneEvalTemplatePtr GenerateTemplate() const override;
 
-	// IKeyframeSection interface.
+	void UpdateChannelProxy();
 
-	virtual bool NewKeyIsNewData( float Time, const FTransformKey& KeyData ) const override;
-	virtual bool HasKeys( const FTransformKey& KeyData ) const override;
-	virtual void AddKey( float Time, const FTransformKey& KeyData, EMovieSceneKeyInterpolation KeyInterpolation ) override;
-	virtual void SetDefault( const FTransformKey& KeyData ) override;
-	virtual void ClearDefaults() override;
-	virtual FMovieSceneEvalTemplatePtr GenerateTemplate() const override;
-	
 private:
 
 	UPROPERTY()
@@ -382,23 +271,37 @@ private:
 
 	/** Translation curves */
 	UPROPERTY()
-	FRichCurve Translation[3];
+	FMovieSceneFloatChannel Translation[3];
 	
 	/** Rotation curves */
 	UPROPERTY()
-	FRichCurve Rotation[3];
+	FMovieSceneFloatChannel Rotation[3];
 
 	/** Scale curves */
 	UPROPERTY()
-	FRichCurve Scale[3];
+	FMovieSceneFloatChannel Scale[3];
 
 	/** Manual weight curve */
 	UPROPERTY()
-	FRichCurve ManualWeight;
+	FMovieSceneFloatChannel ManualWeight;
+
+	/** Unserialized mask that defines the mask of the current channel proxy so we don't needlessly re-create it on post-undo */
+	EMovieSceneTransformChannel ProxyChannels;
 
 #if WITH_EDITORONLY_DATA
+
+public:
+
+	/**
+	 * Return the trajectory visibility
+	 */
+	MOVIESCENETRACKS_API EShow3DTrajectory GetShow3DTrajectory() { return Show3DTrajectory; }
+
+private:
+
 	/** Whether to show the 3d trajectory */
 	UPROPERTY(EditAnywhere, DisplayName = "Show 3D Trajectory", Category = "Transform")
 	EShow3DTrajectory Show3DTrajectory;
-#endif
+
+#endif // WITH_EDITORONLY_DATA
 };
diff --git a/Engine/Source/Runtime/MovieSceneTracks/Public/Sections/MovieSceneActorReferenceSection.h b/Engine/Source/Runtime/MovieSceneTracks/Public/Sections/MovieSceneActorReferenceSection.h
--- a/Engine/Source/Runtime/MovieSceneTracks/Public/Sections/MovieSceneActorReferenceSection.h
+++ b/Engine/Source/Runtime/MovieSceneTracks/Public/Sections/MovieSceneActorReferenceSection.h
@@ -6,77 +6,171 @@
 #include "UObject/ObjectMacros.h"
 #include "Misc/Guid.h"
 #include "Curves/KeyHandle.h"
+#include "MovieSceneObjectBindingID.h"
 #include "MovieSceneSection.h"
 #include "Curves/IntegralCurve.h"
-#include "Sections/IKeyframeSection.h"
+#include "Channels/MovieSceneChannel.h"
 #include "MovieSceneActorReferenceSection.generated.h"
 
-/**
- * A single actor reference point section
- */
-UCLASS(MinimalAPI)
-class UMovieSceneActorReferenceSection
-	: public UMovieSceneSection
-	, public IKeyframeSection<FGuid>
+USTRUCT()
+struct FMovieSceneActorReferenceKey
 {
-	GENERATED_UCLASS_BODY()
+	GENERATED_BODY()
 
-public:
+	FMovieSceneActorReferenceKey()
+	{}
+
+	FMovieSceneActorReferenceKey(const FMovieSceneObjectBindingID& InBindingID)
+		: Object(InBindingID)
+	{}
+
+	friend bool operator==(const FMovieSceneActorReferenceKey& A, const FMovieSceneActorReferenceKey& B)
+	{
+		return A.Object == B.Object;
+	}
+
+	friend bool operator!=(const FMovieSceneActorReferenceKey& A, const FMovieSceneActorReferenceKey& B)
+	{
+		return A.Object != B.Object;
+	}
+
+	UPROPERTY(EditAnywhere, Category="Key")
+	FMovieSceneObjectBindingID Object;
+};
+
+/** A curve of events */
+USTRUCT()
+struct FMovieSceneActorReferenceData
+{
+	GENERATED_BODY()
+
+	FMovieSceneActorReferenceData()
+		: DefaultValue()
+	{}
+
+	/**
+	 * Access an integer that uniquely identifies this channel type.
+	 *
+	 * @return A static identifier that was allocated using FMovieSceneChannelEntry::RegisterNewID
+	 */
+	MOVIESCENETRACKS_API static uint32 GetChannelID();
+
+	/**
+	 * Access a mutable interface for this channel
+	 *
+	 * @return An object that is able to manipulate this channel
+	 */
+	FORCEINLINE TMovieSceneChannel<FMovieSceneActorReferenceKey> GetInterface()
+	{
+		return TMovieSceneChannel<FMovieSceneActorReferenceKey>(&KeyTimes, &KeyValues, &KeyHandles);
+	}
 
 	/**
-	 * Updates this section
+	 * Access a constant interface for this channel
 	 *
-	 * @param Position	The position in time within the movie scene
+	 * @return An object that is able to interrogate this channel
 	 */
-	MOVIESCENETRACKS_API FGuid Eval(float Position) const;
+	FORCEINLINE TMovieSceneChannel<const FMovieSceneActorReferenceKey> GetInterface() const
+	{
+		return TMovieSceneChannel<const FMovieSceneActorReferenceKey>(&KeyTimes, &KeyValues);
+	}
 
-	// IKeyframeSection interface.
+	/**
+	 * Evaluate this channel
+	 *
+	 * @param InTime     The time to evaluate at
+	 * @return the result of the evaluation
+	 */
+	MOVIESCENETRACKS_API FMovieSceneActorReferenceKey Evaluate(FFrameTime InTime) const;
 
-	virtual void AddKey(float Time, const FGuid& Value, EMovieSceneKeyInterpolation KeyInterpolation) override;
-	virtual bool NewKeyIsNewData(float Time, const FGuid& Value) const override;
-	virtual bool HasKeys(const FGuid& Value) const override;
-	virtual void SetDefault(const FGuid& Value) override;
-	virtual void ClearDefaults() override;
+public:
 
 	/**
-	 * @return The integral curve on this section
+	 * Set this channel's default value that should be used when no keys are present
+	 *
+	 * @param InDefaultValue The desired default value
 	 */
-	FIntegralCurve& GetActorReferenceCurve()
+	FORCEINLINE void SetDefault(FMovieSceneActorReferenceKey InDefaultValue)
 	{
-		return ActorGuidIndexCurve;
+		DefaultValue = InDefaultValue;
 	}
 
-	const FIntegralCurve& GetActorReferenceCurve() const { return ActorGuidIndexCurve; }
+	/**
+	 * Get this channel's default value that will be used when no keys are present
+	 *
+	 * @return (Optional) The channel's default value
+	 */
+	FORCEINLINE FMovieSceneActorReferenceKey GetDefault() const
+	{
+		return DefaultValue;
+	}
 
-	const TArray<FGuid>& GetActorGuids() const { return ActorGuids; }
+	/**
+	 * Upgrade legacy data by appending to the end of the array. Assumes sorted data
+	 */
+	void UpgradeLegacyTime(UObject* Context, double Time, FMovieSceneActorReferenceKey Value)
+	{
+		FFrameRate LegacyFrameRate = GetLegacyConversionFrameRate();
+		FFrameNumber KeyTime = UpgradeLegacyMovieSceneTime(nullptr, LegacyFrameRate, Time);
 
-public:
+		check(KeyTimes.Num() == 0 || KeyTime >= KeyTimes.Last());
 
-	MOVIESCENETRACKS_API FKeyHandle AddKey(float Time, const FGuid& Value);
+		KeyTimes.Add(KeyTime);
+		KeyValues.Add(Value);
+	}
+
+private:
+
+	/** Sorted array of key times */
+	UPROPERTY()
+	TArray<FFrameNumber> KeyTimes;
 
-	//~ UMovieSceneSection interface
+	/** Default value used when there are no keys */
+	FMovieSceneActorReferenceKey DefaultValue;
+
+	/** Array of values that correspond to each key time */
+	UPROPERTY()
+	TArray<FMovieSceneActorReferenceKey> KeyValues;
+
+	FMovieSceneKeyHandleMap KeyHandles;
+};
 
-	virtual void MoveSection(float DeltaPosition, TSet<FKeyHandle>& KeyHandles) override;
-	virtual void DilateSection(float DilationFactor, float Origin, TSet<FKeyHandle>& KeyHandles) override;
-	virtual void GetKeyHandles(TSet<FKeyHandle>& OutKeyHandles, TRange<float> TimeRange) const override;
-	virtual TOptional<float> GetKeyTime(FKeyHandle KeyHandle) const override;
-	virtual void SetKeyTime(FKeyHandle KeyHandle, float Time) override;
+/**
+ * A single actor reference point section
+ */
+UCLASS(MinimalAPI)
+class UMovieSceneActorReferenceSection
+	: public UMovieSceneSection
+{
+	GENERATED_UCLASS_BODY()
 
 public:
 
 	//~ UObject interface
-
-	virtual void PreSave(const class ITargetPlatform* TargetPlatform) override;
 	virtual void PostLoad() override;
 
+	const FMovieSceneActorReferenceData& GetActorReferenceData() const { return ActorReferenceData; }
+
 private:
 
-	/** Curve data */
 	UPROPERTY()
-	FIntegralCurve ActorGuidIndexCurve;
+	FMovieSceneActorReferenceData ActorReferenceData;
 
-	TArray<FGuid> ActorGuids;
+private:
 
+	/** Curve data */
 	UPROPERTY()
-	TArray<FString> ActorGuidStrings;
+	FIntegralCurve ActorGuidIndexCurve_DEPRECATED;
+
+	UPROPERTY()
+	TArray<FString> ActorGuidStrings_DEPRECATED;
 };
+
+inline bool EvaluateChannel(const FMovieSceneActorReferenceData* InChannel, FFrameTime InTime, FMovieSceneActorReferenceKey& OutValue)
+{
+	OutValue = InChannel->Evaluate(InTime);
+	return true;
+}
+
+inline void ClearChannelDefault(FMovieSceneActorReferenceData* InChannel)
+{}
diff --git a/Engine/Source/Runtime/MovieSceneTracks/Public/Sections/MovieSceneAudioSection.h b/Engine/Source/Runtime/MovieSceneTracks/Public/Sections/MovieSceneAudioSection.h
--- a/Engine/Source/Runtime/MovieSceneTracks/Public/Sections/MovieSceneAudioSection.h
+++ b/Engine/Source/Runtime/MovieSceneTracks/Public/Sections/MovieSceneAudioSection.h
@@ -4,10 +4,10 @@
 
 #include "CoreMinimal.h"
 #include "UObject/ObjectMacros.h"
-#include "Curves/KeyHandle.h"
 #include "MovieSceneSection.h"
 #include "Runtime/Engine/Classes/Components/AudioComponent.h"
 #include "Sound/SoundAttenuation.h"
+#include "Channels/MovieSceneFloatChannel.h"
 #include "MovieSceneAudioSection.generated.h"
 
 class USoundBase;
@@ -28,36 +28,26 @@ public:
 
 	/** Gets the sound for this section */
 	class USoundBase* GetSound() const {return Sound;}
-	
+
 	/** Set the offset into the beginning of the audio clip */
 	void SetStartOffset(float InStartOffset) {StartOffset = InStartOffset;}
-	
+
 	/** Get the offset into the beginning of the audio clip */
 	float GetStartOffset() const {return StartOffset;}
-	
-	/**
-	 * @return The range of times that the sound plays, truncated by the section limits
-	 */
-	TRange<float> GetAudioRange() const;
-	
-	DEPRECATED(4.15, "Audio true range no longer supported.")
-	TRange<float> GetAudioTrueRange() const;
-	
+
 	/**
 	 * Gets the sound volume curve
 	 *
 	 * @return The rich curve for this sound volume
 	 */
-	FRichCurve& GetSoundVolumeCurve() { return SoundVolume; }
-	const FRichCurve& GetSoundVolumeCurve() const { return SoundVolume; }
+	const FMovieSceneFloatChannel& GetSoundVolumeChannel() const { return SoundVolume; }
 
 	/**
 	 * Gets the sound pitch curve
 	 *
 	 * @return The rich curve for this sound pitch
 	 */
-	FRichCurve& GetPitchMultiplierCurve() { return PitchMultiplier; }
-	const FRichCurve& GetPitchMultiplierCurve() const { return PitchMultiplier; }
+	const FMovieSceneFloatChannel& GetPitchMultiplierChannel() const { return PitchMultiplier; }
 
 	/**
 	 * Return the sound volume
@@ -65,15 +55,12 @@ public:
 	 * @param InTime	The position in time within the movie scene
 	 * @return The volume the sound will be played with.
 	 */
-	float GetSoundVolume(float InTime) const { return SoundVolume.Eval(InTime); }
-
-	/**
-	 * Sets the sound volume
-	 *
-	 * @param InTime	The position in time within the movie scene
-	 * @param InVolume	The volume to set
-	 */
-	void SetSoundVolume(float InTime, float InVolume) { SoundVolume.AddKey(InTime, InVolume); }
+	float GetSoundVolume(FFrameTime InTime) const
+	{
+		float OutValue = 0.f;
+		SoundVolume.Evaluate(InTime, OutValue);
+		return OutValue;
+	}
 
 	/**
 	 * Return the pitch multiplier
@@ -81,26 +68,11 @@ public:
 	 * @param Position	The position in time within the movie scene
 	 * @return The pitch multiplier the sound will be played with.
 	 */
-	float GetPitchMultiplier(float InTime) const { return PitchMultiplier.Eval(InTime); }
-
-	/**
-	 * Set the pitch multiplier
-	 *
-	 * @param Position	The position in time within the movie scene
-	 * @param InPitch The pitch multiplier to set
-	 */
-	void SetPitchMultiplier(float InTime, float InPitchMultiplier) { PitchMultiplier.AddKey(InTime, InPitchMultiplier); }
-
-	/**
-	 * Returns whether or not a provided position in time is within the timespan of the audio range
-	 *
-	 * @param Position	The position to check
-	 * @return true if the position is within the timespan, false otherwise
-	 */
-	bool IsTimeWithinAudioRange( float Position ) const 
+	float GetPitchMultiplier(FFrameTime InTime) const
 	{
-		TRange<float> AudioRange = GetAudioRange();
-		return Position >= AudioRange.GetLowerBoundValue() && Position <= AudioRange.GetUpperBoundValue();
+		float OutValue = 0.f;
+		PitchMultiplier.Evaluate(InTime, OutValue);
+		return OutValue;
 	}
 
 	/**
@@ -166,16 +138,9 @@ public:
 
 public:
 
-	// MovieSceneSection interface
-
-	virtual void MoveSection( float DeltaPosition, TSet<FKeyHandle>& KeyHandles ) override;
-	virtual void DilateSection( float DilationFactor, float Origin, TSet<FKeyHandle>& KeyHandles ) override;
-	virtual UMovieSceneSection* SplitSection(float SplitTime) override;
-	virtual void GetKeyHandles(TSet<FKeyHandle>& OutKeyHandles, TRange<float> TimeRange) const override;
-	virtual void GetSnapTimes(TArray<float>& OutSnapTimes, bool bGetSectionBorders) const override;
-	virtual TOptional<float> GetOffsetTime() const override { return TOptional<float>(StartOffset); }
-	virtual TOptional<float> GetKeyTime( FKeyHandle KeyHandle ) const override { return TOptional<float>(); }
-	virtual void SetKeyTime( FKeyHandle KeyHandle, float Time ) override { }
+	//~ UMovieSceneSection interface
+	virtual UMovieSceneSection* SplitSection(FFrameNumber SplitTime) override;
+	virtual TOptional<FFrameTime> GetOffsetTime() const override;
 	virtual FMovieSceneEvalTemplatePtr GenerateTemplate() const override;
 
 private:
@@ -202,11 +167,11 @@ private:
 
 	/** The volume the sound will be played with. */
 	UPROPERTY( EditAnywhere, Category = "Audio" )
-	FRichCurve SoundVolume;
+	FMovieSceneFloatChannel SoundVolume;
 
 	/** The pitch multiplier the sound will be played with. */
 	UPROPERTY( EditAnywhere, Category = "Audio" )
-	FRichCurve PitchMultiplier;
+	FMovieSceneFloatChannel PitchMultiplier;
 
 	UPROPERTY( EditAnywhere, Category="Audio" )
 	bool bSuppressSubtitles;
diff --git a/Engine/Source/Runtime/MovieSceneTracks/Public/Sections/MovieSceneBoolSection.h b/Engine/Source/Runtime/MovieSceneTracks/Public/Sections/MovieSceneBoolSection.h
--- a/Engine/Source/Runtime/MovieSceneTracks/Public/Sections/MovieSceneBoolSection.h
+++ b/Engine/Source/Runtime/MovieSceneTracks/Public/Sections/MovieSceneBoolSection.h
@@ -4,10 +4,8 @@
 
 #include "CoreMinimal.h"
 #include "UObject/ObjectMacros.h"
-#include "Curves/KeyHandle.h"
 #include "MovieSceneSection.h"
-#include "Curves/IntegralCurve.h"
-#include "Sections/IKeyframeSection.h"
+#include "Channels/MovieSceneBoolChannel.h"
 #include "MovieSceneBoolSection.generated.h"
 
 /**
@@ -16,7 +14,6 @@
 UCLASS(MinimalAPI)
 class UMovieSceneBoolSection 
 	: public UMovieSceneSection
-	, public IKeyframeSection<bool>
 {
 	GENERATED_UCLASS_BODY()
 
@@ -28,36 +25,8 @@ public:
 
 public:
 
-	/**
-	 * Update this section.
-	 *
-	 * @param Position The position in time within the movie scene.
-	 */
-	virtual bool Eval(float Position, bool DefaultValue) const;
-
-	/** Gets all the keys of this boolean section. */
-	FIntegralCurve& GetCurve() { return BoolCurve; }
-	const FIntegralCurve& GetCurve() const { return BoolCurve; }
-	
-public:
-
-	//~ IKeyframeSection interface
-
-	virtual void AddKey(float Time, const bool& Value, EMovieSceneKeyInterpolation KeyInterpolation) override;
-	virtual void SetDefault(const bool& Value) override;
-	virtual bool NewKeyIsNewData(float Time, const bool& Value) const override;
-	virtual bool HasKeys(const bool& Value) const override;
-	virtual void ClearDefaults() override;
-
-public:
-
-	//~ UMovieSceneSection interface 
-
-	virtual void MoveSection(float DeltaPosition, TSet<FKeyHandle>& KeyHandles) override;
-	virtual void DilateSection(float DilationFactor, float Origin, TSet<FKeyHandle>& KeyHandles) override;
-	virtual void GetKeyHandles(TSet<FKeyHandle>& OutKeyHandles, TRange<float> TimeRange) const override;
-	virtual TOptional<float> GetKeyTime(FKeyHandle KeyHandle) const override;
-	virtual void SetKeyTime(FKeyHandle KeyHandle, float Time) override;
+	FMovieSceneBoolChannel& GetChannel() { return BoolCurve; }
+	const FMovieSceneBoolChannel& GetChannel() const { return BoolCurve; }
 
 public:
 
@@ -65,11 +34,9 @@ public:
 
 	virtual void PostLoad() override;
 
-private:
+protected:
 
 	/** Ordered curve data */
-	// @todo Sequencer This could be optimized by packing the bools separately
-	// but that may not be worth the effort
 	UPROPERTY()
-	FIntegralCurve BoolCurve;
+	FMovieSceneBoolChannel BoolCurve;
 };
diff --git a/Engine/Source/Runtime/MovieSceneTracks/Public/Sections/MovieSceneByteSection.h b/Engine/Source/Runtime/MovieSceneTracks/Public/Sections/MovieSceneByteSection.h
--- a/Engine/Source/Runtime/MovieSceneTracks/Public/Sections/MovieSceneByteSection.h
+++ b/Engine/Source/Runtime/MovieSceneTracks/Public/Sections/MovieSceneByteSection.h
@@ -4,10 +4,8 @@
 
 #include "CoreMinimal.h"
 #include "UObject/ObjectMacros.h"
-#include "Curves/KeyHandle.h"
 #include "MovieSceneSection.h"
-#include "Curves/IntegralCurve.h"
-#include "Sections/IKeyframeSection.h"
+#include "Channels/MovieSceneByteChannel.h"
 #include "MovieSceneByteSection.generated.h"
 
 /**
@@ -16,40 +14,12 @@
 UCLASS(MinimalAPI)
 class UMovieSceneByteSection 
 	: public UMovieSceneSection
-	, public IKeyframeSection<uint8>
 {
 	GENERATED_UCLASS_BODY()
-public:
-
-	/** Gets all the keys of this byte section */
-	FIntegralCurve& GetCurve() { return ByteCurve; }
-	const FIntegralCurve& GetCurve() const { return ByteCurve; }
 
 public:
 
-	//~ IKeyframeSection interface
-
-	virtual void AddKey(float Time, const uint8& Value, EMovieSceneKeyInterpolation KeyInterpolation) override;
-	virtual bool NewKeyIsNewData(float Time, const uint8& Value) const override;
-	virtual bool HasKeys(const uint8& Value) const override;
-	virtual void SetDefault(const uint8& Value) override;
-	virtual void ClearDefaults() override;
-
-public:
-
-	//~ UMovieSceneSection interface 
-
-	virtual void MoveSection(float DeltaPosition, TSet<FKeyHandle>& KeyHandles) override;
-	virtual void DilateSection(float DilationFactor, float Origin, TSet<FKeyHandle>& KeyHandles) override;
-	virtual void GetKeyHandles(TSet<FKeyHandle>& OutKeyHandles, TRange<float> TimeRange) const override;
-	virtual TOptional<float> GetKeyTime(FKeyHandle KeyHandle) const override;
-	virtual void SetKeyTime(FKeyHandle KeyHandle, float Time) override;
-
-private:
-
 	/** Ordered curve data */
-	// @todo Sequencer This could be optimized by packing the bytes separately
-	// but that may not be worth the effort
 	UPROPERTY()
-	FIntegralCurve ByteCurve;
+	FMovieSceneByteChannel ByteCurve;
 };
diff --git a/Engine/Source/Runtime/MovieSceneTracks/Public/Sections/MovieSceneCameraCutSection.h b/Engine/Source/Runtime/MovieSceneTracks/Public/Sections/MovieSceneCameraCutSection.h
--- a/Engine/Source/Runtime/MovieSceneTracks/Public/Sections/MovieSceneCameraCutSection.h
+++ b/Engine/Source/Runtime/MovieSceneTracks/Public/Sections/MovieSceneCameraCutSection.h
@@ -55,11 +55,8 @@ public:
 		CameraBindingID = InCameraBindingID;
 	}
 
+	//~ UMovieSceneSection interface
 	virtual FMovieSceneEvalTemplatePtr GenerateTemplate() const override;
-
-	// UMovieSceneSection interface
-	virtual TOptional<float> GetKeyTime(FKeyHandle KeyHandle) const override { return TOptional<float>(); }
-	virtual void SetKeyTime(FKeyHandle KeyHandle, float Time) override { }
 	virtual void OnBindingsUpdated(const TMap<FGuid, FGuid>& OldGuidToNewGuidMap) override;
 
 	/** ~UObject interface */
diff --git a/Engine/Source/Runtime/MovieSceneTracks/Public/Sections/MovieSceneCameraShakeSection.h b/Engine/Source/Runtime/MovieSceneTracks/Public/Sections/MovieSceneCameraShakeSection.h
--- a/Engine/Source/Runtime/MovieSceneTracks/Public/Sections/MovieSceneCameraShakeSection.h
+++ b/Engine/Source/Runtime/MovieSceneTracks/Public/Sections/MovieSceneCameraShakeSection.h
@@ -50,7 +50,7 @@ public:
 	virtual FMovieSceneEvalTemplatePtr GenerateTemplate() const override;
 	virtual void PostLoad() override;
 	
-	UPROPERTY(EditAnywhere, Category="Cmaera Shake", meta=(ShowOnlyInnerProperties))
+	UPROPERTY(EditAnywhere, Category="Camera Shake", meta=(ShowOnlyInnerProperties))
 	FMovieSceneCameraShakeSectionData ShakeData;
 
 public:
diff --git a/Engine/Source/Runtime/MovieSceneTracks/Public/Sections/MovieSceneColorSection.h b/Engine/Source/Runtime/MovieSceneTracks/Public/Sections/MovieSceneColorSection.h
--- a/Engine/Source/Runtime/MovieSceneTracks/Public/Sections/MovieSceneColorSection.h
+++ b/Engine/Source/Runtime/MovieSceneTracks/Public/Sections/MovieSceneColorSection.h
@@ -7,37 +7,12 @@
 #include "Curves/KeyHandle.h"
 #include "Curves/RichCurve.h"
 #include "MovieSceneSection.h"
-#include "Sections/IKeyframeSection.h"
 #include "MovieSceneKeyStruct.h"
+#include "Channels/MovieSceneFloatChannel.h"
 #include "MovieSceneColorSection.generated.h"
 
-class FStructOnScope;
 struct FPropertyChangedEvent;
 
-enum class EKeyColorChannel
-{
-	Red,
-	Green,
-	Blue,
-	Alpha,
-};
-
-
-struct FColorKey
-{
-	FColorKey(EKeyColorChannel InChannel, float InChannelValue, bool InbIsSlateColor)
-	{
-		Channel = InChannel;
-		ChannelValue = InChannelValue;
-		bIsSlateColor = InbIsSlateColor;
-	}
-
-	EKeyColorChannel Channel;
-	float ChannelValue;
-	bool bIsSlateColor;
-};
-
-
 /**
  * Proxy structure for color section key data.
  */
@@ -53,13 +28,13 @@ struct FMovieSceneColorKeyStruct
 
 	/** The key's time. */
 	UPROPERTY(EditAnywhere, Category=Key)
-	float Time;
+	FFrameNumber Time;
 
-	FRichCurveKey* Keys[4];
-	FRichCurve* Curves[4];
+	FMovieSceneKeyStructHelper KeyStructInterop;
 
 	virtual void PropagateChanges(const FPropertyChangedEvent& ChangeEvent) override;
 };
+template<> struct TStructOpsTypeTraits<FMovieSceneColorKeyStruct> : public TStructOpsTypeTraitsBase2<FMovieSceneColorKeyStruct> { enum { WithCopy = false }; };
 
 
 /**
@@ -68,84 +43,43 @@ struct FMovieSceneColorKeyStruct
 UCLASS(MinimalAPI)
 class UMovieSceneColorSection 
 	: public UMovieSceneSection
-	, public IKeyframeSection<FColorKey>
 {
 	GENERATED_UCLASS_BODY()
 
 public:
 
-	/**
-	 * Gets the red color curve
-	 *
-	 * @return The rich curve for this color channel
-	 * @see GetAlphaCurve, GetBlueCurve, GetGreenCurve
-	 */
-	FRichCurve& GetRedCurve() { return RedCurve; }
-	const FRichCurve& GetRedCurve() const { return RedCurve; }
-
-	/**
-	 * Gets the green color curve
-	 *
-	 * @return The rich curve for this color channel
-	 * @see GetAlphaCurve, GetBlueCurve, GetRedCurve
-	 */
-	FRichCurve& GetGreenCurve() { return GreenCurve; }
-	const FRichCurve& GetGreenCurve() const { return GreenCurve; }
-
-	/**
-	 * Gets the blue color curve
-	 *
-	 * @return The rich curve for this color channel
-	 * @see GetAlphaCurve, GetGreenCurve, GetRedCurve
-	 */
-	FRichCurve& GetBlueCurve() { return BlueCurve; }
-	const FRichCurve& GetBlueCurve() const { return BlueCurve; }
-	
-	/**
-	 * Gets the alpha color curve
-	 *
-	 * @return The rich curve for this color channel
-	 * @see GetBlueCurve, GetGreenCurve, GetRedCurve
-	 */
-	FRichCurve& GetAlphaCurve() { return AlphaCurve; }
-	const FRichCurve& GetAlphaCurve() const { return AlphaCurve; }
-
-public:
+	/** Access the underlying generation function for the red component of this section */
+	const FMovieSceneFloatChannel& GetRedChannel() const { return RedCurve; }
 
-	//~ UMovieSceneSection interface
+	/** Access the underlying generation function for the green component of this section */
+	const FMovieSceneFloatChannel& GetGreenChannel() const { return GreenCurve; }
 
-	virtual void MoveSection(float DeltaPosition, TSet<FKeyHandle>& KeyHandles) override;
-	virtual void DilateSection(float DilationFactor, float Origin, TSet<FKeyHandle>& KeyHandles) override;
-	virtual void GetKeyHandles(TSet<FKeyHandle>& OutKeyHandles, TRange<float> TimeRange) const override;
-	virtual TSharedPtr<FStructOnScope> GetKeyStruct(const TArray<FKeyHandle>& KeyHandles) override;
-	virtual TOptional<float> GetKeyTime(FKeyHandle KeyHandle) const override;
-	virtual void SetKeyTime(FKeyHandle KeyHandle, float Time) override;
+	/** Access the underlying generation function for the blue component of this section */
+	const FMovieSceneFloatChannel& GetBlueChannel() const { return BlueCurve; }
 
-public:
+	/** Access the underlying generation function for the alpha component of this section */
+	const FMovieSceneFloatChannel& GetAlphaChannel() const { return AlphaCurve; }
 
-	//~ IKeyframeSection interface
+protected:
 
-	virtual void AddKey(float Time, const FColorKey& Key, EMovieSceneKeyInterpolation KeyInterpolation) override;
-	virtual bool NewKeyIsNewData(float Time, const FColorKey& Key) const override;
-	virtual bool HasKeys(const FColorKey& Key) const override;
-	virtual void SetDefault(const FColorKey& Key) override;
-	virtual void ClearDefaults() override;
+	//~ UMovieSceneSection interface
+	virtual TSharedPtr<FStructOnScope> GetKeyStruct(TArrayView<const FKeyHandle> KeyHandles) override;
 
 private:
 
 	/** Red curve data */
 	UPROPERTY()
-	FRichCurve RedCurve;
+	FMovieSceneFloatChannel RedCurve;
 
 	/** Green curve data */
 	UPROPERTY()
-	FRichCurve GreenCurve;
+	FMovieSceneFloatChannel GreenCurve;
 
 	/** Blue curve data */
 	UPROPERTY()
-	FRichCurve BlueCurve;
+	FMovieSceneFloatChannel BlueCurve;
 
 	/** Alpha curve data */
 	UPROPERTY()
-	FRichCurve AlphaCurve;
+	FMovieSceneFloatChannel AlphaCurve;
 };
diff --git a/Engine/Source/Runtime/MovieSceneTracks/Public/Sections/MovieSceneEnumSection.h b/Engine/Source/Runtime/MovieSceneTracks/Public/Sections/MovieSceneEnumSection.h
--- a/Engine/Source/Runtime/MovieSceneTracks/Public/Sections/MovieSceneEnumSection.h
+++ b/Engine/Source/Runtime/MovieSceneTracks/Public/Sections/MovieSceneEnumSection.h
@@ -4,7 +4,7 @@
 
 #include "Curves/IntegralCurve.h"
 #include "MovieSceneSection.h"
-#include "Sections/IKeyframeSection.h"
+#include "Channels/MovieSceneByteChannel.h"
 #include "MovieSceneEnumSection.generated.h"
 
 
@@ -12,42 +12,14 @@
  * A single enum section.
  */
 UCLASS(MinimalAPI)
-class UMovieSceneEnumSection 
+class UMovieSceneEnumSection
 	: public UMovieSceneSection
-	, public IKeyframeSectionEnum
 {
 	GENERATED_UCLASS_BODY()
-public:
-
-	/** Gets all the keys of this enum section */
-	FIntegralCurve& GetCurve() { return EnumCurve; }
-	const FIntegralCurve& GetCurve() const { return EnumCurve; }
 
 public:
 
-	//~ IKeyframeSection interface
-
-	virtual void AddKey(float Time, const int64& Value, EMovieSceneKeyInterpolation KeyInterpolation) override;
-	virtual bool NewKeyIsNewData(float Time, const int64& Value) const override;
-	virtual bool HasKeys(const int64& Value) const override;
-	virtual void SetDefault(const int64& Value) override;
-	virtual void ClearDefaults() override;
-
-public:
-
-	//~ UMovieSceneSection interface 
-
-	virtual void MoveSection(float DeltaPosition, TSet<FKeyHandle>& KeyHandles) override;
-	virtual void DilateSection(float DilationFactor, float Origin, TSet<FKeyHandle>& KeyHandles) override;
-	virtual void GetKeyHandles(TSet<FKeyHandle>& OutKeyHandles, TRange<float> TimeRange) const override;
-	virtual TOptional<float> GetKeyTime(FKeyHandle KeyHandle) const override;
-	virtual void SetKeyTime(FKeyHandle KeyHandle, float Time) override;
-
-private:
-
 	/** Ordered curve data */
-	// @todo Sequencer This could be optimized by packing the enums separately
-	// but that may not be worth the effort
 	UPROPERTY()
-	FIntegralCurve EnumCurve;
+	FMovieSceneByteChannel EnumCurve;
 };
diff --git a/Engine/Source/Runtime/MovieSceneTracks/Public/Sections/MovieSceneEventSection.h b/Engine/Source/Runtime/MovieSceneTracks/Public/Sections/MovieSceneEventSection.h
--- a/Engine/Source/Runtime/MovieSceneTracks/Public/Sections/MovieSceneEventSection.h
+++ b/Engine/Source/Runtime/MovieSceneTracks/Public/Sections/MovieSceneEventSection.h
@@ -8,10 +8,10 @@
 #include "Curves/KeyHandle.h"
 #include "MovieSceneSection.h"
 #include "Curves/NameCurve.h"
-#include "Curves/CurveInterface.h"
 #include "UObject/StructOnScope.h"
 #include "Engine/Engine.h"
 #include "UObject/SoftObjectPath.h"
+#include "Channels/MovieSceneChannel.h"
 #include "MovieSceneEventSection.generated.h"
 
 struct EventData;
@@ -24,30 +24,28 @@ struct MOVIESCENETRACKS_API FMovieSceneEventParameters
 	FMovieSceneEventParameters() {}
 
 	/** Construction from a struct type */
-	FMovieSceneEventParameters(UStruct& InStruct)
+	FMovieSceneEventParameters(UScriptStruct& InStruct)
 		: StructType(&InStruct)
 	{
 	}
 
-	FMovieSceneEventParameters(const FMovieSceneEventParameters& RHS) = default;
-	FMovieSceneEventParameters& operator=(const FMovieSceneEventParameters& RHS) = default;
+	MOVIESCENETRACKS_API friend bool operator==(const FMovieSceneEventParameters& A, const FMovieSceneEventParameters& B);
 
-	FMovieSceneEventParameters(FMovieSceneEventParameters&&) = default;
-	FMovieSceneEventParameters& operator=(FMovieSceneEventParameters&&) = default;
+	MOVIESCENETRACKS_API friend bool operator!=(const FMovieSceneEventParameters& A, const FMovieSceneEventParameters& B);
 
 	/**
 	 * Access the struct type of this event parameter payload
-	 * @return A valid UStruct* or nullptr if the struct is not set, or no longer available
+	 * @return A valid UScriptStruct* or nullptr if the struct is not set, or no longer available
 	 */
-	UStruct* GetStructType() const
+	UScriptStruct* GetStructType() const
 	{
-		return Cast<UStruct>(StructType.TryLoad());
+		return Cast<UScriptStruct>(StructType.TryLoad());
 	}
 
 	/**
 	 * Change the type of this event parameter payload to be the specified struct
 	 */
-	void Reassign(UStruct* NewStruct)
+	void Reassign(UScriptStruct* NewStruct)
 	{
 		StructType = NewStruct;
 
@@ -109,6 +107,15 @@ struct FEventPayload
 	FEventPayload() {}
 	FEventPayload(FName InEventName) : EventName(InEventName) {}
 
+	friend bool operator==(const FEventPayload& A, const FEventPayload& B)
+	{
+		return A.EventName == B.EventName && A.Parameters == B.Parameters;
+	}
+
+	friend bool operator!=(const FEventPayload& A, const FEventPayload& B)
+	{
+		return A.EventName != B.EventName || A.Parameters != B.Parameters;
+	}
 	/** The name of the event to trigger */
 	UPROPERTY(EditAnywhere, Category=Event)
 	FName EventName;
@@ -123,39 +130,77 @@ USTRUCT()
 struct FMovieSceneEventSectionData
 {
 	GENERATED_BODY()
-	
-	FMovieSceneEventSectionData() = default;
 
-	FMovieSceneEventSectionData(const FMovieSceneEventSectionData& RHS)
-		: KeyTimes(RHS.KeyTimes)
-		, KeyValues(RHS.KeyValues)
+	/**
+	 * Access an integer that uniquely identifies this channel type.
+	 *
+	 * @return A static identifier that was allocated using FMovieSceneChannelEntry::RegisterNewID
+	 */
+	MOVIESCENETRACKS_API static uint32 GetChannelID();
+
+	/**
+	 * Called after this section data has been serialized to upgrade old data
+	 */
+	MOVIESCENETRACKS_API void PostSerialize(const FArchive& Ar);
+
+	/**
+	 * Access a mutable interface for this channel
+	 *
+	 * @return An object that is able to manipulate this channel
+	 */
+	FORCEINLINE TMovieSceneChannel<FEventPayload> GetInterface()
 	{
+		return TMovieSceneChannel<FEventPayload>(&Times, &KeyValues, &KeyHandles);
 	}
 
-	FMovieSceneEventSectionData& operator=(const FMovieSceneEventSectionData& RHS)
+	/**
+	 * Access a constant interface for this channel
+	 *
+	 * @return An object that is able to interrogate this channel
+	 */
+	FORCEINLINE TMovieSceneChannel<const FEventPayload> GetInterface() const
 	{
-		KeyTimes = RHS.KeyTimes;
-		KeyValues = RHS.KeyValues;
-#if WITH_EDITORONLY_DATA
-		KeyHandles.Reset();
-#endif
-		return *this;
+		return TMovieSceneChannel<const FEventPayload>(&Times, &KeyValues);
+	}
+
+	TArrayView<const FFrameNumber> GetKeyTimes() const
+	{
+		return Times;
+	}
+
+	TArrayView<const FEventPayload> GetKeyValues() const
+	{
+		return KeyValues;
 	}
 
-	/** Sorted array of key times */
+private:
+
 	UPROPERTY()
-	TArray<float> KeyTimes;
+	TArray<FFrameNumber> Times;
 
 	/** Array of values that correspond to each key time */
 	UPROPERTY()
 	TArray<FEventPayload> KeyValues;
 
+	FMovieSceneKeyHandleMap KeyHandles;
+
+
 #if WITH_EDITORONLY_DATA
-	/** Transient key handles */
-	FKeyHandleLookupTable KeyHandles;
+
+	UPROPERTY()
+	TArray<float> KeyTimes_DEPRECATED;
+
 #endif
 };
 
+
+template<>
+struct TStructOpsTypeTraits<FMovieSceneEventSectionData> : public TStructOpsTypeTraitsBase2<FMovieSceneEventSectionData>
+{
+	enum { WithPostSerialize = true };
+};
+
+
 /**
  * Implements a section in movie scene event tracks.
  */
@@ -179,18 +224,8 @@ public:
 	 * @return Event data.
 	 */
 	const FMovieSceneEventSectionData& GetEventData() const { return EventData; }
-	
-	TCurveInterface<FEventPayload, float> GetCurveInterface() { return CurveInterface.GetValue(); }
 
-public:
-
-	//~ UMovieSceneSection interface
-
-	virtual void DilateSection(float DilationFactor, float Origin, TSet<FKeyHandle>& KeyHandles) override;
-	virtual void GetKeyHandles(TSet<FKeyHandle>& KeyHandles, TRange<float> TimeRange) const override;
-	virtual void MoveSection(float DeltaPosition, TSet<FKeyHandle>& KeyHandles) override;
-	virtual TOptional<float> GetKeyTime(FKeyHandle KeyHandle) const override;
-	virtual void SetKeyTime(FKeyHandle KeyHandle, float Time) override;
+protected:
 
 private:
 
@@ -199,6 +234,28 @@ private:
 
 	UPROPERTY()
 	FMovieSceneEventSectionData EventData;
-
-	TOptional<TCurveInterface<FEventPayload, float>> CurveInterface;
 };
+
+
+
+/** Stub out unnecessary functions */
+inline bool EvaluateChannel(const FMovieSceneEventSectionData* InChannel, FFrameTime InTime, FEventPayload& OutValue)
+{
+	// Can't evaluate event section data in the typical sense
+	return false;
+}
+
+inline bool ValueExistsAtTime(const FMovieSceneEventSectionData* InChannel, FFrameNumber Time, const FEventPayload& Value, int32 Tolerance = 0)
+{
+	// true if any value exists
+	return InChannel->GetInterface().FindKey(Time, Tolerance) != INDEX_NONE;
+}
+
+inline void Optimize(FMovieSceneEventSectionData* InChannel, const FKeyDataOptimizationParams& Params)
+{}
+
+inline void SetChannelDefault(FMovieSceneEventSectionData* Channel, const FEventPayload& DefaultValue)
+{}
+
+inline void ClearChannelDefault(FMovieSceneEventSectionData* InChannel)
+{}
\ No newline at end of file
diff --git a/Engine/Source/Runtime/MovieSceneTracks/Public/Sections/MovieSceneFloatSection.h b/Engine/Source/Runtime/MovieSceneTracks/Public/Sections/MovieSceneFloatSection.h
--- a/Engine/Source/Runtime/MovieSceneTracks/Public/Sections/MovieSceneFloatSection.h
+++ b/Engine/Source/Runtime/MovieSceneTracks/Public/Sections/MovieSceneFloatSection.h
@@ -7,56 +7,29 @@
 #include "Curves/KeyHandle.h"
 #include "Curves/RichCurve.h"
 #include "MovieSceneSection.h"
-#include "Sections/IKeyframeSection.h"
+#include "Channels/MovieSceneFloatChannel.h"
 #include "MovieSceneFloatSection.generated.h"
 
+
 /**
  * A single floating point section
  */
 UCLASS( MinimalAPI )
 class UMovieSceneFloatSection
 	: public UMovieSceneSection
-	, public IKeyframeSection<float>
 {
 	GENERATED_UCLASS_BODY()
 
 public:
-	/**
-	 * Updates this section
-	 *
-	 * @param Position	The position in time within the movie scene
-	 */
-	virtual float Eval( float Position, float DefaultValue ) const;
 
 	/**
-	 * @return The float curve on this section
+	 * Public access to this section's internal data function
 	 */
-	FRichCurve& GetFloatCurve() { return FloatCurve; }
-	const FRichCurve& GetFloatCurve() const { return FloatCurve; }
-
-public:
-
-	//~ IKeyframeSection interface
-
-	virtual void AddKey( float Time, const float& Value, EMovieSceneKeyInterpolation KeyInterpolation ) override;
-	virtual bool NewKeyIsNewData(float Time, const float& Value) const override;
-	virtual bool HasKeys( const float& Value ) const override;
-	virtual void SetDefault( const float& Value ) override;
-	virtual void ClearDefaults() override;
-
-public:
-
-	//~ UMovieSceneSection interface
+	const FMovieSceneFloatChannel& GetChannel() const { return FloatCurve; }
 
-	virtual void MoveSection(float DeltaPosition, TSet<FKeyHandle>& KeyHandles) override;
-	virtual void DilateSection(float DilationFactor, float Origin, TSet<FKeyHandle>& KeyHandles) override;
-	virtual void GetKeyHandles(TSet<FKeyHandle>& OutKeyHandles, TRange<float> TimeRange) const override;
-	virtual TOptional<float> GetKeyTime( FKeyHandle KeyHandle ) const override;
-	virtual void SetKeyTime( FKeyHandle KeyHandle, float Time ) override;
+protected:
 
-private:
-	
-	/** Curve data */
+	/** Float data */
 	UPROPERTY()
-	FRichCurve FloatCurve;
+	FMovieSceneFloatChannel FloatCurve;
 };
diff --git a/Engine/Source/Runtime/MovieSceneTracks/Public/Sections/MovieSceneIntegerSection.h b/Engine/Source/Runtime/MovieSceneTracks/Public/Sections/MovieSceneIntegerSection.h
--- a/Engine/Source/Runtime/MovieSceneTracks/Public/Sections/MovieSceneIntegerSection.h
+++ b/Engine/Source/Runtime/MovieSceneTracks/Public/Sections/MovieSceneIntegerSection.h
@@ -7,7 +7,7 @@
 #include "Curves/KeyHandle.h"
 #include "Curves/IntegralCurve.h"
 #include "MovieSceneSection.h"
-#include "Sections/IKeyframeSection.h"
+#include "Channels/MovieSceneIntegerChannel.h"
 #include "MovieSceneIntegerSection.generated.h"
 
 
@@ -17,38 +17,18 @@
 UCLASS(MinimalAPI)
 class UMovieSceneIntegerSection 
 	: public UMovieSceneSection
-	, public IKeyframeSection<int32>
 {
 	GENERATED_UCLASS_BODY()
-public:
-
-	/** Gets all the keys of this integer section */
-	FIntegralCurve& GetCurve() { return IntegerCurve; }
-	const FIntegralCurve& GetCurve() const { return IntegerCurve; }
-
-public:
-
-	//~ IKeyframeSection interface
-
-	virtual void AddKey(float Time, const int32& Value, EMovieSceneKeyInterpolation KeyInterpolation) override;
-	virtual bool NewKeyIsNewData(float Time, const int32& Value) const override;
-	virtual bool HasKeys(const int32& Value) const override;
-	virtual void SetDefault(const int32& Value) override;
-	virtual void ClearDefaults() override;
 
 public:
-
-	//~ UMovieSceneSection interface 
-
-	virtual void MoveSection(float DeltaPosition, TSet<FKeyHandle>& KeyHandles) override;
-	virtual void DilateSection(float DilationFactor, float Origin, TSet<FKeyHandle>& KeyHandles) override;
-	virtual void GetKeyHandles(TSet<FKeyHandle>& OutKeyHandles, TRange<float> TimeRange) const override;
-	virtual TOptional<float> GetKeyTime(FKeyHandle KeyHandle) const override;
-	virtual void SetKeyTime(FKeyHandle KeyHandle, float Time) override;
+	/**
+	 * Access this section's underlying raw data
+	 */
+	const FMovieSceneIntegerChannel& GetChannel() const { return IntegerCurve; }
 
 private:
 
 	/** Ordered curve data */
 	UPROPERTY()
-	FIntegralCurve IntegerCurve;
+	FMovieSceneIntegerChannel IntegerCurve;
 };
diff --git a/Engine/Source/Runtime/MovieSceneTracks/Public/Sections/MovieSceneLevelVisibilitySection.h b/Engine/Source/Runtime/MovieSceneTracks/Public/Sections/MovieSceneLevelVisibilitySection.h
--- a/Engine/Source/Runtime/MovieSceneTracks/Public/Sections/MovieSceneLevelVisibilitySection.h
+++ b/Engine/Source/Runtime/MovieSceneTracks/Public/Sections/MovieSceneLevelVisibilitySection.h
@@ -44,8 +44,6 @@ public:
 
 	//~ UMovieSceneSection interface
 	virtual FMovieSceneEvalTemplatePtr GenerateTemplate() const override;
-	virtual TOptional<float> GetKeyTime(FKeyHandle KeyHandle) const override;
-	virtual void SetKeyTime(FKeyHandle KeyHandle, float Time) override;
 
 private:
 
diff --git a/Engine/Source/Runtime/MovieSceneTracks/Public/Sections/MovieSceneParameterSection.h b/Engine/Source/Runtime/MovieSceneTracks/Public/Sections/MovieSceneParameterSection.h
--- a/Engine/Source/Runtime/MovieSceneTracks/Public/Sections/MovieSceneParameterSection.h
+++ b/Engine/Source/Runtime/MovieSceneTracks/Public/Sections/MovieSceneParameterSection.h
@@ -7,8 +7,11 @@
 #include "Curves/KeyHandle.h"
 #include "Curves/RichCurve.h"
 #include "MovieSceneSection.h"
+#include "Channels/MovieSceneFloatChannel.h"
 #include "MovieSceneParameterSection.generated.h"
 
+struct FMovieSceneParameterPropertyInterface;
+struct FMovieSceneConstParameterPropertyInterface;
 
 /**
  * Structure representing the animated value of a scalar parameter.
@@ -79,9 +82,7 @@ struct FScalarParameterNameAndCurve
 	GENERATED_USTRUCT_BODY()
 
 	FScalarParameterNameAndCurve()
-	{
-		Index = 0;
-	}
+	{}
 
 	/** Creates a new FScalarParameterNameAndCurve for a specific scalar parameter. */
 	FScalarParameterNameAndCurve(FName InParameterName);
@@ -90,12 +91,9 @@ struct FScalarParameterNameAndCurve
 	UPROPERTY()
 	FName ParameterName;
 
-	UPROPERTY()
-	int32 Index;
-
 	/** The curve which contains the animation data for the scalar parameter. */
 	UPROPERTY()
-	FRichCurve ParameterCurve;
+	FMovieSceneFloatChannel ParameterCurve;
 };
 
 
@@ -108,9 +106,7 @@ struct FVectorParameterNameAndCurves
 	GENERATED_USTRUCT_BODY()
 
 	FVectorParameterNameAndCurves() 
-	{
-		Index = 0;
-	}
+	{}
 
 	/** Creates a new FVectorParameterNameAndCurve for a specific vector parameter. */
 	FVectorParameterNameAndCurves(FName InParameterName);
@@ -119,20 +115,17 @@ struct FVectorParameterNameAndCurves
 	UPROPERTY()
 	FName ParameterName;
 
-	UPROPERTY()
-	int32 Index;
-
 	/** The curve which contains the animation data for the x component of the vector parameter. */
 	UPROPERTY()
-	FRichCurve XCurve;
+	FMovieSceneFloatChannel XCurve;
 
 	/** The curve which contains the animation data for the y component of the vector parameter. */
 	UPROPERTY()
-	FRichCurve YCurve;
+	FMovieSceneFloatChannel YCurve;
 
 	/** The curve which contains the animation data for the z component of the vector parameter. */
 	UPROPERTY()
-	FRichCurve ZCurve;
+	FMovieSceneFloatChannel ZCurve;
 };
 
 
@@ -145,9 +138,7 @@ struct FColorParameterNameAndCurves
 	GENERATED_USTRUCT_BODY()
 
 	FColorParameterNameAndCurves()
-	{
-		Index = 0;
-	}
+	{}
 
 	/** Creates a new FVectorParameterNameAndCurve for a specific color parameter. */
 	FColorParameterNameAndCurves(FName InParameterName);
@@ -156,24 +147,21 @@ struct FColorParameterNameAndCurves
 	UPROPERTY()
 	FName ParameterName;
 
-	UPROPERTY()
-	int32 Index;
-
 	/** The curve which contains the animation data for the red component of the color parameter. */
 	UPROPERTY()
-	FRichCurve RedCurve;
+	FMovieSceneFloatChannel RedCurve;
 
 	/** The curve which contains the animation data for the green component of the color parameter. */
 	UPROPERTY()
-	FRichCurve GreenCurve;
+	FMovieSceneFloatChannel GreenCurve;
 
 	/** The curve which contains the animation data for the blue component of the color parameter. */
 	UPROPERTY()
-	FRichCurve BlueCurve;
+	FMovieSceneFloatChannel BlueCurve;
 
 	/** The curve which contains the animation data for the alpha component of the color parameter. */
 	UPROPERTY()
-	FRichCurve AlphaCurve;
+	FMovieSceneFloatChannel AlphaCurve;
 };
 
 
@@ -188,13 +176,13 @@ class UMovieSceneParameterSection
 
 public:
 	/** Adds a a key for a specific scalar parameter at the specified time with the specified value. */
-	void AddScalarParameterKey(FName InParameterName, float InTime, float InValue);
+	void AddScalarParameterKey(FName InParameterName, FFrameNumber InTime, float InValue);
 
 	/** Adds a a key for a specific vector parameter at the specified time with the specified value. */
-	void AddVectorParameterKey(FName InParameterName, float InTime, FVector InValue);
+	void AddVectorParameterKey(FName InParameterName, FFrameNumber InTime, FVector InValue);
 
 	/** Adds a a key for a specific color parameter at the specified time with the specified value. */
-	void AddColorParameterKey(FName InParameterName, float InTime, FLinearColor InValue);
+	void AddColorParameterKey(FName InParameterName, FFrameNumber InTime, FLinearColor InValue);
 
 	/** 
 	 * Removes a scalar parameter from this section. 
@@ -221,35 +209,26 @@ public:
 	MOVIESCENETRACKS_API bool RemoveColorParameter(FName InParameterName);
 
 	/** Gets the animated scalar parameters and their associated curves. */
-	MOVIESCENETRACKS_API TArray<FScalarParameterNameAndCurve>* GetScalarParameterNamesAndCurves();
+	MOVIESCENETRACKS_API TArray<FScalarParameterNameAndCurve>& GetScalarParameterNamesAndCurves();
 	MOVIESCENETRACKS_API const TArray<FScalarParameterNameAndCurve>& GetScalarParameterNamesAndCurves() const;
 
 	/** Gets the animated vector parameters and their associated curves. */
-	MOVIESCENETRACKS_API TArray<FVectorParameterNameAndCurves>* GetVectorParameterNamesAndCurves();
+	MOVIESCENETRACKS_API TArray<FVectorParameterNameAndCurves>& GetVectorParameterNamesAndCurves();
 	MOVIESCENETRACKS_API const TArray<FVectorParameterNameAndCurves>& GetVectorParameterNamesAndCurves() const;
 
 	/** Gets the animated color parameters and their associated curves. */
-	MOVIESCENETRACKS_API TArray<FColorParameterNameAndCurves>* GetColorParameterNamesAndCurves();
+	MOVIESCENETRACKS_API TArray<FColorParameterNameAndCurves>& GetColorParameterNamesAndCurves();
 	MOVIESCENETRACKS_API const TArray<FColorParameterNameAndCurves>& GetColorParameterNamesAndCurves() const;
 
 	/** Gets the set of all parameter names used by this section. */
 	MOVIESCENETRACKS_API void GetParameterNames(TSet<FName>& ParameterNames) const;
 
-public:
+protected:
 
 	//~ UMovieSceneSection interface
+	virtual void Serialize(FArchive& Ar) override;
 
-	virtual void DilateSection(float DilationFactor, float Origin, TSet<FKeyHandle>& KeyHandles) override;
-	virtual void GetKeyHandles(TSet<FKeyHandle>& OutKeyHandles, TRange<float> TimeRange) const override;
-	virtual void MoveSection(float DeltaPosition, TSet<FKeyHandle>& KeyHandles) override;
-	virtual TOptional<float> GetKeyTime(FKeyHandle KeyHandle) const override;
-	virtual void SetKeyTime(FKeyHandle KeyHandle, float Time) override;
-
-private:
-
-	void UpdateParameterIndicesFromRemoval(int32 RemovedIndex);
-	void GatherCurves(TArray<const FRichCurve*> &OutCurves) const;
-	void GatherCurves(TArray<FRichCurve*> &OutCurves);
+	void ReconstructChannelProxy();
 
 private:
 
@@ -264,4 +243,4 @@ private:
 	/** The vector parameter names and their associated curves. */
 	UPROPERTY()
 	TArray<FColorParameterNameAndCurves> ColorParameterNamesAndCurves;
-};
+};
\ No newline at end of file
diff --git a/Engine/Source/Runtime/MovieSceneTracks/Public/Sections/MovieSceneParticleSection.h b/Engine/Source/Runtime/MovieSceneTracks/Public/Sections/MovieSceneParticleSection.h
--- a/Engine/Source/Runtime/MovieSceneTracks/Public/Sections/MovieSceneParticleSection.h
+++ b/Engine/Source/Runtime/MovieSceneTracks/Public/Sections/MovieSceneParticleSection.h
@@ -6,24 +6,52 @@
 #include "UObject/ObjectMacros.h"
 #include "Curves/KeyHandle.h"
 #include "MovieSceneSection.h"
-#include "Curves/IntegralCurve.h"
+#include "Channels/MovieSceneByteChannel.h"
 #include "MovieSceneParticleSection.generated.h"
 
-
 /**
  * Defines the types of particle keys.
  */
 UENUM()
-namespace EParticleKey
+enum class EParticleKey : uint8
 {
-	enum Type
-	{
-		Activate = 0,
-		Deactivate = 1,
-		Trigger = 2,
-	};
-}
+	Activate   = 0,
+	Deactivate = 1,
+	Trigger    = 2,
+};
+
+USTRUCT()
+struct FMovieSceneParticleChannel : public FMovieSceneByteChannel
+{
+	GENERATED_BODY()
+
+	MOVIESCENETRACKS_API FMovieSceneParticleChannel();
+
+	/**
+	 * Access an integer that uniquely identifies this channel type.
+	 *
+	 * @return A static identifier that was allocated using FMovieSceneChannelEntry::RegisterNewID
+	 */
+	MOVIESCENETRACKS_API static uint32 GetChannelID();
+};
 
+template<>
+struct TStructOpsTypeTraits<FMovieSceneParticleChannel> : public TStructOpsTypeTraitsBase2<FMovieSceneParticleChannel>
+{
+	enum { WithSerializeFromMismatchedTag = true };
+};
+
+
+template<>
+struct TMovieSceneChannelTraits<FMovieSceneParticleChannel>
+{
+#if WITH_EDITORONLY_DATA
+
+	/** Byte channels can have external values (ie, they can get their values from external objects for UI purposes) */
+	typedef TMovieSceneExternalValue<uint8> EditorDataType;
+
+#endif
+};
 
 /**
  * Particle section, for particle toggling and triggering.
@@ -34,25 +62,38 @@ class UMovieSceneParticleSection
 {
 	GENERATED_UCLASS_BODY()
 
-	MOVIESCENETRACKS_API void AddKey(float Time, EParticleKey::Type KeyType);
+public:
 
-	FIntegralCurve& GetParticleCurve() { return ParticleKeys; }
-	const FIntegralCurve& GetParticleCurve() const { return ParticleKeys; }
+	/** Curve containing the particle keys. */
+	UPROPERTY()
+	FMovieSceneParticleChannel ParticleKeys;
 
-public:
+protected:
 
 	//~ UMovieSceneSection interface
-
-	virtual void MoveSection(float DeltaPosition, TSet<FKeyHandle>& KeyHandles) override;
-	virtual void DilateSection(float DilationFactor, float Origin, TSet<FKeyHandle>& KeyHandles) override;
-	virtual void GetKeyHandles(TSet<FKeyHandle>& OutKeyHandles, TRange<float> TimeRange) const override;
-	virtual TOptional<float> GetKeyTime(FKeyHandle KeyHandle) const override;
-	virtual void SetKeyTime(FKeyHandle KeyHandle, float Time) override;
 	virtual FMovieSceneEvalTemplatePtr GenerateTemplate() const override;
-	
-private:
-
-	/** Curve containing the particle keys. */
-	UPROPERTY()
-	FIntegralCurve ParticleKeys;
 };
+
+
+inline void AssignValue(FMovieSceneParticleChannel* InChannel, int32 ValueIndex, EParticleKey InValue)
+{
+	InChannel->GetInterface().GetValues()[ValueIndex] = (uint8)InValue;
+}
+
+inline bool EvaluateChannel(const FMovieSceneParticleChannel* InChannel, FFrameTime InTime, EParticleKey& OutValue)
+{
+	uint8 RawValue = 0;
+	if (InChannel->Evaluate(InTime, RawValue))
+	{
+		OutValue = (EParticleKey)RawValue;
+		return true;
+	}
+	return false;
+}
+
+// Stubbed out functions since FMovieSceneParticleChannel doesn't have an explicit default
+inline void SetChannelDefault(FMovieSceneParticleChannel* InChannel, EParticleKey DefaultValue)
+{}
+
+inline void ClearChannelDefault(FMovieSceneParticleChannel* InChannel)
+{}
diff --git a/Engine/Source/Runtime/MovieSceneTracks/Public/Sections/MovieSceneSkeletalAnimationSection.h b/Engine/Source/Runtime/MovieSceneTracks/Public/Sections/MovieSceneSkeletalAnimationSection.h
--- a/Engine/Source/Runtime/MovieSceneTracks/Public/Sections/MovieSceneSkeletalAnimationSection.h
+++ b/Engine/Source/Runtime/MovieSceneTracks/Public/Sections/MovieSceneSkeletalAnimationSection.h
@@ -4,9 +4,9 @@
 
 #include "CoreMinimal.h"
 #include "UObject/ObjectMacros.h"
-#include "Curves/KeyHandle.h"
 #include "MovieSceneSection.h"
 #include "Animation/AnimSequenceBase.h"
+#include "Channels/MovieSceneFloatChannel.h"
 #include "MovieSceneSkeletalAnimationSection.generated.h"
 
 USTRUCT()
@@ -48,7 +48,7 @@ struct FMovieSceneSkeletalAnimationParams
 
 	/** The weight curve for this animation section */
 	UPROPERTY( EditAnywhere, Category = "Animation" )
-	FRichCurve Weight;	
+	FMovieSceneFloatChannel Weight;
 };
 
 /**
@@ -65,23 +65,18 @@ public:
 	UPROPERTY(EditAnywhere, Category = "Animation", meta=(ShowOnlyInnerProperties))
 	FMovieSceneSkeletalAnimationParams Params;
 
-public:
-
-	//~ MovieSceneSection interface
+protected:
 
-	virtual void MoveSection( float DeltaPosition, TSet<FKeyHandle>& KeyHandles ) override;
-	virtual void DilateSection( float DilationFactor, float Origin, TSet<FKeyHandle>& KeyHandles  ) override;
-	virtual UMovieSceneSection* SplitSection(float SplitTime) override;
-	virtual void GetKeyHandles(TSet<FKeyHandle>& OutKeyHandles, TRange<float> TimeRange) const override;
-	virtual void GetSnapTimes(TArray<float>& OutSnapTimes, bool bGetSectionBorders) const override;
-	virtual TOptional<float> GetOffsetTime() const override { return TOptional<float>(Params.StartOffset); }
-	virtual TOptional<float> GetKeyTime( FKeyHandle KeyHandle ) const override { return TOptional<float>(); }
-	virtual void SetKeyTime( FKeyHandle KeyHandle, float Time ) override { }
+	//~ UMovieSceneSection interface
+	virtual UMovieSceneSection* SplitSection(FFrameNumber SplitTime) override;
+	virtual void GetSnapTimes(TArray<FFrameNumber>& OutSnapTimes, bool bGetSectionBorders) const override;
+	virtual TOptional<FFrameTime> GetOffsetTime() const override;
 	virtual FMovieSceneEvalTemplatePtr GenerateTemplate() const override;
 
 	/** ~UObject interface */
 	virtual void PostLoad() override;
 	virtual void Serialize(FArchive& Ar) override;
+
 private:
 
 	//~ UObject interface
diff --git a/Engine/Source/Runtime/MovieSceneTracks/Public/Sections/MovieSceneStringChannel.h b/Engine/Source/Runtime/MovieSceneTracks/Public/Sections/MovieSceneStringChannel.h
new file mode 100644
--- /dev/null
+++ b/Engine/Source/Runtime/MovieSceneTracks/Public/Sections/MovieSceneStringChannel.h
@@ -0,0 +1,147 @@
+// Copyright 1998-2018 Epic Games, Inc. All Rights Reserved.
+
+#pragma once
+
+#include "Containers/UnrealString.h"
+#include "FrameNumber.h"
+#include "UObject/ObjectMacros.h"
+#include "Channels/MovieSceneChannelTraits.h"
+#include "Channels/MovieSceneChannel.h"
+#include "MovieSceneStringChannel.generated.h"
+
+USTRUCT()
+struct FMovieSceneStringChannel
+{
+	GENERATED_BODY()
+
+	FMovieSceneStringChannel()
+		: DefaultValue(), bHasDefaultValue(false)
+	{}
+
+	/**
+	 * Access an integer that uniquely identifies this channel type.
+	 *
+	 * @return A static identifier that was allocated using FMovieSceneChannelEntry::RegisterNewID
+	 */
+	MOVIESCENETRACKS_API static uint32 GetChannelID();
+
+	/**
+	 * Serialize this type from another
+	 */
+	MOVIESCENETRACKS_API bool SerializeFromMismatchedTag(const FPropertyTag& Tag, FArchive& Ar);
+
+	/**
+	 * Access a mutable interface for this channel
+	 *
+	 * @return An object that is able to manipulate this channel
+	 */
+	FORCEINLINE TMovieSceneChannel<FString> GetInterface()
+	{
+		return TMovieSceneChannel<FString>(&Times, &Values, &KeyHandles);
+	}
+
+	/**
+	 * Access a constant interface for this channel
+	 *
+	 * @return An object that is able to interrogate this channel
+	 */
+	FORCEINLINE TMovieSceneChannel<const FString> GetInterface() const
+	{
+		return TMovieSceneChannel<const FString>(&Times, &Values);
+	}
+
+	/**
+	 * Evaluate this channel
+	 *
+	 * @param InTime     The time to evaluate at
+	 * @return A pointer to the string, or nullptr
+	 */
+	MOVIESCENETRACKS_API const FString* Evaluate(FFrameTime InTime) const;
+
+public:
+
+	/**
+	 * Set this channel's default value that should be used when no keys are present
+	 *
+	 * @param InDefaultValue The desired default value
+	 */
+	FORCEINLINE void SetDefault(FString InDefaultValue)
+	{
+		bHasDefaultValue = true;
+		DefaultValue = InDefaultValue;
+	}
+
+	/**
+	 * Get this channel's default value that will be used when no keys are present
+	 *
+	 * @return (Optional) The channel's default value
+	 */
+	FORCEINLINE TOptional<FString> GetDefault() const
+	{
+		return bHasDefaultValue ? TOptional<FString>(DefaultValue) : TOptional<FString>();
+	}
+
+	/**
+	 * Remove this channel's default value causing the channel to have no effect where no keys are present
+	 */
+	FORCEINLINE void RemoveDefault()
+	{
+		bHasDefaultValue = false;
+	}
+
+private:
+
+	UPROPERTY()
+	TArray<FFrameNumber> Times;
+
+	/** Array of values that correspond to each key time */
+	UPROPERTY()
+	TArray<FString> Values;
+
+	/** Default value used when there are no keys */
+	UPROPERTY()
+	FString DefaultValue;
+
+	/** */
+	UPROPERTY()
+	bool bHasDefaultValue;
+
+	FMovieSceneKeyHandleMap KeyHandles;
+};
+
+template<>
+struct TStructOpsTypeTraits<FMovieSceneStringChannel> : public TStructOpsTypeTraitsBase2<FMovieSceneStringChannel>
+{
+	enum { WithSerializeFromMismatchedTag = true };
+};
+
+
+template<>
+struct TMovieSceneChannelTraits<FMovieSceneStringChannel>
+{
+#if WITH_EDITORONLY_DATA
+
+	/** String channels can have external values (ie, they can get their values from external objects for UI purposes) */
+	typedef TMovieSceneExternalValue<FString> EditorDataType;
+
+#endif
+};
+
+inline bool EvaluateChannel(const FMovieSceneStringChannel* InChannel, FFrameTime InTime, FString& OutValue)
+{
+	if (const FString* Result = InChannel->Evaluate(InTime))
+	{
+		OutValue = *Result;
+		return true;
+	}
+	return false;
+}
+
+inline bool ValueExistsAtTime(const FMovieSceneStringChannel* Channel, FFrameNumber InFrameNumber, const FString& Value)
+{
+	const FFrameTime FrameTime(InFrameNumber);
+
+	const FString* ExistingValue = Channel->Evaluate(FrameTime);
+	return ExistingValue && Value == *ExistingValue;
+}
+
diff --git a/Engine/Source/Runtime/MovieSceneTracks/Public/Sections/MovieSceneStringSection.h b/Engine/Source/Runtime/MovieSceneTracks/Public/Sections/MovieSceneStringSection.h
--- a/Engine/Source/Runtime/MovieSceneTracks/Public/Sections/MovieSceneStringSection.h
+++ b/Engine/Source/Runtime/MovieSceneTracks/Public/Sections/MovieSceneStringSection.h
@@ -4,10 +4,8 @@
 
 #include "CoreMinimal.h"
 #include "UObject/ObjectMacros.h"
-#include "Curves/KeyHandle.h"
 #include "MovieSceneSection.h"
-#include "Sections/IKeyframeSection.h"
-#include "Curves/StringCurve.h"
+#include "MovieSceneStringChannel.h"
 #include "MovieSceneStringSection.generated.h"
 
 /**
@@ -16,57 +14,19 @@
 UCLASS(MinimalAPI)
 class UMovieSceneStringSection
 	: public UMovieSceneSection
-	, public IKeyframeSection<FString>
 {
 	GENERATED_UCLASS_BODY()
 
 public:
 
 	/**
-	 * Evaluates this section at the specified time.
-	 *
-	 * @param Time The time at which to evaluate the section.
-	 * @param DefaultString The default value to return.
+	 * Public access to this section's internal data function
 	 */
-	virtual FString Eval(float Time, const FString& DefaultString) const;
-
-	/**
-	 * Get the section's string curve.
-	 *
-	 * @return String curve.
-	 */
-	FStringCurve& GetStringCurve()
-	{
-		return StringCurve;
-	}
-	const FStringCurve& GetStringCurve() const
-	{
-		return StringCurve;
-	}
-
-public:
-
-	//~ IKeyframeSection interface.
-
-	virtual bool NewKeyIsNewData( float Time, const FString& KeyData ) const;
-	virtual bool HasKeys( const FString& KeyData ) const;
-	virtual void AddKey( float Time, const FString& KeyData, EMovieSceneKeyInterpolation KeyInterpolation );
-	virtual void SetDefault( const FString& KeyData );
-	virtual void ClearDefaults() override;
-
-public:
-
-	//~ UMovieSceneSection interface 
-
-	virtual void DilateSection(float DilationFactor, float Origin, TSet<FKeyHandle>& KeyHandles) override;
-	virtual void GetKeyHandles(TSet<FKeyHandle>& OutKeyHandles, TRange<float> TimeRange) const override;
-	virtual TOptional<float> GetKeyTime(FKeyHandle KeyHandle) const override;
-	virtual void MoveSection(float DeltaPosition, TSet<FKeyHandle>& KeyHandles) override;
-	virtual void SetKeyTime(FKeyHandle KeyHandle, float Time) override;
+	const FMovieSceneStringChannel& GetChannel() const { return StringCurve; }
 
 private:
 
 	/** Curve data */
 	UPROPERTY()
-	FStringCurve StringCurve;
+	FMovieSceneStringChannel StringCurve;
 };
diff --git a/Engine/Source/Runtime/MovieSceneTracks/Public/Sections/MovieSceneVectorSection.h b/Engine/Source/Runtime/MovieSceneTracks/Public/Sections/MovieSceneVectorSection.h
--- a/Engine/Source/Runtime/MovieSceneTracks/Public/Sections/MovieSceneVectorSection.h
+++ b/Engine/Source/Runtime/MovieSceneTracks/Public/Sections/MovieSceneVectorSection.h
@@ -7,34 +7,13 @@
 #include "Curves/KeyHandle.h"
 #include "Curves/RichCurve.h"
 #include "MovieSceneSection.h"
-#include "Sections/IKeyframeSection.h"
 #include "MovieSceneKeyStruct.h"
+#include "Channels/MovieSceneFloatChannel.h"
 #include "MovieSceneVectorSection.generated.h"
 
 class FStructOnScope;
 struct FPropertyChangedEvent;
 
-enum class EKeyVectorChannel
-{
-	X,
-	Y,
-	Z,
-	W
-};
-
-
-struct MOVIESCENETRACKS_API FVectorKey
-{
-	FVectorKey(EKeyVectorChannel InChannel, float InValue)
-	{
-		Channel = InChannel;
-		Value = InValue;
-	}
-	EKeyVectorChannel Channel;
-	float Value;
-};
-
-
 /**
 * Base Proxy structure for vector section key data.
 */
@@ -44,24 +23,18 @@ struct FMovieSceneVectorKeyStructBase
 {
 	GENERATED_BODY();
 
-	FRichCurveKey* Keys[4];
-	FRichCurve* Curves[4];
-
 	/** The key's time. */
 	UPROPERTY(EditAnywhere, Category=Key)
-	float Time;
+	FFrameNumber Time;
 
-	virtual void PropagateChanges(const FPropertyChangedEvent& ChangeEvent) override;
+	FMovieSceneKeyStructHelper KeyStructInterop;
 
-	/** Gets the number of channels used by this vector key struct. */
-	virtual int32 GetChannelsUsed() const PURE_VIRTUAL(FMovieSceneVectorKeyStructBase::GetChannelsUsed, return 0;);
+	virtual void PropagateChanges(const FPropertyChangedEvent& ChangeEvent) override;
 
-	/** Gets the value of a channel by index, 0-3 = x-z */
-	virtual float GetPropertyChannelByIndex(int32 Index) const PURE_VIRTUAL(FMovieSceneVectorKeyStructBase::GetPropertyChannelByIndex, return 0;);
-	
-	/** Sets the value of a channel by index, 0-3 = x-z */
-	virtual void SetPropertyChannelByIndex(int32 Index, float Value) PURE_VIRTUAL(FMovieSceneVectorKeyStructBase::SetPropertyChannelByIndex,);
+	/** Gets a ptr value of a channel by index, 0-3 = x-w */
+	virtual float* GetPropertyChannelByIndex(int32 Index) PURE_VIRTUAL(FMovieSceneVectorKeyStructBase::GetPropertyChannelByIndex, return nullptr; );
 };
+template<> struct TStructOpsTypeTraits<FMovieSceneVectorKeyStructBase> : public TStructOpsTypeTraitsBase2<FMovieSceneVectorKeyStructBase> { enum { WithCopy = false }; };
 
 
 /**
@@ -73,17 +46,14 @@ struct FMovieSceneVector2DKeyStruct
 {
 	GENERATED_BODY()
 
-		/** They key's vector value. */
-		UPROPERTY(EditAnywhere, Category=Key)
-		FVector2D Vector;
+	/** They key's vector value. */
+	UPROPERTY(EditAnywhere, Category=Key)
+	FVector2D Vector;
 
 	//~ FMovieSceneVectorKeyStructBase interface
-
-	virtual int32 GetChannelsUsed() const override { return 2; }
-	virtual float GetPropertyChannelByIndex(int32 Index) const override { return Vector[Index]; }
-	virtual void SetPropertyChannelByIndex(int32 Index, float Value) override { Vector[Index] = Value; }
+	virtual float* GetPropertyChannelByIndex(int32 Index) override { return &Vector[Index]; }
 };
-
+template<> struct TStructOpsTypeTraits<FMovieSceneVector2DKeyStruct> : public TStructOpsTypeTraitsBase2<FMovieSceneVector2DKeyStruct> { enum { WithCopy = false }; };
 
 /**
 * Proxy structure for vector section key data.
@@ -94,17 +64,14 @@ struct FMovieSceneVectorKeyStruct
 {
 	GENERATED_BODY()
 
-		/** They key's vector value. */
-		UPROPERTY(EditAnywhere, Category = Key)
-		FVector Vector;
+	/** They key's vector value. */
+	UPROPERTY(EditAnywhere, Category = Key)
+	FVector Vector;
 
 	//~ FMovieSceneVectorKeyStructBase interface
-
-	virtual int32 GetChannelsUsed() const override { return 3; }
-	virtual float GetPropertyChannelByIndex(int32 Index) const override { return Vector[Index]; }
-	virtual void SetPropertyChannelByIndex(int32 Index, float Value) override { Vector[Index] = Value; }
+	virtual float* GetPropertyChannelByIndex(int32 Index) override { return &Vector[Index]; }
 };
-
+template<> struct TStructOpsTypeTraits<FMovieSceneVectorKeyStruct> : public TStructOpsTypeTraitsBase2<FMovieSceneVectorKeyStruct> { enum { WithCopy = false }; };
 
 /**
 * Proxy structure for vector4 section key data.
@@ -115,16 +82,14 @@ struct FMovieSceneVector4KeyStruct
 {
 	GENERATED_BODY()
 
-		/** They key's vector value. */
-		UPROPERTY(EditAnywhere, Category = Key)
-		FVector4 Vector;
+	/** They key's vector value. */
+	UPROPERTY(EditAnywhere, Category = Key)
+	FVector4 Vector;
 
 	//~ FMovieSceneVectorKeyStructBase interface
-
-	virtual int32 GetChannelsUsed() const override { return 4; }
-	virtual float GetPropertyChannelByIndex(int32 Index) const override { return Vector[Index]; }
-	virtual void SetPropertyChannelByIndex(int32 Index, float Value) override { Vector[Index] = Value; }
+	virtual float* GetPropertyChannelByIndex(int32 Index) override { return &Vector[Index]; }
 };
+template<> struct TStructOpsTypeTraits<FMovieSceneVector4KeyStruct> : public TStructOpsTypeTraitsBase2<FMovieSceneVector4KeyStruct> { enum { WithCopy = false }; };
 
 
 /**
@@ -133,52 +98,44 @@ struct FMovieSceneVector4KeyStruct
 UCLASS(MinimalAPI)
 class UMovieSceneVectorSection
 	: public UMovieSceneSection
-	, public IKeyframeSection<FVectorKey>
 {
 	GENERATED_UCLASS_BODY()
 
 public:
 
-	/** Gets one of four curves in this section */
-	FRichCurve& GetCurve(int32 Index) { check(Index >= 0 && Index < 4); return Curves[Index]; }
-	const FRichCurve& GetCurve(int32 Index) const { check(Index >= 0 && Index < 4); return Curves[Index]; }
-
 	/** Sets how many channels are to be used */
 	void SetChannelsUsed(int32 InChannelsUsed) 
 	{
 		checkf(InChannelsUsed >= 2 && InChannelsUsed <= 4, TEXT("Only 2-4 channels are supported."));
 		ChannelsUsed = InChannelsUsed;
+		RecreateChannelProxy();
 	}
-	
-	/** Gets the number of channels in use */
-	int32 GetChannelsUsed() const {return ChannelsUsed;}
-
-public:
 
-	//~ UMovieSceneSection interface
+	/** Gets the number of channels in use */
+	int32 GetChannelsUsed() const { return ChannelsUsed; }
 
-	virtual void MoveSection(float DeltaPosition, TSet<FKeyHandle>& KeyHandles) override;
-	virtual void DilateSection(float DilationFactor, float Origin, TSet<FKeyHandle>& KeyHandles) override;
-	virtual void GetKeyHandles(TSet<FKeyHandle>& OutKeyHandles, TRange<float> TimeRange) const override;
-	virtual TSharedPtr<FStructOnScope> GetKeyStruct(const TArray<FKeyHandle>& KeyHandles) override;
-	virtual TOptional<float> GetKeyTime( FKeyHandle KeyHandle ) const override;
-	virtual void SetKeyTime( FKeyHandle KeyHandle, float Time ) override;
+	/**
+	 * Public access to this section's internal data function
+	 */
+	const FMovieSceneFloatChannel& GetChannel(int32 Index) const
+	{
+		check(Index >= 0 && Index < GetChannelsUsed());
+		return Curves[Index];
+	}
 
-public:
+protected:
 
-	//~ IKeyframeSection interface.
+	//~ UMovieSceneSection interface
+	virtual void Serialize(FArchive& Ar) override;
+	virtual TSharedPtr<FStructOnScope> GetKeyStruct(TArrayView<const FKeyHandle> KeyHandles) override;
 
-	virtual void AddKey(float Time, const FVectorKey& Key, EMovieSceneKeyInterpolation KeyInterpolation) override;
-	virtual bool NewKeyIsNewData(float Time, const FVectorKey& Key) const override;
-	virtual bool HasKeys(const FVectorKey& Key) const override;
-	virtual void SetDefault(const FVectorKey& Key) override;
-	virtual void ClearDefaults() override;
+	void RecreateChannelProxy();
 
 private:
 
-	/** Vector t */
+	/** Float functions for the X,Y,Z,W components of the vector */
 	UPROPERTY()
-	FRichCurve Curves[4];
+	FMovieSceneFloatChannel Curves[4];
 
 	/** How many curves are actually used */
 	UPROPERTY()
diff --git a/Engine/Source/Runtime/MovieSceneTracks/Public/Tracks/MovieScene3DAttachTrack.h b/Engine/Source/Runtime/MovieSceneTracks/Public/Tracks/MovieScene3DAttachTrack.h
--- a/Engine/Source/Runtime/MovieSceneTracks/Public/Tracks/MovieScene3DAttachTrack.h
+++ b/Engine/Source/Runtime/MovieSceneTracks/Public/Tracks/MovieScene3DAttachTrack.h
@@ -20,7 +20,7 @@ public:
 
 	// UMovieScene3DConstraintTrack interface
 
-	virtual void AddConstraint( float Time, float ConstraintEndTime, const FName SocketName, const FName ComponentName, const FMovieSceneObjectBindingID& ConstraintBindingID ) override;
+	virtual void AddConstraint( FFrameNumber Time, int32 Duration, const FName SocketName, const FName ComponentName, const FMovieSceneObjectBindingID& ConstraintBindingID ) override;
 	virtual class UMovieSceneSection* CreateNewSection() override;
 
 public:
diff --git a/Engine/Source/Runtime/MovieSceneTracks/Public/Tracks/MovieScene3DConstraintTrack.h b/Engine/Source/Runtime/MovieSceneTracks/Public/Tracks/MovieScene3DConstraintTrack.h
--- a/Engine/Source/Runtime/MovieSceneTracks/Public/Tracks/MovieScene3DConstraintTrack.h
+++ b/Engine/Source/Runtime/MovieSceneTracks/Public/Tracks/MovieScene3DConstraintTrack.h
@@ -19,19 +19,16 @@ class UMovieScene3DConstraintTrack
 
 public:
 
-	DEPRECATED(4.18, "Constraint guid no longer supported, Use ConstraintBindingID.")
-	virtual void AddConstraint(float Time, float ConstraintEndTime, const FName SocketName, const FName ComponentName, const FGuid& ConstraintId);
-
 	/**
 	 * Adds a constraint.
 	 *
 	 * @param Time The time relative to the owning movie scene where the section should be.
-	 * @param ConstraintEndTime Set the constraint to end at this time.
+	 * @param Duration The length of the constraint section
 	 * @param SocketName The socket name for the constraint.
 	 * @param ComponentName The name of the component the socket resides in.
 	 * @param FMovieSceneObjectBindingID The object binding id to the constraint.
 	 */
-	virtual void AddConstraint(float Time, float ConstraintEndTime, const FName SocketName, const FName ComponentName, const FMovieSceneObjectBindingID& ConstraintBindingID) { }
+	virtual void AddConstraint(FFrameNumber Time, int32 Duration, const FName SocketName, const FName ComponentName, const FMovieSceneObjectBindingID& ConstraintBindingID) { }
 
 public:
 
@@ -42,7 +39,6 @@ public:
 	virtual void AddSection(UMovieSceneSection& Section) override;
 	virtual void RemoveSection(UMovieSceneSection& Section) override;
 	virtual bool IsEmpty() const override;
-	virtual TRange<float> GetSectionBoundaries() const override;
 	virtual const TArray<UMovieSceneSection*>& GetAllSections() const override;
 
 protected:
diff --git a/Engine/Source/Runtime/MovieSceneTracks/Public/Tracks/MovieScene3DPathTrack.h b/Engine/Source/Runtime/MovieSceneTracks/Public/Tracks/MovieScene3DPathTrack.h
--- a/Engine/Source/Runtime/MovieSceneTracks/Public/Tracks/MovieScene3DPathTrack.h
+++ b/Engine/Source/Runtime/MovieSceneTracks/Public/Tracks/MovieScene3DPathTrack.h
@@ -20,7 +20,7 @@ public:
 
 	// UMovieScene3DConstraintTrack interface
 
-	virtual void AddConstraint(float Time, float ConstraintEndTime, const FName SocketName, const FName ComponentName, const FMovieSceneObjectBindingID& ConstraintBindingID) override;
+	virtual void AddConstraint(FFrameNumber Time, int32 Duration, const FName SocketName, const FName ComponentName, const FMovieSceneObjectBindingID& ConstraintBindingID) override;
 
 public:
 
diff --git a/Engine/Source/Runtime/MovieSceneTracks/Public/Tracks/MovieScene3DTransformTrack.h b/Engine/Source/Runtime/MovieSceneTracks/Public/Tracks/MovieScene3DTransformTrack.h
--- a/Engine/Source/Runtime/MovieSceneTracks/Public/Tracks/MovieScene3DTransformTrack.h
+++ b/Engine/Source/Runtime/MovieSceneTracks/Public/Tracks/MovieScene3DTransformTrack.h
@@ -16,30 +16,35 @@ struct FTrajectoryKey
 {
 	struct FData
 	{
-		FData(UMovieScene3DTransformSection* InSection, ERichCurveInterpMode InInterpMode, EMovieSceneTransformChannel InChannel = EMovieSceneTransformChannel::None)
-			: Section(InSection), InterpMode(InInterpMode), Channel(InChannel)
+		FData(UMovieScene3DTransformSection* InSection, TOptional<FKeyHandle> InKeyHandle, ERichCurveInterpMode InInterpMode, FName InChannelName)
+			: Section(InSection), KeyHandle(InKeyHandle), ChannelName(InChannelName), InterpMode(InInterpMode)
 		{}
 
 		TWeakObjectPtr<UMovieScene3DTransformSection> Section;
+		TOptional<FKeyHandle> KeyHandle;
+		FName ChannelName;
 		ERichCurveInterpMode InterpMode;
-		EMovieSceneTransformChannel Channel;
 	};
 
-	FTrajectoryKey(float InTime) : Time(InTime) {}
+	FTrajectoryKey(FFrameNumber InTime) : Time(InTime) {}
 
 	bool Is(ERichCurveInterpMode InInterpMode) const
 	{
-		for (const FData& Value : KeyData)
+		if (KeyData.Num())
 		{
-			if (Value.InterpMode != InInterpMode)
+			for (const FData& Value : KeyData)
 			{
-				return false;
+				if (Value.InterpMode != InInterpMode)
+				{
+					return false;
+				}
 			}
+			return true;
 		}
-		return true;
+		return false;
 	}
 
-	float Time;
+	FFrameNumber Time;
 
 	TArray<FData, TInlineAllocator<1>> KeyData;
 };
@@ -59,10 +64,15 @@ public:
 
 	virtual UMovieSceneSection* CreateNewSection() override;
 
-	MOVIESCENETRACKS_API TArray<FTrajectoryKey> GetTrajectoryData(float Time, int32 MaxNumDataPoints) const;
-
 	/**
 	 * Access the interrogation key for transform data - any interrgation data stored with this key is guaranteed to be of type 'FTransform'
 	 */
 	MOVIESCENETRACKS_API static FMovieSceneInterrogationKey GetInterrogationKey();
+
+
+#if WITH_EDITOR
+
+	MOVIESCENETRACKS_API TArray<FTrajectoryKey> GetTrajectoryData(FFrameNumber Time, int32 MaxNumDataPoints) const;
+
+#endif
 };
diff --git a/Engine/Source/Runtime/MovieSceneTracks/Public/Tracks/MovieSceneAudioTrack.h b/Engine/Source/Runtime/MovieSceneTracks/Public/Tracks/MovieSceneAudioTrack.h
--- a/Engine/Source/Runtime/MovieSceneTracks/Public/Tracks/MovieSceneAudioTrack.h
+++ b/Engine/Source/Runtime/MovieSceneTracks/Public/Tracks/MovieSceneAudioTrack.h
@@ -27,7 +27,10 @@ class MOVIESCENETRACKS_API UMovieSceneAudioTrack
 public:
 
 	/** Adds a new sound cue to the audio */
-	virtual void AddNewSound(USoundBase* Sound, float Time);
+	virtual UMovieSceneSection* AddNewSoundOnRow(USoundBase* Sound, FFrameNumber Time, int32 RowIndex);
+
+	/** Adds a new sound cue on the next available/non-overlapping row */
+	virtual UMovieSceneSection* AddNewSound(USoundBase* Sound, FFrameNumber Time) { return AddNewSoundOnRow(Sound, Time, INDEX_NONE); }
 
 	/** @return The audio sections on this track */
 	const TArray<UMovieSceneSection*>& GetAudioSections() const
@@ -49,7 +52,6 @@ public:
 	virtual void AddSection(UMovieSceneSection& Section) override;
 	virtual void RemoveSection(UMovieSceneSection& Section) override;
 	virtual bool IsEmpty() const override;
-	virtual TRange<float> GetSectionBoundaries() const override;
 	virtual const TArray<UMovieSceneSection*>& GetAllSections() const override;
 	virtual bool SupportsMultipleRows() const override;
 	virtual FMovieSceneTrackRowSegmentBlenderPtr GetRowSegmentBlender() const override;
diff --git a/Engine/Source/Runtime/MovieSceneTracks/Public/Tracks/MovieSceneByteTrack.h b/Engine/Source/Runtime/MovieSceneTracks/Public/Tracks/MovieSceneByteTrack.h
--- a/Engine/Source/Runtime/MovieSceneTracks/Public/Tracks/MovieSceneByteTrack.h
+++ b/Engine/Source/Runtime/MovieSceneTracks/Public/Tracks/MovieSceneByteTrack.h
@@ -7,6 +7,8 @@
 #include "Tracks/MovieScenePropertyTrack.h"
 #include "MovieSceneByteTrack.generated.h"
 
+class UEnum;
+
 /**
  * Handles manipulation of byte properties in a movie scene
  */
@@ -17,6 +19,7 @@ class MOVIESCENETRACKS_API UMovieSceneByteTrack : public UMovieScenePropertyTrac
 
 public:
 	/** UMovieSceneTrack interface */
+	virtual void PostLoad() override;
 	virtual UMovieSceneSection* CreateNewSection() override;
 	virtual FMovieSceneEvalTemplatePtr CreateTemplateForSection(const UMovieSceneSection& InSection) const override;
 
diff --git a/Engine/Source/Runtime/MovieSceneTracks/Public/Tracks/MovieSceneCameraAnimTrack.h b/Engine/Source/Runtime/MovieSceneTracks/Public/Tracks/MovieSceneCameraAnimTrack.h
--- a/Engine/Source/Runtime/MovieSceneTracks/Public/Tracks/MovieSceneCameraAnimTrack.h
+++ b/Engine/Source/Runtime/MovieSceneTracks/Public/Tracks/MovieSceneCameraAnimTrack.h
@@ -21,8 +21,8 @@ class UMovieSceneCameraAnimTrack : public UMovieSceneNameableTrack
 	GENERATED_BODY()
 
 public:
-	virtual void AddNewCameraAnim(float KeyTime, class UCameraAnim* CameraAnim);
-	void GetCameraAnimSectionsAtTime(float Time, TArray<UMovieSceneCameraAnimSection*>& OutSections);
+	virtual UMovieSceneSection* AddNewCameraAnim(FFrameNumber KeyTime, class UCameraAnim* CameraAnim);
+	void GetCameraAnimSectionsAtTime(FFrameNumber Time, TArray<UMovieSceneCameraAnimSection*>& OutSections);
 
 public:
 
@@ -31,7 +31,6 @@ public:
 	virtual void AddSection(UMovieSceneSection& Section) override;
 	virtual void RemoveSection(UMovieSceneSection& Section) override;
 	virtual bool IsEmpty() const override;
-	virtual TRange<float> GetSectionBoundaries() const override;
 	virtual bool SupportsMultipleRows() const override { return true; }
 	virtual const TArray<UMovieSceneSection*>& GetAllSections() const override;
 	virtual UMovieSceneSection* CreateNewSection() override;
diff --git a/Engine/Source/Runtime/MovieSceneTracks/Public/Tracks/MovieSceneCameraCutTrack.h b/Engine/Source/Runtime/MovieSceneTracks/Public/Tracks/MovieSceneCameraCutTrack.h
--- a/Engine/Source/Runtime/MovieSceneTracks/Public/Tracks/MovieSceneCameraCutTrack.h
+++ b/Engine/Source/Runtime/MovieSceneTracks/Public/Tracks/MovieSceneCameraCutTrack.h
@@ -23,17 +23,14 @@ class UMovieSceneCameraCutTrack
 	
 public:
 
-	DEPRECATED(4.18, "Camera guid no longer supported, Use CameraBindingID.")
-	MOVIESCENETRACKS_API void AddNewCameraCut(FGuid CameraHandle, float Time);
-
 	/** 
 	 * Adds a new CameraCut at the specified time.
 	 *	
 	 * @param CameraBindingID Handle to the camera that the CameraCut switches to when active.
-	 * @param TimeRange The range within this track's movie scene where the CameraCut is initially placed.
+	 * @param Time The within this track's movie scene where the CameraCut is initially placed.
 	 * @return The newly created camera cut section
 	 */
-	MOVIESCENETRACKS_API UMovieSceneCameraCutSection* AddNewCameraCut(const FMovieSceneObjectBindingID& CameraBindingID, float Time);
+	MOVIESCENETRACKS_API UMovieSceneCameraCutSection* AddNewCameraCut(const FMovieSceneObjectBindingID& CameraBindingID, FFrameNumber Time);
 
 public:
 
@@ -54,7 +51,7 @@ public:
 
 protected:
 
-	float FindEndTimeForCameraCut(float StartTime);
+	FFrameNumber FindEndTimeForCameraCut(FFrameNumber StartTime);
 
 private:
 
diff --git a/Engine/Source/Runtime/MovieSceneTracks/Public/Tracks/MovieSceneCameraShakeTrack.h b/Engine/Source/Runtime/MovieSceneTracks/Public/Tracks/MovieSceneCameraShakeTrack.h
--- a/Engine/Source/Runtime/MovieSceneTracks/Public/Tracks/MovieSceneCameraShakeTrack.h
+++ b/Engine/Source/Runtime/MovieSceneTracks/Public/Tracks/MovieSceneCameraShakeTrack.h
@@ -22,7 +22,7 @@ class UMovieSceneCameraShakeTrack : public UMovieSceneNameableTrack
 	GENERATED_BODY()
 
 public:
-	virtual void AddNewCameraShake(float KeyTime, TSubclassOf<UCameraShake> ShakeClass);
+	virtual UMovieSceneSection* AddNewCameraShake(FFrameNumber KeyTime, TSubclassOf<UCameraShake> ShakeClass);
 	
 public:
 
@@ -31,7 +31,6 @@ public:
 	virtual void AddSection(UMovieSceneSection& Section) override;
 	virtual void RemoveSection(UMovieSceneSection& Section) override;
 	virtual bool IsEmpty() const override;
-	virtual TRange<float> GetSectionBoundaries() const override;
 	virtual bool SupportsMultipleRows() const override { return true; }
 	virtual const TArray<UMovieSceneSection*>& GetAllSections() const override;
 	virtual UMovieSceneSection* CreateNewSection() override;
diff --git a/Engine/Source/Runtime/MovieSceneTracks/Public/Tracks/MovieSceneCinematicShotTrack.h b/Engine/Source/Runtime/MovieSceneTracks/Public/Tracks/MovieSceneCinematicShotTrack.h
--- a/Engine/Source/Runtime/MovieSceneTracks/Public/Tracks/MovieSceneCinematicShotTrack.h
+++ b/Engine/Source/Runtime/MovieSceneTracks/Public/Tracks/MovieSceneCinematicShotTrack.h
@@ -27,7 +27,8 @@ public:
 	
 	// UMovieSceneSubTrack interface
 
-	MOVIESCENETRACKS_API virtual UMovieSceneSubSection* AddSequence(UMovieSceneSequence* Sequence, float StartTime, float Duration, const bool& bInsertSequence = false);
+	MOVIESCENETRACKS_API virtual UMovieSceneSubSection* AddSequence(UMovieSceneSequence* Sequence, FFrameNumber StartTime, int32 Duration) { return AddSequenceOnRow(Sequence, StartTime, Duration, INDEX_NONE); }
+	MOVIESCENETRACKS_API virtual UMovieSceneSubSection* AddSequenceOnRow(UMovieSceneSequence* Sequence, FFrameNumber StartTime, int32 Duration, int32 RowIndex);
 
 	// UMovieSceneTrack interface
 
diff --git a/Engine/Source/Runtime/MovieSceneTracks/Public/Tracks/MovieSceneEnumTrack.h b/Engine/Source/Runtime/MovieSceneTracks/Public/Tracks/MovieSceneEnumTrack.h
--- a/Engine/Source/Runtime/MovieSceneTracks/Public/Tracks/MovieSceneEnumTrack.h
+++ b/Engine/Source/Runtime/MovieSceneTracks/Public/Tracks/MovieSceneEnumTrack.h
@@ -7,6 +7,8 @@
 #include "Tracks/MovieScenePropertyTrack.h"
 #include "MovieSceneEnumTrack.generated.h"
 
+class UEnum;
+
 /**
  * Handles manipulation of byte properties in a movie scene
  */
@@ -17,6 +19,7 @@ class MOVIESCENETRACKS_API UMovieSceneEnumTrack : public UMovieScenePropertyTrac
 
 public:
 	/** UMovieSceneTrack interface */
+	virtual void PostLoad() override;
 	virtual UMovieSceneSection* CreateNewSection() override;
 	virtual FMovieSceneEvalTemplatePtr CreateTemplateForSection(const UMovieSceneSection& InSection) const override;
 
diff --git a/Engine/Source/Runtime/MovieSceneTracks/Public/Tracks/MovieSceneEventTrack.h b/Engine/Source/Runtime/MovieSceneTracks/Public/Tracks/MovieSceneEventTrack.h
--- a/Engine/Source/Runtime/MovieSceneTracks/Public/Tracks/MovieSceneEventTrack.h
+++ b/Engine/Source/Runtime/MovieSceneTracks/Public/Tracks/MovieSceneEventTrack.h
@@ -52,7 +52,6 @@ public:
 	virtual void AddSection(UMovieSceneSection& Section) override;
 	virtual UMovieSceneSection* CreateNewSection() override;
 	virtual const TArray<UMovieSceneSection*>& GetAllSections() const override;
-	virtual TRange<float> GetSectionBoundaries() const override;
 	virtual bool HasSection(const UMovieSceneSection& Section) const override;
 	virtual bool IsEmpty() const override;
 	virtual void RemoveAllAnimationData() override;
diff --git a/Engine/Source/Runtime/MovieSceneTracks/Public/Tracks/MovieSceneFloatTrack.h b/Engine/Source/Runtime/MovieSceneTracks/Public/Tracks/MovieSceneFloatTrack.h
--- a/Engine/Source/Runtime/MovieSceneTracks/Public/Tracks/MovieSceneFloatTrack.h
+++ b/Engine/Source/Runtime/MovieSceneTracks/Public/Tracks/MovieSceneFloatTrack.h
@@ -19,7 +19,4 @@ public:
 	/** UMovieSceneTrack interface */
 	virtual UMovieSceneSection* CreateNewSection() override;
 	virtual FMovieSceneEvalTemplatePtr CreateTemplateForSection(const UMovieSceneSection& InSection) const override;
-
-	DEPRECATED(4.15, "Direct evaluation of float tracks is no longer supported. Please create an evaluation template (see FMovieSceneFloatPropertySectionTemplate).")
-	virtual bool Eval( float Position, float LastPosition, float& InOutFloat ) const;
 };
diff --git a/Engine/Source/Runtime/MovieSceneTracks/Public/Tracks/MovieSceneMaterialTrack.h b/Engine/Source/Runtime/MovieSceneTracks/Public/Tracks/MovieSceneMaterialTrack.h
--- a/Engine/Source/Runtime/MovieSceneTracks/Public/Tracks/MovieSceneMaterialTrack.h
+++ b/Engine/Source/Runtime/MovieSceneTracks/Public/Tracks/MovieSceneMaterialTrack.h
@@ -27,7 +27,6 @@ public:
 	virtual void AddSection(UMovieSceneSection& Section) override;
 	virtual void RemoveSection(UMovieSceneSection& Section) override;
 	virtual bool IsEmpty() const override;
-	virtual TRange<float> GetSectionBoundaries() const override;
 	virtual const TArray<UMovieSceneSection*>& GetAllSections() const override;
 
 public:
@@ -38,7 +37,7 @@ public:
 	 * @param Time The time to add the new key.
 	 * @param The value for the new key.
 	 */
-	void AddScalarParameterKey(FName ParameterName, float Position, float Value);
+	void AddScalarParameterKey(FName ParameterName, FFrameNumber Position, float Value);
 
 	/**
 	* Adds a color parameter key to the track.
@@ -46,7 +45,7 @@ public:
 	* @param Time The time to add the new key.
 	* @param The value for the new key.
 	*/
-	void AddColorParameterKey(FName ParameterName, float Position, FLinearColor Value);
+	void AddColorParameterKey(FName ParameterName, FFrameNumber Position, FLinearColor Value);
 
 private:
 
