Sequencer: Ensure that we only set montages to playback in cooked builds only if the movie scene player itself is also playing to maintain parity with PreviewSetAnimPosition

diff --git a/Engine/Source/Runtime/MovieSceneTracks/Private/Evaluation/MovieSceneSkeletalAnimationTemplate.cpp b/Engine/Source/Runtime/MovieSceneTracks/Private/Evaluation/MovieSceneSkeletalAnimationTemplate.cpp
--- a/Engine/Source/Runtime/MovieSceneTracks/Private/Evaluation/MovieSceneSkeletalAnimationTemplate.cpp
+++ b/Engine/Source/Runtime/MovieSceneTracks/Private/Evaluation/MovieSceneSkeletalAnimationTemplate.cpp
@@ -221,7 +221,7 @@ namespace MovieScene
 				{
 					SetAnimPosition(PersistentData, Player, SkeletalMeshComponent,
 						AnimParams.SlotName, AnimParams.Section, AnimParams.Animation, AnimParams.EvalTime, AnimParams.BlendWeight,
-						bLooping, bFireNotifies);
+						bLooping, Player.GetPlaybackStatus() == EMovieScenePlayerStatus::Playing, bFireNotifies);
 				}
 			}
 
@@ -245,7 +245,7 @@ namespace MovieScene
 			return nullptr;
 		}
 
-		void SetAnimPosition(FPersistentEvaluationData& PersistentData, IMovieScenePlayer& Player, USkeletalMeshComponent* SkeletalMeshComponent, FName SlotName, FObjectKey Section, UAnimSequenceBase* InAnimSequence, float InPosition, float Weight, bool bLooping, bool bFireNotifies)
+		void SetAnimPosition(FPersistentEvaluationData& PersistentData, IMovieScenePlayer& Player, USkeletalMeshComponent* SkeletalMeshComponent, FName SlotName, FObjectKey Section, UAnimSequenceBase* InAnimSequence, float InPosition, float Weight, bool bLooping, bool bPlaying, bool bFireNotifies)
 		{
 			if (!CanPlayAnimation(SkeletalMeshComponent, InAnimSequence))
 			{
@@ -278,9 +278,9 @@ namespace MovieScene
 					FMovieSceneAnimTypeID SlotTypeID = SectionToAnimationIDs.GetAnimTypeID(Section);
 					Player.SavePreAnimatedState(*Montage.Get(), SlotTypeID, FStopPlayingMontageTokenProducer(AnimInst, InstanceId));
 
-					// make sure it's playing
+					// make sure it's playing if the sequence is
 					FAnimMontageInstance* Instance = AnimInst->GetMontageInstanceForID(InstanceId);
-					Instance->bPlaying = true;
+					Instance->bPlaying = bPlaying;
 				}
 			}
 		}
