Sequencer: Fix for crash OnKeySelected and GetAutoSizeRange

diff --git a/Engine/Source/Editor/Sequencer/Private/SequencerEdMode.cpp b/Engine/Source/Editor/Sequencer/Private/SequencerEdMode.cpp
--- a/Engine/Source/Editor/Sequencer/Private/SequencerEdMode.cpp
+++ b/Engine/Source/Editor/Sequencer/Private/SequencerEdMode.cpp
@@ -132,6 +132,11 @@ void FSequencerEdMode::AddReferencedObjects(FReferenceCollector& Collector)
 
 void FSequencerEdMode::OnKeySelected(FViewport* Viewport, HMovieSceneKeyProxy* KeyProxy)
 {
+	if (!KeyProxy)
+	{
+		return;
+	}
+
 	const bool bToggleSelection = Viewport->KeyState(EKeys::LeftControl) || Viewport->KeyState(EKeys::RightControl);
 	const bool bAddToSelection = Viewport->KeyState(EKeys::LeftShift) || Viewport->KeyState(EKeys::RightShift);
 
diff --git a/Engine/Source/Runtime/MovieSceneTracks/Private/Sections/MovieSceneAudioSection.cpp b/Engine/Source/Runtime/MovieSceneTracks/Private/Sections/MovieSceneAudioSection.cpp
--- a/Engine/Source/Runtime/MovieSceneTracks/Private/Sections/MovieSceneAudioSection.cpp
+++ b/Engine/Source/Runtime/MovieSceneTracks/Private/Sections/MovieSceneAudioSection.cpp
@@ -109,6 +109,11 @@ float GetStartOffsetAtTrimTime(FQualifiedFrameTime TrimTime, float StartOffset,
 	
 TOptional<TRange<FFrameNumber> > UMovieSceneAudioSection::GetAutoSizeRange() const
 {
+	if (!Sound)
+	{
+		return TRange<FFrameNumber>();
+	}
+
 	float SoundDuration = MovieSceneHelpers::GetSoundDuration(Sound);
 
 	FFrameRate FrameRate = GetTypedOuter<UMovieScene>()->GetTickResolution();
