Sequencer: Expose AddScalarParameterKey, AddVectorParameterKey, AddColorParameterKey

diff --git a/Engine/Source/Runtime/MovieSceneTracks/Public/Sections/MovieSceneParameterSection.h b/Engine/Source/Runtime/MovieSceneTracks/Public/Sections/MovieSceneParameterSection.h
--- a/Engine/Source/Runtime/MovieSceneTracks/Public/Sections/MovieSceneParameterSection.h
+++ b/Engine/Source/Runtime/MovieSceneTracks/Public/Sections/MovieSceneParameterSection.h
@@ -176,13 +176,13 @@ class UMovieSceneParameterSection
 
 public:
 	/** Adds a a key for a specific scalar parameter at the specified time with the specified value. */
-	void AddScalarParameterKey(FName InParameterName, FFrameNumber InTime, float InValue);
+	MOVIESCENETRACKS_API void AddScalarParameterKey(FName InParameterName, FFrameNumber InTime, float InValue);
 
 	/** Adds a a key for a specific vector parameter at the specified time with the specified value. */
-	void AddVectorParameterKey(FName InParameterName, FFrameNumber InTime, FVector InValue);
+	MOVIESCENETRACKS_API void AddVectorParameterKey(FName InParameterName, FFrameNumber InTime, FVector InValue);
 
 	/** Adds a a key for a specific color parameter at the specified time with the specified value. */
-	void AddColorParameterKey(FName InParameterName, FFrameNumber InTime, FLinearColor InValue);
+	MOVIESCENETRACKS_API void AddColorParameterKey(FName InParameterName, FFrameNumber InTime, FLinearColor InValue);
 
 	/** 
 	 * Removes a scalar parameter from this section. 
