Sequencer: Allow particle activation/deactivation keys to fire even if they're not in the playback range.

Trigger keys remain uneffected - they still fire only if they're within the playback range.

diff --git a/Engine/Source/Runtime/MovieSceneTracks/Private/Evaluation/MovieSceneParticleTemplate.cpp b/Engine/Source/Runtime/MovieSceneTracks/Private/Evaluation/MovieSceneParticleTemplate.cpp
--- a/Engine/Source/Runtime/MovieSceneTracks/Private/Evaluation/MovieSceneParticleTemplate.cpp
+++ b/Engine/Source/Runtime/MovieSceneTracks/Private/Evaluation/MovieSceneParticleTemplate.cpp
@@ -116,7 +116,7 @@ FMovieSceneParticleSectionTemplate::FMovieSceneParticleSectionTemplate(const UMo
 
 void FMovieSceneParticleSectionTemplate::Evaluate(const FMovieSceneEvaluationOperand& Operand, const FMovieSceneContext& Context, const FPersistentEvaluationData& PersistentData, FMovieSceneExecutionTokens& ExecutionTokens) const
 {
-	const bool bPlaying = Context.GetDirection() == EPlayDirection::Forwards && Context.GetRange().Size<FFrameTime>() >= FFrameTime(0) && Context.GetStatus() == EMovieScenePlayerStatus::Playing;
+	const bool bPlaying = Context.GetDirection() == EPlayDirection::Forwards && Context.GetRange().Size<FFrameTime>() >= FFrameTime(0);
 
 	if (!bPlaying)
 	{
@@ -133,10 +133,25 @@ void FMovieSceneParticleSectionTemplate::Evaluate(const FMovieSceneEvaluationOpe
 		TArrayView<const uint8>        Values = ChannelData.GetValues();
 
 		const int32 LastKeyIndex = Algo::UpperBound(Times, PlaybackRange.GetUpperBoundValue())-1;
-		if (LastKeyIndex >= 0 && PlaybackRange.Contains(Times[LastKeyIndex]))
+		if (LastKeyIndex >= 0)
 		{
-			FParticleTrackExecutionToken NewToken((EParticleKey)Values[LastKeyIndex]);
-			ExecutionTokens.Add(MoveTemp(NewToken));
+			EParticleKey ParticleKey((EParticleKey)Values[LastKeyIndex]);
+			
+			// If the particle key is a trigger, limit the activation/deactivation to within the playback range
+			if (ParticleKey == EParticleKey::Trigger)
+			{
+				if (PlaybackRange.Contains(Times[LastKeyIndex]))
+				{
+					FParticleTrackExecutionToken NewToken(ParticleKey);
+					ExecutionTokens.Add(MoveTemp(NewToken));
+				}
+			}
+			// Otherwise, this is within the activation/deactivation keyframe range
+			else
+			{
+				FParticleTrackExecutionToken NewToken(ParticleKey);
+				ExecutionTokens.Add(MoveTemp(NewToken));
+			}
 		}
 	}
 }
