Sequencer: Don't reactivate the particle system if it's already active.

diff --git a/Engine/Source/Runtime/MovieSceneTracks/Private/Evaluation/MovieSceneParticleTemplate.cpp b/Engine/Source/Runtime/MovieSceneTracks/Private/Evaluation/MovieSceneParticleTemplate.cpp
--- a/Engine/Source/Runtime/MovieSceneTracks/Private/Evaluation/MovieSceneParticleTemplate.cpp
+++ b/Engine/Source/Runtime/MovieSceneTracks/Private/Evaluation/MovieSceneParticleTemplate.cpp
@@ -87,11 +87,10 @@ struct FParticleTrackExecutionToken
 
 				if ( ParticleKey == EParticleKey::Activate)
 				{
-					if ( ParticleSystemComponent->IsActive() )
+					if ( !ParticleSystemComponent->IsActive() )
 					{
-						ParticleSystemComponent->SetActive(false, true);
+						ParticleSystemComponent->SetActive(true, true);
 					}
-					ParticleSystemComponent->SetActive(true, true);
 				}
 				else if( ParticleKey == EParticleKey::Deactivate )
 				{
