Sequencer: Forcibly tick skeletal mesh component in case it has already been ticked.

This fixes first frame T pose issues where skeletal component ticks first, then sequencer resolves bindings which leads to adding tick prerequisites for the next frame.

Copy from Release-4.20

diff --git a/Engine/Source/Runtime/MovieSceneTracks/Private/Evaluation/MovieSceneSkeletalAnimationTemplate.cpp b/Engine/Source/Runtime/MovieSceneTracks/Private/Evaluation/MovieSceneSkeletalAnimationTemplate.cpp
--- a/Engine/Source/Runtime/MovieSceneTracks/Private/Evaluation/MovieSceneSkeletalAnimationTemplate.cpp
+++ b/Engine/Source/Runtime/MovieSceneTracks/Private/Evaluation/MovieSceneSkeletalAnimationTemplate.cpp
@@ -228,6 +228,21 @@ namespace MovieScene
 				}
 			}
 
+			// If the skeletal component has already ticked this frame because tick prerequisites weren't set up yet, forcibly tick this component to update.
+			// This resolves first frame issues where the skeletal component ticks first, then the sequencer binding is resolved which sets up tick prerequisites
+			// for the next frame.
+			if (SkeletalMeshComponent->PoseTickedThisFrame())
+			{
+				SkeletalMeshComponent->TickAnimation(0.f, false);
+
+				SkeletalMeshComponent->RefreshBoneTransforms();
+				SkeletalMeshComponent->RefreshSlaveComponents();
+				SkeletalMeshComponent->UpdateComponentToWorld();
+				SkeletalMeshComponent->FinalizeBoneTransform();
+				SkeletalMeshComponent->MarkRenderTransformDirty();
+				SkeletalMeshComponent->MarkRenderDynamicDataDirty();
+			}
+
 			Player.PreAnimatedState.SetCaptureEntity(FMovieSceneEvaluationKey(), EMovieSceneCompletionMode::KeepState);
 		}
 
