Sequencer: Fix restore state for attach tracks so that the original parent is re-attached

diff --git a/Engine/Source/Runtime/MovieSceneTracks/Private/Evaluation/MovieScene3DAttachTemplate.cpp b/Engine/Source/Runtime/MovieSceneTracks/Private/Evaluation/MovieScene3DAttachTemplate.cpp
--- a/Engine/Source/Runtime/MovieSceneTracks/Private/Evaluation/MovieScene3DAttachTemplate.cpp
+++ b/Engine/Source/Runtime/MovieSceneTracks/Private/Evaluation/MovieScene3DAttachTemplate.cpp
@@ -40,6 +40,7 @@ struct F3DAttachTrackToken : public F3DAttachRuleState
 {
 	F3DAttachTrackToken(USceneComponent* InAttachParent, FName InAttachSocketName, bool bInShouldbeAttached, EAttachmentRule InAttachmentLocationRule, EAttachmentRule InAttachmentRotationRule, EAttachmentRule InAttachmentScaleRule, EDetachmentRule InDetachmentLocationRule, EDetachmentRule InDetachmentRotationRule, EDetachmentRule InDetachmentScaleRule)
 		: F3DAttachRuleState(InAttachmentLocationRule, InAttachmentRotationRule, InAttachmentScaleRule, InDetachmentLocationRule, InDetachmentRotationRule, InDetachmentScaleRule)
+		, AttachParent(InAttachParent)
 		, AttachSocketName(InAttachSocketName)
 		, bShouldBeAttached(bInShouldbeAttached)
 	{}
@@ -95,6 +96,11 @@ struct F3DAttachTrackPreAnimatedToken : F3DAttachTrackToken, IMovieScenePreAnima
 
 			SceneComponent->DetachFromComponent(DetachmentRules);
 		}
+		else
+		{
+			// Note, attach rules here are insignificant because the transform will be restored by the transform token
+			SceneComponent->AttachToComponent(AttachParent.Get(), FAttachmentTransformRules::KeepWorldTransform, AttachSocketName);
+		}
 	}
 };
 
