Sequencer: Disable audio when playing backwards

diff --git a/Engine/Source/Runtime/MovieSceneTracks/Private/Evaluation/MovieSceneAudioTemplate.cpp b/Engine/Source/Runtime/MovieSceneTracks/Private/Evaluation/MovieSceneAudioTemplate.cpp
--- a/Engine/Source/Runtime/MovieSceneTracks/Private/Evaluation/MovieSceneAudioTemplate.cpp
+++ b/Engine/Source/Runtime/MovieSceneTracks/Private/Evaluation/MovieSceneAudioTemplate.cpp
@@ -185,7 +185,7 @@ struct FAudioSectionExecutionToken : IMovieSceneExecutionToken
 	{
 		FCachedAudioTrackData& TrackData = PersistentData.GetOrAddTrackData<FCachedAudioTrackData>();
 
-		if ((Context.GetStatus() != EMovieScenePlayerStatus::Playing && Context.GetStatus() != EMovieScenePlayerStatus::Scrubbing) || Context.HasJumped())
+		if ((Context.GetStatus() != EMovieScenePlayerStatus::Playing && Context.GetStatus() != EMovieScenePlayerStatus::Scrubbing) || Context.HasJumped() || Context.GetDirection() == EPlayDirection::Backwards)
 		{
 			// stopped, recording, etc
 			TrackData.StopAllSounds();
