Sequencer: Event tracks now always fill empty space to ensure that they always get evaluated when playing off a finite section

diff --git a/Engine/Source/Runtime/MovieSceneTracks/Private/Tracks/MovieSceneEventTrack.cpp b/Engine/Source/Runtime/MovieSceneTracks/Private/Tracks/MovieSceneEventTrack.cpp
--- a/Engine/Source/Runtime/MovieSceneTracks/Private/Tracks/MovieSceneEventTrack.cpp
+++ b/Engine/Source/Runtime/MovieSceneTracks/Private/Tracks/MovieSceneEventTrack.cpp
@@ -98,6 +98,28 @@ void UMovieSceneEventTrack::PostCompile(FMovieSceneEvaluationTrack& Track, const
 	Track.SetEvaluationMethod(EEvaluationMethod::Swept);
 }
 
+FMovieSceneTrackSegmentBlenderPtr UMovieSceneEventTrack::GetTrackSegmentBlender() const
+{
+	// This is a temporary measure to alleviate some issues with event tracks with finite ranges.
+	// By filling empty space between sections, we're essentially always making this track evaluate
+	// which allows it to sweep sections correctly when the play-head moves from a finite section
+	// to empty space. This doesn't address the issue of the play-head moving from inside a sub-sequence
+	// to outside, but that specific issue is even more nuanced and complicated to address.
+	struct FMovieSceneEventTrackSegmentBlender : FMovieSceneTrackSegmentBlender
+	{
+		FMovieSceneEventTrackSegmentBlender()
+		{
+			bCanFillEmptySpace = true;
+			bAllowEmptySegments = true;
+		}
+		virtual TOptional<FMovieSceneSegment> InsertEmptySpace(const TRange<FFrameNumber>& Range, const FMovieSceneSegment* PreviousSegment, const FMovieSceneSegment* NextSegment) const
+		{
+			return FMovieSceneSegment(Range);
+		}
+	};
+	return FMovieSceneEventTrackSegmentBlender();
+}
+
 #if WITH_EDITORONLY_DATA
 
 FText UMovieSceneEventTrack::GetDefaultDisplayName() const
diff --git a/Engine/Source/Runtime/MovieSceneTracks/Public/Tracks/MovieSceneEventTrack.h b/Engine/Source/Runtime/MovieSceneTracks/Public/Tracks/MovieSceneEventTrack.h
--- a/Engine/Source/Runtime/MovieSceneTracks/Public/Tracks/MovieSceneEventTrack.h
+++ b/Engine/Source/Runtime/MovieSceneTracks/Public/Tracks/MovieSceneEventTrack.h
@@ -59,6 +59,7 @@ public:
 	virtual FMovieSceneEvalTemplatePtr CreateTemplateForSection(const UMovieSceneSection& InSection) const override;
 	virtual void PostCompile(FMovieSceneEvaluationTrack& Track, const FMovieSceneTrackCompilerArgs& Args) const override;
 	virtual bool SupportsMultipleRows() const override { return true; }
+	virtual FMovieSceneTrackSegmentBlenderPtr GetTrackSegmentBlender() const override;
 
 #if WITH_EDITORONLY_DATA
 	virtual FText GetDefaultDisplayName() const override;
