Sequencer: Event repeater sections now only trigger if the play head is over the section

Changed FMovieSceneEvalTemplate::EvaluateSwept to include the swept range rather than providing a cut-and-sliced context

diff --git a/Engine/Source/Runtime/MovieScene/Private/Evaluation/MovieSceneEvaluationTrack.cpp b/Engine/Source/Runtime/MovieScene/Private/Evaluation/MovieSceneEvaluationTrack.cpp
--- a/Engine/Source/Runtime/MovieScene/Private/Evaluation/MovieSceneEvaluationTrack.cpp
+++ b/Engine/Source/Runtime/MovieScene/Private/Evaluation/MovieSceneEvaluationTrack.cpp
@@ -522,19 +522,20 @@ void FMovieSceneEvaluationTrack::EvaluateSwept(FMovieSceneSegmentIdentifier Segm
 
 	GatherSweptSegments(Context.GetFrameNumberRange(), SortedIndex, Segments.GetSorted(), ImplToAccumulatedRange);
 
-	for (auto& Pair : ImplToAccumulatedRange)
+	ExecutionTokens.SetContext(Context);
+	for (const TTuple<int32, TRange<FFrameNumber>>& Pair : ImplToAccumulatedRange)
 	{
 		const int32 SectionIndex = Pair.Key;
-		TRange<FFrameTime> ClampRange = FMovieSceneEvaluationRange::NumberRangeToTimeRange(Pair.Value);
-		const FMovieSceneEvalTemplate& Template = GetChildTemplate(SectionIndex);
+		const TRange<FFrameNumber>&    SweptRange = Pair.Value;
+		const FMovieSceneEvalTemplate& Template   = GetChildTemplate(SectionIndex);
 
 		PersistentData.DeriveSectionKey(SectionIndex);
 		ExecutionTokens.SetCurrentScope(FMovieSceneEvaluationScope(PersistentData.GetSectionKey(), Template.GetCompletionMode()));
-		ExecutionTokens.SetContext(Context);
-		
+
 		Template.EvaluateSwept(
 			Operand,
-			Context.Clamp(ClampRange),
+			Context,
+			SweptRange,
 			PersistentData,
 			ExecutionTokens);
 	}
diff --git a/Engine/Source/Runtime/MovieScene/Public/Evaluation/MovieSceneEvalTemplate.h b/Engine/Source/Runtime/MovieScene/Public/Evaluation/MovieSceneEvalTemplate.h
--- a/Engine/Source/Runtime/MovieScene/Public/Evaluation/MovieSceneEvalTemplate.h
+++ b/Engine/Source/Runtime/MovieScene/Public/Evaluation/MovieSceneEvalTemplate.h
@@ -108,9 +108,18 @@ public:
 	 *
 	 * @param Operand			Unique handle to the operand on which we are to operate. Only to be used as a reference, or forwarded throgh to an execution token.
 	 * @param Context			Evaluation context specifying the current evaluation time, sub sequence transform and other relevant information.
+	 * @param SweptRange		The range this is to be swept in this evaluation - always fully contained by the context's range
 	 * @param PersistentData	Persistent data store which can be used to access arbitrary data pertaining to the current template that should have been set up in initialize.
 	 * @param ExecutionTokens	Stack of execution tokens that will be used to apply animated state to the environment at a later time.
-	 */
+	*/
+	virtual void EvaluateSwept(const FMovieSceneEvaluationOperand& Operand, const FMovieSceneContext& Context, const TRange<FFrameNumber>& SweptRange, const FPersistentEvaluationData& PersistentData, FMovieSceneExecutionTokens& ExecutionTokens) const
+	{
+		// Call deprecated one
+		PRAGMA_DISABLE_DEPRECATION_WARNINGS
+		EvaluateSwept(Operand, Context, PersistentData, ExecutionTokens);
+		PRAGMA_ENABLE_DEPRECATION_WARNINGS
+	}
+	DEPRECATED(4.21, "Please override void EvaluateSwept(const FMovieSceneEvaluationOperand&, const FMovieSceneContext&, const TRange<FFrameNumber>&, const FPersistentEvaluationData&, FMovieSceneExecutionTokens&) instead.")
 	virtual void EvaluateSwept(const FMovieSceneEvaluationOperand& Operand, const FMovieSceneContext& Context, const FPersistentEvaluationData& PersistentData, FMovieSceneExecutionTokens& ExecutionTokens) const
 	{
 		ensureMsgf(false, TEXT("FMovieSceneEvalTemplate::EvaluateSwept has not been implemented. Verify that this template's evaluation track has correct evaluation method (usually set in UMovieSceneTrack::PostCompile), or implement this function."));
diff --git a/Engine/Source/Runtime/MovieSceneTracks/Private/Evaluation/MovieSceneEventTemplate.cpp b/Engine/Source/Runtime/MovieSceneTracks/Private/Evaluation/MovieSceneEventTemplate.cpp
--- a/Engine/Source/Runtime/MovieSceneTracks/Private/Evaluation/MovieSceneEventTemplate.cpp
+++ b/Engine/Source/Runtime/MovieSceneTracks/Private/Evaluation/MovieSceneEventTemplate.cpp
@@ -344,7 +344,7 @@ FMovieSceneEventSectionTemplate::FMovieSceneEventSectionTemplate(const UMovieSce
 {
 }
 
-void FMovieSceneEventSectionTemplate::EvaluateSwept(const FMovieSceneEvaluationOperand& Operand, const FMovieSceneContext& Context, const FPersistentEvaluationData& PersistentData, FMovieSceneExecutionTokens& ExecutionTokens) const
+void FMovieSceneEventSectionTemplate::EvaluateSwept(const FMovieSceneEvaluationOperand& Operand, const FMovieSceneContext& Context, const TRange<FFrameNumber>& SweptRange, const FPersistentEvaluationData& PersistentData, FMovieSceneExecutionTokens& ExecutionTokens) const
 {
 	// Don't allow events to fire when playback is in a stopped state. This can occur when stopping 
 	// playback and returning the current position to the start of playback. It's not desireable to have 
@@ -364,8 +364,6 @@ void FMovieSceneEventSectionTemplate::EvaluateSwept(const FMovieSceneEvaluationO
 
 	TArray<FMovieSceneEventData> Events;
 
-	TRange<FFrameNumber>            SweptRange = Context.GetFrameNumberRange();
-
 	TArrayView<const FFrameNumber>  KeyTimes   = EventData.GetKeyTimes();
 	TArrayView<const FEventPayload> KeyValues  = EventData.GetKeyValues();
 
@@ -423,7 +421,7 @@ FMovieSceneEventTriggerTemplate::FMovieSceneEventTriggerTemplate(const UMovieSce
 	}
 }
 
-void FMovieSceneEventTriggerTemplate::EvaluateSwept(const FMovieSceneEvaluationOperand& Operand, const FMovieSceneContext& Context, const FPersistentEvaluationData& PersistentData, FMovieSceneExecutionTokens& ExecutionTokens) const
+void FMovieSceneEventTriggerTemplate::EvaluateSwept(const FMovieSceneEvaluationOperand& Operand, const FMovieSceneContext& Context, const TRange<FFrameNumber>& SweptRange, const FPersistentEvaluationData& PersistentData, FMovieSceneExecutionTokens& ExecutionTokens) const
 {
 	// Don't allow events to fire when playback is in a stopped state. This can occur when stopping 
 	// playback and returning the current position to the start of playback. It's not desireable to have 
@@ -443,8 +441,6 @@ void FMovieSceneEventTriggerTemplate::EvaluateSwept(const FMovieSceneEvaluationO
 
 	TArray<FName> Events;
 
-	TRange<FFrameNumber> SweptRange = Context.GetFrameNumberRange();
-
 	const int32 First = bBackwards ? EventTimes.Num() - 1 : 0;
 	const int32 Last = bBackwards ? 0 : EventTimes.Num() - 1;
 	const int32 Inc = bBackwards ? -1 : 1;
@@ -488,12 +484,12 @@ FMovieSceneEventRepeaterTemplate::FMovieSceneEventRepeaterTemplate(const UMovieS
 {
 }
 
-void FMovieSceneEventRepeaterTemplate::EvaluateSwept(const FMovieSceneEvaluationOperand& Operand, const FMovieSceneContext& Context, const FPersistentEvaluationData& PersistentData, FMovieSceneExecutionTokens& ExecutionTokens) const
+void FMovieSceneEventRepeaterTemplate::EvaluateSwept(const FMovieSceneEvaluationOperand& Operand, const FMovieSceneContext& Context, const TRange<FFrameNumber>& SweptRange, const FPersistentEvaluationData& PersistentData, FMovieSceneExecutionTokens& ExecutionTokens) const
 {
 	// Don't allow events to fire when playback is in a stopped state. This can occur when stopping 
 	// playback and returning the current position to the start of playback. It's not desireable to have 
 	// all the events from the last playback position to the start of playback be fired.
-	if (Context.GetStatus() == EMovieScenePlayerStatus::Stopped || Context.IsSilent())
+	if (!SweptRange.Contains(Context.GetTime().FrameNumber) || Context.GetStatus() == EMovieScenePlayerStatus::Stopped || Context.IsSilent())
 	{
 		return;
 	}
diff --git a/Engine/Source/Runtime/MovieSceneTracks/Private/Evaluation/MovieSceneEventTemplate.h b/Engine/Source/Runtime/MovieSceneTracks/Private/Evaluation/MovieSceneEventTemplate.h
--- a/Engine/Source/Runtime/MovieSceneTracks/Private/Evaluation/MovieSceneEventTemplate.h
+++ b/Engine/Source/Runtime/MovieSceneTracks/Private/Evaluation/MovieSceneEventTemplate.h
@@ -51,7 +51,7 @@ struct FMovieSceneEventSectionTemplate : public FMovieSceneEventTemplateBase
 private:
 
 	virtual UScriptStruct& GetScriptStructImpl() const override { return *StaticStruct(); }
-	virtual void EvaluateSwept(const FMovieSceneEvaluationOperand& Operand, const FMovieSceneContext& Context, const FPersistentEvaluationData& PersistentData, FMovieSceneExecutionTokens& ExecutionTokens) const override;
+	virtual void EvaluateSwept(const FMovieSceneEvaluationOperand& Operand, const FMovieSceneContext& Context, const TRange<FFrameNumber>& SweptRange, const FPersistentEvaluationData& PersistentData, FMovieSceneExecutionTokens& ExecutionTokens) const override;
 };
 
 
@@ -72,7 +72,7 @@ struct FMovieSceneEventTriggerTemplate : public FMovieSceneEventTemplateBase
 private:
 
 	virtual UScriptStruct& GetScriptStructImpl() const override { return *StaticStruct(); }
-	virtual void EvaluateSwept(const FMovieSceneEvaluationOperand& Operand, const FMovieSceneContext& Context, const FPersistentEvaluationData& PersistentData, FMovieSceneExecutionTokens& ExecutionTokens) const override;
+	virtual void EvaluateSwept(const FMovieSceneEvaluationOperand& Operand, const FMovieSceneContext& Context, const TRange<FFrameNumber>& SweptRange, const FPersistentEvaluationData& PersistentData, FMovieSceneExecutionTokens& ExecutionTokens) const override;
 };
 
 USTRUCT()
@@ -89,5 +89,5 @@ struct FMovieSceneEventRepeaterTemplate : public FMovieSceneEventTemplateBase
 private:
 
 	virtual UScriptStruct& GetScriptStructImpl() const override { return *StaticStruct(); }
-	virtual void EvaluateSwept(const FMovieSceneEvaluationOperand& Operand, const FMovieSceneContext& Context, const FPersistentEvaluationData& PersistentData, FMovieSceneExecutionTokens& ExecutionTokens) const override;
+	virtual void EvaluateSwept(const FMovieSceneEvaluationOperand& Operand, const FMovieSceneContext& Context, const TRange<FFrameNumber>& SweptRange, const FPersistentEvaluationData& PersistentData, FMovieSceneExecutionTokens& ExecutionTokens) const override;
 };
