Sequencer: Force tick if a new instance is bound

This fixes another first frame issue with warmup frames when movie rendering. If we change to tick sequencer first always in its own tick group, this can go away.

diff --git a/Engine/Source/Runtime/MovieSceneTracks/Private/Evaluation/MovieSceneSkeletalAnimationTemplate.cpp b/Engine/Source/Runtime/MovieSceneTracks/Private/Evaluation/MovieSceneSkeletalAnimationTemplate.cpp
--- a/Engine/Source/Runtime/MovieSceneTracks/Private/Evaluation/MovieSceneSkeletalAnimationTemplate.cpp
+++ b/Engine/Source/Runtime/MovieSceneTracks/Private/Evaluation/MovieSceneSkeletalAnimationTemplate.cpp
@@ -191,7 +191,9 @@ namespace MovieScene
 
 			OriginalStack.SavePreAnimatedState(Player, *SkeletalMeshComponent, GetAnimControlTypeID(), FPreAnimatedAnimationTokenProducer());
 
-			UAnimCustomInstance::BindToSkeletalMeshComponent<UAnimSequencerInstance>(SkeletalMeshComponent);
+			const UAnimInstance* ExistingAnimInstance = SkeletalMeshComponent->GetAnimInstance();
+
+			const UAnimSequencerInstance* SequencerInstance = UAnimCustomInstance::BindToSkeletalMeshComponent<UAnimSequencerInstance>(SkeletalMeshComponent);
 
 			const bool bPreviewPlayback = ShouldUsePreviewPlayback(Player, *SkeletalMeshComponent);
 
@@ -228,10 +230,10 @@ namespace MovieScene
 				}
 			}
 
-			// If the skeletal component has already ticked this frame because tick prerequisites weren't set up yet, forcibly tick this component to update.
+			// If the skeletal component has already ticked this frame because tick prerequisites weren't set up yet or a new binding was created, forcibly tick this component to update.
 			// This resolves first frame issues where the skeletal component ticks first, then the sequencer binding is resolved which sets up tick prerequisites
 			// for the next frame.
-			if (SkeletalMeshComponent->PoseTickedThisFrame())
+			if (SkeletalMeshComponent->PoseTickedThisFrame() || SequencerInstance != ExistingAnimInstance)
 			{
 				SkeletalMeshComponent->TickAnimation(0.f, false);
 
