Fix integration from dev-sequencer on various platforms  with no logging.
Looks like we needed o wrap #if !NO_LOGGING around where GenerateEventListingString is used since it's not defined.

diff --git a/Engine/Source/Runtime/MovieSceneTracks/Private/Evaluation/MovieSceneEventTemplate.cpp b/Engine/Source/Runtime/MovieSceneTracks/Private/Evaluation/MovieSceneEventTemplate.cpp
--- a/Engine/Source/Runtime/MovieSceneTracks/Private/Evaluation/MovieSceneEventTemplate.cpp
+++ b/Engine/Source/Runtime/MovieSceneTracks/Private/Evaluation/MovieSceneEventTemplate.cpp
@@ -175,7 +175,9 @@ struct FEventTriggerExecutionToken
 		UObject* DirectorInstance = Player.GetEvaluationTemplate().GetOrCreateDirectorInstance(Operand.SequenceID, Player);
 		if (!DirectorInstance)
 		{
+#if !NO_LOGGING
 			UE_LOG(LogMovieScene, Warning, TEXT("Failed to trigger the following events because no director instance was available: %s."), *GenerateEventListString());
+#endif			
 			return;
 		}
 
