Sequencer: Fixed legacy upgrade for events not using the correct graph GUID due to uobject load order dependency

  - Since widget animations live inside a widget blueprint type, the graph that an event is bound to may not be loaded yet (still has RF_NeedsLoad), which means that the GraphGuid is not necessarily serialized yet.
  - Instead of performing the upgrade on load we now do the conversion on save where the object is bound to be valid.

diff --git a/Engine/Source/Runtime/MovieSceneTracks/Private/Channels/MovieSceneEvent.cpp b/Engine/Source/Runtime/MovieSceneTracks/Private/Channels/MovieSceneEvent.cpp
--- a/Engine/Source/Runtime/MovieSceneTracks/Private/Channels/MovieSceneEvent.cpp
+++ b/Engine/Source/Runtime/MovieSceneTracks/Private/Channels/MovieSceneEvent.cpp
@@ -11,6 +11,16 @@
 
 UK2Node_FunctionEntry* FMovieSceneEvent::GetFunctionEntry() const
 {
+	if (SoftBlueprintPath.IsNull())
+	{
+		// The function entry used to be serialized but is now only stored transiently. We use this pointer for the current lifecycle until the asset is saved, when we do the data upgrade.
+		UK2Node_FunctionEntry* FunctionEntryPtr = CastChecked<UK2Node_FunctionEntry>(FunctionEntry_DEPRECATED.Get(), ECastCheckedType::NullAllowed);
+		if (FunctionEntryPtr)
+		{
+			return FunctionEntryPtr;
+		}
+	}
+
 	UK2Node_FunctionEntry* Cached = CastChecked<UK2Node_FunctionEntry>(CachedFunctionEntry.Get(), ECastCheckedType::NullAllowed);
 	if (Cached)
 	{
@@ -105,9 +115,22 @@ void FMovieSceneEvent::PostSerialize(const FArchive& Ar)
 {
 #if WITH_EDITORONLY_DATA
 	if (Ar.IsLoading() && !Ar.HasAnyPortFlags(PPF_Duplicate | PPF_DuplicateForPIE))
+	{
+		// Re-cache the function name when loading in-editor in case of renamed function graphs and the like
+		CacheFunctionName();
+	}
+#endif
+}
+
+bool FMovieSceneEvent::Serialize(FArchive& Ar)
+{
+#if WITH_EDITORONLY_DATA
+	if (Ar.IsSaving())
 	{
 		// ---------------------------------------------------------------------------------------
 		// Data upgrade for content that was saved with FunctionEntry_DEPRECATED instead of SoftFunctionGraph
+		// We do this on save because there is no reliable way to ensure that FunctionGraph is fully loaded here
+		// since the track may live inside or outside of a blueprint. When not fully loaded, GraphGuid is not correct.
 		if (SoftBlueprintPath.IsNull())
 		{
 			// The function entry used to be serialized but is now only stored transiently. If this is set without the soft function graph being set, copy the graph reference over
@@ -121,11 +144,11 @@ void FMovieSceneEvent::PostSerialize(const FArchive& Ar)
 				GraphGuid           = FunctionGraph->GraphGuid;
 			}
 		}
-
-		// Re-cache the function name when loading in-editor in case of renamed function graphs and the like
-		CacheFunctionName();
 	}
 #endif
+
+	// Return false to ensure that the struct receives default serialization
+	return false;
 }
 
 bool FMovieSceneEvent::IsValidFunction(UFunction* Function)
diff --git a/Engine/Source/Runtime/MovieSceneTracks/Public/Channels/MovieSceneEvent.h b/Engine/Source/Runtime/MovieSceneTracks/Public/Channels/MovieSceneEvent.h
--- a/Engine/Source/Runtime/MovieSceneTracks/Public/Channels/MovieSceneEvent.h
+++ b/Engine/Source/Runtime/MovieSceneTracks/Public/Channels/MovieSceneEvent.h
@@ -38,6 +38,10 @@ struct MOVIESCENETRACKS_API FMovieSceneEvent
 	 */
 	void PostSerialize(const FArchive& Ar);
 
+	/**
+	 * Called to perform custom serialization logic for this struct.
+	 */
+	bool Serialize(FArchive& Ar);
 
 	/**
 	 * Check whether the specified function is valid for a movie scene event
@@ -122,6 +126,6 @@ private:
 template<>
 struct TStructOpsTypeTraits<FMovieSceneEvent> : TStructOpsTypeTraitsBase2<FMovieSceneEvent>
 {
-	enum { WithPostSerialize = true };
+	enum { WithSerializer = true, WithPostSerialize = true };
 };
 
