Sequencer: Allow attach and path sections to support infinite ranges. This fixes an upgrade issue where existing sections would not load as infinite.

diff --git a/Engine/Source/Runtime/MovieSceneTracks/Private/Sections/MovieScene3DConstraintSection.cpp b/Engine/Source/Runtime/MovieSceneTracks/Private/Sections/MovieScene3DConstraintSection.cpp
--- a/Engine/Source/Runtime/MovieSceneTracks/Private/Sections/MovieScene3DConstraintSection.cpp
+++ b/Engine/Source/Runtime/MovieSceneTracks/Private/Sections/MovieScene3DConstraintSection.cpp
@@ -6,7 +6,9 @@
 
 UMovieScene3DConstraintSection::UMovieScene3DConstraintSection( const FObjectInitializer& ObjectInitializer )
 	: Super( ObjectInitializer )
-{ }
+{ 
+	bSupportsInfiniteRange = true;
+}
 
 void UMovieScene3DConstraintSection::SetConstraintId(const FGuid& InConstraintId, const FMovieSceneSequenceID& SequenceID)
 {
