Sequencer: Fix audio evaluation from crashing when there's no sound

diff --git a/Engine/Source/Runtime/MovieSceneTracks/Private/Evaluation/MovieSceneAudioTemplate.cpp b/Engine/Source/Runtime/MovieSceneTracks/Private/Evaluation/MovieSceneAudioTemplate.cpp
--- a/Engine/Source/Runtime/MovieSceneTracks/Private/Evaluation/MovieSceneAudioTemplate.cpp
+++ b/Engine/Source/Runtime/MovieSceneTracks/Private/Evaluation/MovieSceneAudioTemplate.cpp
@@ -346,7 +346,7 @@ void FMovieSceneAudioSectionTemplateData::EnsureAudioIsPlaying(UAudioComponent&
 		}
 
 		const float AudioTime = (Context.GetTime() / Context.GetFrameRate()) - SectionStartTimeSeconds + (float)Context.GetFrameRate().AsSeconds(AudioStartOffset);
-		if (AudioTime >= 0.f && AudioTime < AudioComponent.Sound->GetDuration())
+		if (AudioTime >= 0.f && AudioComponent.Sound && AudioTime < AudioComponent.Sound->GetDuration())
 		{
 			AudioComponent.Play(AudioTime);
 		}
