Sequencer: Added missing modify to camera anim/shake track

diff --git a/Engine/Source/Runtime/MovieSceneTracks/Private/Tracks/MovieSceneCameraAnimTrack.cpp b/Engine/Source/Runtime/MovieSceneTracks/Private/Tracks/MovieSceneCameraAnimTrack.cpp
--- a/Engine/Source/Runtime/MovieSceneTracks/Private/Tracks/MovieSceneCameraAnimTrack.cpp
+++ b/Engine/Source/Runtime/MovieSceneTracks/Private/Tracks/MovieSceneCameraAnimTrack.cpp
@@ -16,6 +16,8 @@
 
 UMovieSceneSection* UMovieSceneCameraAnimTrack::AddNewCameraAnim(FFrameNumber KeyTime, UCameraAnim* CameraAnim)
 {
+	Modify();
+
 	UMovieSceneCameraAnimSection* const NewSection = Cast<UMovieSceneCameraAnimSection>(CreateNewSection());
 	if (NewSection)
 	{
diff --git a/Engine/Source/Runtime/MovieSceneTracks/Private/Tracks/MovieSceneCameraShakeTrack.cpp b/Engine/Source/Runtime/MovieSceneTracks/Private/Tracks/MovieSceneCameraShakeTrack.cpp
--- a/Engine/Source/Runtime/MovieSceneTracks/Private/Tracks/MovieSceneCameraShakeTrack.cpp
+++ b/Engine/Source/Runtime/MovieSceneTracks/Private/Tracks/MovieSceneCameraShakeTrack.cpp
@@ -14,6 +14,8 @@
 
 UMovieSceneSection* UMovieSceneCameraShakeTrack::AddNewCameraShake(FFrameNumber KeyTime, TSubclassOf<UCameraShake> ShakeClass)
 {
+	Modify();
+
 	UMovieSceneCameraShakeSection* const NewSection = Cast<UMovieSceneCameraShakeSection>(CreateNewSection());
 	if (NewSection)
 	{
