Fix keying rotation values > 180.
Changed interrogation type internally to FTransformData to get the raw Rotator instead of one corrupted by FTransform->FRotator conversion.

diff --git a/Engine/Source/Editor/MovieSceneTools/Private/TrackEditors/TransformTrackEditor.cpp b/Engine/Source/Editor/MovieSceneTools/Private/TrackEditors/TransformTrackEditor.cpp
--- a/Engine/Source/Editor/MovieSceneTools/Private/TrackEditors/TransformTrackEditor.cpp
+++ b/Engine/Source/Editor/MovieSceneTools/Private/TrackEditors/TransformTrackEditor.cpp
@@ -852,11 +852,11 @@ bool F3DTransformTrackEditor::ModifyGeneratedKeysByCurrentAndWeight(UObject *Obj
 			
 	FVector CurrentPos; FRotator CurrentRot;
 	FVector CurrentScale;
-	for (const FTransform& Transform : InterrogationData.Iterate<FTransform>(UMovieScene3DTransformSection::GetInterrogationKey()))
+	for (const FTransformData& Transform : InterrogationData.Iterate<FTransformData>(UMovieScene3DTransformSection::GetInterrogationKey()))
 	{
-		CurrentPos = Transform.GetTranslation();
-		CurrentRot = Transform.Rotator();
-		CurrentScale = Transform.GetScale3D();
+		CurrentPos = Transform.Translation;
+		CurrentRot = Transform.Rotation;
+		CurrentScale = Transform.Scale;
 		break;
 	}
 	FMovieSceneChannelProxy& Proxy = SectionToKey->GetChannelProxy();
diff --git a/Engine/Source/Editor/Sequencer/Private/SequencerEdMode.cpp b/Engine/Source/Editor/Sequencer/Private/SequencerEdMode.cpp
--- a/Engine/Source/Editor/Sequencer/Private/SequencerEdMode.cpp
+++ b/Engine/Source/Editor/Sequencer/Private/SequencerEdMode.cpp
@@ -472,10 +472,10 @@ void FSequencerEdMode::GetLocationAtTime(FMovieSceneEvaluationTrack* Track, UObj
 	FMovieSceneContext Context(FMovieSceneEvaluationRange(KeyTime, Sequencer->GetFocusedTickResolution()));
 	Track->Interrogate(Context, InterrogationData, Object);
 
-	for (const FTransform& Transform : InterrogationData.Iterate<FTransform>(UMovieScene3DTransformSection::GetInterrogationKey()))
+	for (const FTransformData& Transform : InterrogationData.Iterate<FTransformData>(UMovieScene3DTransformSection::GetInterrogationKey()))
 	{
-		KeyPos = Transform.GetTranslation();
-		KeyRot = Transform.GetRotation().Rotator();
+		KeyPos = Transform.Translation;
+		KeyRot = Transform.Rotation;
 		break;
 	}
 }
diff --git a/Engine/Source/Runtime/MovieSceneTracks/Private/Evaluation/MovieScene3DTransformTemplate.cpp b/Engine/Source/Runtime/MovieSceneTracks/Private/Evaluation/MovieScene3DTransformTemplate.cpp
--- a/Engine/Source/Runtime/MovieSceneTracks/Private/Evaluation/MovieScene3DTransformTemplate.cpp
+++ b/Engine/Source/Runtime/MovieSceneTracks/Private/Evaluation/MovieScene3DTransformTemplate.cpp
@@ -92,7 +92,11 @@ struct FComponentTransformActuator : TMovieSceneBlendingActuator<F3DTransformTra
 
 	virtual void Actuate(FMovieSceneInterrogationData& InterrogationData, const F3DTransformTrackToken& InValue, const TBlendableTokenStack<F3DTransformTrackToken>& OriginalStack, const FMovieSceneContext& Context) const override
 	{
-		InterrogationData.Add(FTransform(InValue.Rotation.Quaternion(), InValue.Translation, InValue.Scale), UMovieScene3DTransformSection::GetInterrogationKey());
+		FTransformData Data;
+		Data.Translation = InValue.Translation;
+		Data.Rotation = InValue.Rotation;
+		Data.Scale = InValue.Scale;
+		InterrogationData.Add(FTransformData(Data), UMovieScene3DTransformSection::GetInterrogationKey());
 	}
 };
 
diff --git a/Engine/Source/Runtime/MovieSceneTracks/Private/Sections/MovieScene3DTransformSection.cpp b/Engine/Source/Runtime/MovieSceneTracks/Private/Sections/MovieScene3DTransformSection.cpp
--- a/Engine/Source/Runtime/MovieSceneTracks/Private/Sections/MovieScene3DTransformSection.cpp
+++ b/Engine/Source/Runtime/MovieSceneTracks/Private/Sections/MovieScene3DTransformSection.cpp
@@ -215,11 +215,11 @@ struct F3DTransformChannelEditorData
 
 		FVector CurrentPos; FRotator CurrentRot;
 		FVector CurrentScale;
-		for (const FTransform& Transform : InterrogationData.Iterate<FTransform>(UMovieScene3DTransformSection::GetInterrogationKey()))
+		for (const FTransformData& Transform : InterrogationData.Iterate<FTransformData>(UMovieScene3DTransformSection::GetInterrogationKey()))
 		{
-			CurrentPos = Transform.GetTranslation();
-			CurrentRot = Transform.Rotator();
-			CurrentScale = Transform.GetScale3D();
+			CurrentPos = Transform.Translation;
+			CurrentRot = Transform.Rotation;
+			CurrentScale = Transform.Scale;
 			break;
 		}
 
@@ -719,4 +719,4 @@ FMovieSceneInterrogationKey UMovieScene3DTransformSection::GetInterrogationKey()
 {
 	static FMovieSceneAnimTypeID TypeID = FMovieSceneAnimTypeID::Unique();
 	return TypeID;
-}
\ No newline at end of file
+}
diff --git a/Engine/Source/Runtime/MovieSceneTracks/Private/Sections/MovieSceneCameraCutSection.cpp b/Engine/Source/Runtime/MovieSceneTracks/Private/Sections/MovieSceneCameraCutSection.cpp
--- a/Engine/Source/Runtime/MovieSceneTracks/Private/Sections/MovieSceneCameraCutSection.cpp
+++ b/Engine/Source/Runtime/MovieSceneTracks/Private/Sections/MovieSceneCameraCutSection.cpp
@@ -39,9 +39,9 @@ FMovieSceneEvalTemplatePtr UMovieSceneCameraCutSection::GenerateTemplate() const
 					FMovieSceneInterrogationData Container;
 					TransformTrackTemplate.Interrogate(Context, Container);
 
-					for (auto It = Container.Iterate<FTransform>(UMovieScene3DTransformSection::GetInterrogationKey()); It; ++It)
+					for (const FTransformData& Transform : Container.Iterate<FTransformData>(UMovieScene3DTransformSection::GetInterrogationKey()))
 					{
-						CutTransform = *It;
+						CutTransform = FTransform(Transform.Rotation.Quaternion(), Transform.Translation, Transform.Scale);
 						break;
 					}
 				}
