Fix to underterministic cooking of some WidgetBlueprints.

diff --git a/Engine/Source/Runtime/MovieSceneTracks/Private/Channels/MovieSceneEvent.cpp b/Engine/Source/Runtime/MovieSceneTracks/Private/Channels/MovieSceneEvent.cpp
--- a/Engine/Source/Runtime/MovieSceneTracks/Private/Channels/MovieSceneEvent.cpp
+++ b/Engine/Source/Runtime/MovieSceneTracks/Private/Channels/MovieSceneEvent.cpp
@@ -51,6 +51,22 @@ UK2Node_FunctionEntry* FMovieSceneEvent::GetFunctionEntry() const
 	return nullptr;
 }
 
+bool FMovieSceneEvent::IsReadyToGetFunctionEntry()
+{
+	// We cannot get the function entry until all graphs in the blueprint are loaded.
+	if (UBlueprint* Blueprint = SoftBlueprintPath.LoadSynchronous())
+	{
+		for (UEdGraph* Graph : Blueprint->FunctionGraphs)
+		{
+			if (Graph->HasAnyFlags(RF_NeedLoad))
+			{
+				return false;
+			}
+		}
+	}
+	return true;
+}
+
 void FMovieSceneEvent::SetFunctionEntry(UK2Node_FunctionEntry* InFunctionEntry)
 {
 	UEdGraph* EdGraph = InFunctionEntry ? InFunctionEntry->GetGraph() : nullptr;
@@ -116,8 +132,14 @@ void FMovieSceneEvent::PostSerialize(const FArchive& Ar)
 #if WITH_EDITORONLY_DATA
 	if (Ar.IsLoading() && !Ar.HasAnyPortFlags(PPF_Duplicate | PPF_DuplicateForPIE))
 	{
-		// Re-cache the function name when loading in-editor in case of renamed function graphs and the like
-		CacheFunctionName();
+		// Re-cache the function name when loading in-editor in case of renamed function graphs and the like 
+		// We can only do it safely if the blueprint and Graph are done loading.
+		// Without this check the call will reset the function name to null and subsequent call to it will mark the object as Changed even if it do not. 
+		// causing Non deterministic cooking.
+		if (IsReadyToGetFunctionEntry())
+		{
+			CacheFunctionName();
+		}
 	}
 #endif
 }
diff --git a/Engine/Source/Runtime/MovieSceneTracks/Public/Channels/MovieSceneEvent.h b/Engine/Source/Runtime/MovieSceneTracks/Public/Channels/MovieSceneEvent.h
--- a/Engine/Source/Runtime/MovieSceneTracks/Public/Channels/MovieSceneEvent.h
+++ b/Engine/Source/Runtime/MovieSceneTracks/Public/Channels/MovieSceneEvent.h
@@ -104,6 +104,11 @@ public:
 
 private:
 
+	/**
+	* Validates that all Graphs in the blueprint are loaded.
+	*/
+	bool IsReadyToGetFunctionEntry();
+
 	/** Serialized soft pointer to the blueprint that contains the function graph endpoint for this event. Stored as a soft path so that renames of the blueprint don't break this event binding. */
 	UPROPERTY()
 	TSoftObjectPtr<UBlueprint> SoftBlueprintPath;
