Fix missing == conversion resulting in implicit cast
#rnx

diff --git a/Engine/Source/Runtime/MovieSceneTracks/Private/Tracks/MovieScenePropertyTrack.cpp b/Engine/Source/Runtime/MovieSceneTracks/Private/Tracks/MovieScenePropertyTrack.cpp
--- a/Engine/Source/Runtime/MovieSceneTracks/Private/Tracks/MovieScenePropertyTrack.cpp
+++ b/Engine/Source/Runtime/MovieSceneTracks/Private/Tracks/MovieScenePropertyTrack.cpp
@@ -94,7 +94,7 @@ void UMovieScenePropertyTrack::RemoveSection(UMovieSceneSection& Section)
 
 void UMovieScenePropertyTrack::RemoveSectionAt(int32 SectionIndex)
 {
-	bool bResetSectionToKey = SectionToKey = Sections[SectionIndex];
+	bool bResetSectionToKey = (SectionToKey == Sections[SectionIndex]);
 
 	Sections.RemoveAt(SectionIndex);
 
