Sequencer: Tick and refresh on restore after unbinding and setting animation mode to force the restored pose

Had to manually turn on SetUpdateAnimationInEditor(true) before TickAnimation as well.

diff --git a/Engine/Source/Runtime/MovieSceneTracks/Private/Evaluation/MovieSceneSkeletalAnimationTemplate.cpp b/Engine/Source/Runtime/MovieSceneTracks/Private/Evaluation/MovieSceneSkeletalAnimationTemplate.cpp
--- a/Engine/Source/Runtime/MovieSceneTracks/Private/Evaluation/MovieSceneSkeletalAnimationTemplate.cpp
+++ b/Engine/Source/Runtime/MovieSceneTracks/Private/Evaluation/MovieSceneSkeletalAnimationTemplate.cpp
@@ -98,14 +98,6 @@ struct FPreAnimatedAnimationTokenProducer : IMovieScenePreAnimatedTokenProducer
 				if (SequencerInst)
 				{
 					SequencerInst->ResetPose();
-					//wish this 'force evaluation' was encapsulated better
-					Component->TickAnimation(0.f, false);
-					Component->RefreshBoneTransforms();
-					Component->RefreshSlaveComponents();
-					Component->UpdateComponentToWorld();
-					Component->FinalizeBoneTransform();
-					Component->MarkRenderTransformDirty();
-					Component->MarkRenderDynamicDataDirty();
 					SequencerInst->ResetNodes();
 				}
 
@@ -123,6 +115,18 @@ struct FPreAnimatedAnimationTokenProducer : IMovieScenePreAnimatedTokenProducer
 					Component->AnimScriptInstance = CachedAnimInstance.Get();
 					CachedAnimInstance.Reset();
 				}
+
+				// Restore pose after unbinding to force the restored pose
+				Component->SetUpdateAnimationInEditor(true);
+				Component->TickAnimation(0.f, false);
+
+				Component->RefreshBoneTransforms();
+				Component->RefreshSlaveComponents();
+				Component->UpdateComponentToWorld();
+				Component->FinalizeBoneTransform();
+				Component->MarkRenderTransformDirty();
+				Component->MarkRenderDynamicDataDirty();
+				
 #if WITH_EDITOR
 				Component->SetUpdateAnimationInEditor(bUpdateAnimationInEditor);
 #endif
