Use new (and proper) AudioComponent API for playing a short duration sound while scrubbing.

diff --git a/Engine/Source/Runtime/MovieSceneTracks/Private/Evaluation/MovieSceneAudioTemplate.cpp b/Engine/Source/Runtime/MovieSceneTracks/Private/Evaluation/MovieSceneAudioTemplate.cpp
--- a/Engine/Source/Runtime/MovieSceneTracks/Private/Evaluation/MovieSceneAudioTemplate.cpp
+++ b/Engine/Source/Runtime/MovieSceneTracks/Private/Evaluation/MovieSceneAudioTemplate.cpp
@@ -353,9 +353,8 @@ struct FAudioSectionExecutionToken : IMovieSceneExecutionToken
 
 			if (Context.GetStatus() == EMovieScenePlayerStatus::Scrubbing)
 			{
-				// Fade out the sound at the same volume in order to simply
-				// set a short duration on the sound, far from ideal soln
-				AudioComponent.FadeOut(AudioTrackConstants::ScrubDuration, 1.f);
+				// While scrubbing, play the sound for a short time and then cut it.
+				AudioComponent.StopDelayed(AudioTrackConstants::ScrubDuration);
 			}
 		}
 
