CIS Fix
#trival

diff --git a/Engine/Source/Runtime/MovieSceneTracks/Private/Evaluation/MovieSceneSkeletalAnimationTemplate.cpp b/Engine/Source/Runtime/MovieSceneTracks/Private/Evaluation/MovieSceneSkeletalAnimationTemplate.cpp
--- a/Engine/Source/Runtime/MovieSceneTracks/Private/Evaluation/MovieSceneSkeletalAnimationTemplate.cpp
+++ b/Engine/Source/Runtime/MovieSceneTracks/Private/Evaluation/MovieSceneSkeletalAnimationTemplate.cpp
@@ -238,7 +238,7 @@ namespace MovieScene
 			UAnimInstance* ExistingAnimInstance = SkeletalMeshComponent->GetAnimInstance();
 			bool bWasCreated = false;
 			UAnimSequencerInstance* SequencerInstance = UAnimCustomInstance::BindToSkeletalMeshComponent<UAnimSequencerInstance>(SkeletalMeshComponent,bWasCreated);
-			if (bWasCreated)
+			if (SequencerInstance && bWasCreated)
 			{
 				SequencerInstance->SavePose();
 			}
