Sequencer: Fix skeletal reverse not using the complete sequence length to compute the reversed time.

diff --git a/Engine/Source/Runtime/MovieSceneTracks/Private/Evaluation/MovieSceneSkeletalAnimationTemplate.cpp b/Engine/Source/Runtime/MovieSceneTracks/Private/Evaluation/MovieSceneSkeletalAnimationTemplate.cpp
--- a/Engine/Source/Runtime/MovieSceneTracks/Private/Evaluation/MovieSceneSkeletalAnimationTemplate.cpp
+++ b/Engine/Source/Runtime/MovieSceneTracks/Private/Evaluation/MovieSceneSkeletalAnimationTemplate.cpp
@@ -550,7 +550,7 @@ float FMovieSceneSkeletalAnimationSectionTemplateParameters::MapTimeToAnimation(
 	AnimPosition += InFrameRate.AsSeconds(StartFrameOffset);
 	if (bReverse)
 	{
-		AnimPosition = (SeqLength - (AnimPosition - InFrameRate.AsSeconds(StartFrameOffset))) + InFrameRate.AsSeconds(StartFrameOffset);
+		AnimPosition = GetSequenceLength() - AnimPosition;
 	}
 
 	return AnimPosition;
