Sequencer: Fix missing UPROPERTY() for deprecated StartOffset

diff --git a/Engine/Source/Runtime/MovieSceneTracks/Public/Sections/MovieSceneAudioSection.h b/Engine/Source/Runtime/MovieSceneTracks/Public/Sections/MovieSceneAudioSection.h
--- a/Engine/Source/Runtime/MovieSceneTracks/Public/Sections/MovieSceneAudioSection.h
+++ b/Engine/Source/Runtime/MovieSceneTracks/Public/Sections/MovieSceneAudioSection.h
@@ -160,6 +160,7 @@ private:
 	FFrameNumber StartFrameOffset;
 
 	/** The offset into the beginning of the audio clip */
+	UPROPERTY()
 	float StartOffset_DEPRECATED;
 
 	/** The absolute time that the sound starts playing at */
