Sequencer: Stop all sounds on tear down

diff --git a/Engine/Source/Runtime/MovieSceneTracks/Private/Evaluation/MovieSceneAudioTemplate.cpp b/Engine/Source/Runtime/MovieSceneTracks/Private/Evaluation/MovieSceneAudioTemplate.cpp
--- a/Engine/Source/Runtime/MovieSceneTracks/Private/Evaluation/MovieSceneAudioTemplate.cpp
+++ b/Engine/Source/Runtime/MovieSceneTracks/Private/Evaluation/MovieSceneAudioTemplate.cpp
@@ -381,3 +381,15 @@ void FMovieSceneAudioSectionTemplate::Evaluate(const FMovieSceneEvaluationOperan
 		ExecutionTokens.Add(FAudioSectionExecutionToken(AudioData));
 	}
 }
+
+void FMovieSceneAudioSectionTemplate::TearDown(FPersistentEvaluationData& PersistentData, IMovieScenePlayer& Player) const
+{
+	MOVIESCENE_DETAILED_SCOPE_CYCLE_COUNTER(MovieSceneEval_AudioTrack_Teardown)
+
+	if (GEngine && GEngine->UseSound())
+	{
+		FCachedAudioTrackData& TrackData = PersistentData.GetOrAddTrackData<FCachedAudioTrackData>();
+
+		TrackData.StopAllSounds();
+	}
+}
\ No newline at end of file
diff --git a/Engine/Source/Runtime/MovieSceneTracks/Private/Evaluation/MovieSceneAudioTemplate.h b/Engine/Source/Runtime/MovieSceneTracks/Private/Evaluation/MovieSceneAudioTemplate.h
--- a/Engine/Source/Runtime/MovieSceneTracks/Private/Evaluation/MovieSceneAudioTemplate.h
+++ b/Engine/Source/Runtime/MovieSceneTracks/Private/Evaluation/MovieSceneAudioTemplate.h
@@ -84,4 +84,9 @@ private:
 
 	virtual UScriptStruct& GetScriptStructImpl() const override { return *StaticStruct(); }
 	virtual void Evaluate(const FMovieSceneEvaluationOperand& Operand, const FMovieSceneContext& Context, const FPersistentEvaluationData& PersistentData, FMovieSceneExecutionTokens& ExecutionTokens) const override;
+	virtual void SetupOverrides() override
+	{
+		EnableOverrides(RequiresTearDownFlag);
+	}
+	virtual void TearDown(FPersistentEvaluationData& PersistentData, IMovieScenePlayer& Player) const override;
 };
