Sequencer: Fix deprecation

diff --git a/Engine/Source/Runtime/MovieSceneTracks/Private/Sections/MovieSceneAudioSection.cpp b/Engine/Source/Runtime/MovieSceneTracks/Private/Sections/MovieSceneAudioSection.cpp
--- a/Engine/Source/Runtime/MovieSceneTracks/Private/Sections/MovieSceneAudioSection.cpp
+++ b/Engine/Source/Runtime/MovieSceneTracks/Private/Sections/MovieSceneAudioSection.cpp
@@ -23,14 +23,19 @@ struct FAudioChannelEditorData
 
 #endif // WITH_EDITOR
 
+namespace
+{
+	float AudioDeprecatedMagicNumber = TNumericLimits<float>::Lowest();
+}
+
 UMovieSceneAudioSection::UMovieSceneAudioSection( const FObjectInitializer& ObjectInitializer )
 	: Super( ObjectInitializer )
 {
 	Sound = nullptr;
 	StartOffset = 0.f;
-	AudioStartTime_DEPRECATED = 0.f;
-	AudioDilationFactor_DEPRECATED = 1.f;
-	AudioVolume_DEPRECATED = 1.f;
+	AudioStartTime_DEPRECATED = AudioDeprecatedMagicNumber;
+	AudioDilationFactor_DEPRECATED = AudioDeprecatedMagicNumber;
+	AudioVolume_DEPRECATED = AudioDeprecatedMagicNumber;
 	bSuppressSubtitles = false;
 	bOverrideAttenuation = false;
 
@@ -75,21 +80,21 @@ void UMovieSceneAudioSection::PostLoad()
 {
 	Super::PostLoad();
 
-	if (AudioDilationFactor_DEPRECATED != FLT_MAX)
+	if (AudioDilationFactor_DEPRECATED != AudioDeprecatedMagicNumber)
 	{
 		PitchMultiplier.SetDefault(AudioDilationFactor_DEPRECATED);
 
-		AudioDilationFactor_DEPRECATED = FLT_MAX;
+		AudioDilationFactor_DEPRECATED = AudioDeprecatedMagicNumber;
 	}
 
-	if (AudioVolume_DEPRECATED != FLT_MAX)
+	if (AudioVolume_DEPRECATED != AudioDeprecatedMagicNumber)
 	{
 		SoundVolume.SetDefault(AudioVolume_DEPRECATED);
 
-		AudioVolume_DEPRECATED = FLT_MAX;
+		AudioVolume_DEPRECATED = AudioDeprecatedMagicNumber;
 	}
 
-	if (AudioStartTime_DEPRECATED != FLT_MAX)
+	if (AudioStartTime_DEPRECATED != AudioDeprecatedMagicNumber)
 	{
 		// Previously, start time in relation to the sequence. Start time was used to calculate the offset into the 
 		// clip at the start of the section evaluation as such: Section Start Time - Start Time. 
@@ -97,7 +102,7 @@ void UMovieSceneAudioSection::PostLoad()
 		{
 			StartOffset = GetInclusiveStartFrame() / GetTypedOuter<UMovieScene>()->GetTickResolution() - AudioStartTime_DEPRECATED;
 		}
-		AudioStartTime_DEPRECATED = FLT_MAX;
+		AudioStartTime_DEPRECATED = AudioDeprecatedMagicNumber;
 	}
 }
 
