Fix stats

diff --git a/Engine/Source/Runtime/MovieSceneTracks/Private/Evaluation/MovieSceneAudioTemplate.cpp b/Engine/Source/Runtime/MovieSceneTracks/Private/Evaluation/MovieSceneAudioTemplate.cpp
--- a/Engine/Source/Runtime/MovieSceneTracks/Private/Evaluation/MovieSceneAudioTemplate.cpp
+++ b/Engine/Source/Runtime/MovieSceneTracks/Private/Evaluation/MovieSceneAudioTemplate.cpp
@@ -18,6 +18,7 @@
 
 
 DECLARE_CYCLE_STAT(TEXT("Audio Track Evaluate"), MovieSceneEval_AudioTrack_Evaluate, STATGROUP_MovieSceneEval);
+DECLARE_CYCLE_STAT(TEXT("Audio Track Tear Down"), MovieSceneEval_AudioTrack_TearDown, STATGROUP_MovieSceneEval);
 DECLARE_CYCLE_STAT(TEXT("Audio Track Token Execute"), MovieSceneEval_AudioTrack_TokenExecute, STATGROUP_MovieSceneEval);
 
 
@@ -379,7 +380,7 @@ void FMovieSceneAudioSectionTemplate::Evaluate(const FMovieSceneEvaluationOperan
 
 void FMovieSceneAudioSectionTemplate::TearDown(FPersistentEvaluationData& PersistentData, IMovieScenePlayer& Player) const
 {
-	MOVIESCENE_DETAILED_SCOPE_CYCLE_COUNTER(MovieSceneEval_AudioTrack_Teardown)
+	MOVIESCENE_DETAILED_SCOPE_CYCLE_COUNTER(MovieSceneEval_AudioTrack_TearDown)
 
 	if (GEngine && GEngine->UseSound())
 	{
