Sequencer: Fixed Location-only properties not appearing on transform tracks

diff --git a/Engine/Source/Runtime/MovieSceneTracks/Private/Sections/MovieScene3DTransformSection.cpp b/Engine/Source/Runtime/MovieSceneTracks/Private/Sections/MovieScene3DTransformSection.cpp
--- a/Engine/Source/Runtime/MovieSceneTracks/Private/Sections/MovieScene3DTransformSection.cpp
+++ b/Engine/Source/Runtime/MovieSceneTracks/Private/Sections/MovieScene3DTransformSection.cpp
@@ -471,7 +471,7 @@ TSharedPtr<FStructOnScope> UMovieScene3DTransformSection::GetKeyStruct(TArrayVie
 		return KeyStruct;
 	}
 
-	if (AnyLocationKeys > 1)
+	if (AnyLocationKeys)
 	{
 		TSharedRef<FStructOnScope> KeyStruct = MakeShareable(new FStructOnScope(FMovieScene3DLocationKeyStruct::StaticStruct()));
 		auto Struct = (FMovieScene3DLocationKeyStruct*)KeyStruct->GetStructMemory();
