Changed template argument from FVector to FVector4 when Moive Scene Vector Section has 4 channels.

diff --git a/Engine/Source/Runtime/MovieSceneTracks/Private/Sections/MovieSceneVectorSection.cpp b/Engine/Source/Runtime/MovieSceneTracks/Private/Sections/MovieSceneVectorSection.cpp
--- a/Engine/Source/Runtime/MovieSceneTracks/Private/Sections/MovieSceneVectorSection.cpp
+++ b/Engine/Source/Runtime/MovieSceneTracks/Private/Sections/MovieSceneVectorSection.cpp
@@ -48,7 +48,7 @@ struct FVectorSectionEditorData
 		}
 		else
 		{
-			return Bindings.GetCurrentValue<FVector>(InObject);
+			return Bindings.GetCurrentValue<FVector4>(InObject);
 		}
 	}
 
