Sequencer: Add RF_Transactional to PrimitiveMaterialTrack::CreateNewSection and ObjectPropertyTrack::CreateNewSection

diff --git a/Engine/Source/Runtime/MovieSceneTracks/Private/Tracks/MovieSceneObjectPropertyTrack.cpp b/Engine/Source/Runtime/MovieSceneTracks/Private/Tracks/MovieSceneObjectPropertyTrack.cpp
--- a/Engine/Source/Runtime/MovieSceneTracks/Private/Tracks/MovieSceneObjectPropertyTrack.cpp
+++ b/Engine/Source/Runtime/MovieSceneTracks/Private/Tracks/MovieSceneObjectPropertyTrack.cpp
@@ -13,7 +13,7 @@ UMovieSceneObjectPropertyTrack::UMovieSceneObjectPropertyTrack(const FObjectInit
 
 UMovieSceneSection* UMovieSceneObjectPropertyTrack::CreateNewSection()
 {
-	UMovieSceneObjectPropertySection* Section = NewObject<UMovieSceneObjectPropertySection>(this);
+	UMovieSceneObjectPropertySection* Section = NewObject<UMovieSceneObjectPropertySection>(this, NAME_None, RF_Transactional);
 	Section->ObjectChannel.SetPropertyClass(PropertyClass);
 	return Section;
 }
diff --git a/Engine/Source/Runtime/MovieSceneTracks/Private/Tracks/MovieScenePrimitiveMaterialTrack.cpp b/Engine/Source/Runtime/MovieSceneTracks/Private/Tracks/MovieScenePrimitiveMaterialTrack.cpp
--- a/Engine/Source/Runtime/MovieSceneTracks/Private/Tracks/MovieScenePrimitiveMaterialTrack.cpp
+++ b/Engine/Source/Runtime/MovieSceneTracks/Private/Tracks/MovieScenePrimitiveMaterialTrack.cpp
@@ -16,7 +16,7 @@ UMovieScenePrimitiveMaterialTrack::UMovieScenePrimitiveMaterialTrack(const FObje
 
 UMovieSceneSection* UMovieScenePrimitiveMaterialTrack::CreateNewSection()
 {
-	return NewObject<UMovieScenePrimitiveMaterialSection>(this);
+	return NewObject<UMovieScenePrimitiveMaterialSection>(this, NAME_None, RF_Transactional);
 }
 
 FMovieSceneEvalTemplatePtr UMovieScenePrimitiveMaterialTrack::CreateTemplateForSection(const UMovieSceneSection& InSection) const
