Sequencer: Master event tracks no longer have an object pin by default. Level blueprints are now only accessible through interfaces

  - User actionable errors are also now FMessageLog so they are reported during PIE, rather than UE_LOGs

[Merged from 4565818]

diff --git a/Engine/Source/Editor/Sequencer/Private/MovieSceneSequenceEditor.cpp b/Engine/Source/Editor/Sequencer/Private/MovieSceneSequenceEditor.cpp
--- a/Engine/Source/Editor/Sequencer/Private/MovieSceneSequenceEditor.cpp
+++ b/Engine/Source/Editor/Sequencer/Private/MovieSceneSequenceEditor.cpp
@@ -31,7 +31,7 @@ namespace
 		FText DefaultComment_08 = LOCTEXT("DefaultComment_08", "\t\t\ta) Objects bound to the track's object binding, or:");
 		FText DefaultComment_09 = LOCTEXT("DefaultComment_09", "\t\t\tb) Objects specified on the track's event receivers array, or:");
 		FText DefaultComment_10 = LOCTEXT("DefaultComment_10", "\t\t\tc) Objects provided by the playback context (level blueprints, widgets etc)");
-		FText DefaultComment_11 = LOCTEXT("DefaultComment_11", "Tip: Trigger events on level blueprints by implementing an interface on it");
+		FText DefaultComment_11 = LOCTEXT("DefaultComment_11", "Tip: Trigger events on level blueprints using an interface pin, and adding the interface to the desired level blueprint");
 
 		return FText::Format(LOCTEXT("DefaultComment_Format", "{0}\n{1}\n{2}\n{3}\n{4}\n{5}\n{6}\n{7}\n{8}\n{9}\n\n{10}"),
 			DefaultComment_01, DefaultComment_02, DefaultComment_03, DefaultComment_04, DefaultComment_05,
@@ -108,7 +108,14 @@ UK2Node_FunctionEntry* FMovieSceneSequenceEditor::CreateEventEndpoint(UMovieScen
 
 void FMovieSceneSequenceEditor::InitializeEndpointForTrack(UMovieSceneEventTrack* EventTrack, UK2Node_FunctionEntry* Endpoint) const
 {
-	SetupDefaultPinForEndpoint(EventTrack, Endpoint);
+	UMovieScene* MovieScene = EventTrack->GetTypedOuter<UMovieScene>();
+
+	// Only add an object pin for the event if this is an event track on an object binding
+	FGuid Unused;
+	if (MovieScene->FindTrackBinding(*EventTrack, Unused))
+	{
+		SetupDefaultPinForEndpoint(EventTrack, Endpoint);
+	}
 }
 
 UClass* FMovieSceneSequenceEditor::FindTrackObjectBindingClass(UMovieSceneTrack* Track)
diff --git a/Engine/Source/Runtime/LevelSequence/Private/LevelSequenceDirector.cpp b/Engine/Source/Runtime/LevelSequence/Private/LevelSequenceDirector.cpp
--- a/Engine/Source/Runtime/LevelSequence/Private/LevelSequenceDirector.cpp
+++ b/Engine/Source/Runtime/LevelSequence/Private/LevelSequenceDirector.cpp
@@ -5,5 +5,9 @@
 
 UWorld* ULevelSequenceDirector::GetWorld() const
 {
+	if (ULevel* OuterLevel = GetTypedOuter<ULevel>())
+	{
+		return OuterLevel->OwningWorld;
+	}
 	return GetTypedOuter<UWorld>();
 }
diff --git a/Engine/Source/Runtime/MovieSceneTracks/Private/Evaluation/MovieSceneEventTemplate.cpp b/Engine/Source/Runtime/MovieSceneTracks/Private/Evaluation/MovieSceneEventTemplate.cpp
--- a/Engine/Source/Runtime/MovieSceneTracks/Private/Evaluation/MovieSceneEventTemplate.cpp
+++ b/Engine/Source/Runtime/MovieSceneTracks/Private/Evaluation/MovieSceneEventTemplate.cpp
@@ -9,7 +9,12 @@
 #include "Evaluation/MovieSceneEvaluation.h"
 #include "IMovieScenePlayer.h"
 #include "Algo/Accumulate.h"
+#include "Engine/LevelScriptActor.h"
+#include "Logging/TokenizedMessage.h"
+#include "Logging/MessageLog.h"
+#include "Misc/UObjectToken.h"
 
+#define LOCTEXT_NAMESPACE "MovieSceneEventTemplate"
 
 DECLARE_CYCLE_STAT(TEXT("Event Track Token Execute"), MovieSceneEval_EventTrack_TokenExecute, STATGROUP_MovieSceneEval);
 
@@ -238,7 +243,10 @@ struct FEventTriggerExecutionToken
 			// Event must have only a single object parameter, and the director instance must be an implementation of the function's class
 			if (!Function)
 			{
-				UE_LOG(LogMovieScene, Warning, TEXT("Failed to trigger the event '%s' because the function does not exist on director instance '%s'."), *EventName.ToString(), *DirectorInstance->GetName());
+				FMessageLog("PIE").Warning()
+					->AddToken(FTextToken::Create(FText::Format(LOCTEXT("LevelBP_MissingEvent_Error1", "Failed to trigger event '{0}' for"), FText::FromName(EventName))))
+					->AddToken(FUObjectToken::Create(Player.GetEvaluationTemplate().GetSequence(Operand.SequenceID)))
+					->AddToken(FTextToken::Create(LOCTEXT("LevelBP_MissingEvent_Error2", "because the function does not exist on the director instance.")));
 				continue;
 			}
 #if WITH_EDITOR
@@ -256,19 +264,43 @@ struct FEventTriggerExecutionToken
 			}
 			else if (Function->NumParms == 1 && Function->PropertyLink && (Function->PropertyLink->GetPropertyFlags() & CPF_ReferenceParm) == 0)
 			{
+				const int32 NumLevelScripts = Algo::Accumulate(EventContexts, 0, 
+					[](int32 Count, UObject* Obj)
+					{
+						return Obj && Obj->IsA<ALevelScriptActor>() ? Count + 1 : Count;
+					}
+				);
+
+				// Never pass through level script actors to event endpoints on non-interface pins.
+				if (NumLevelScripts > 0 && NumLevelScripts == EventContexts.Num() && !Function->PropertyLink->IsA<UInterfaceProperty>())
+				{
+					FMessageLog("PIE").Warning()
+						->AddToken(FTextToken::Create(LOCTEXT("LevelBP_ObjectPin_Error1", "Failed to trigger event")))
+						->AddToken(FUObjectToken::Create(Function))
+						->AddToken(FTextToken::Create(LOCTEXT("LevelBP_ObjectPin_Error2", "within")))
+						->AddToken(FUObjectToken::Create(Player.GetEvaluationTemplate().GetSequence(Operand.SequenceID)))
+						->AddToken(FTextToken::Create(LOCTEXT("LevelBP_ObjectPin_Error3", "because only Interface pins are supported for master tracks within Level Sequences. Please remove the pin, or change it to an interface that is implemented on the desired level blueprint.")));
+					continue;
+				}
+
 				for (UObject* EventContextObject : EventContexts)
 				{
-					TriggerEvent(DirectorInstance, Function, EventContextObject);
+					TriggerEvent(DirectorInstance, Function, EventContextObject, Player, Operand.SequenceID);
 				}
 			}
 			else
 			{
-				UE_LOG(LogMovieScene, Warning, TEXT("Failed to trigger event '%s' because its signature is not compatible. Function signatures must have either 0 or 1 (non-ref) parameters."), *EventName.ToString());
+				FMessageLog("PIE").Warning()
+					->AddToken(FTextToken::Create(LOCTEXT("LevelBP_InvalidEvent_Error1", "Failed to trigger event")))
+					->AddToken(FUObjectToken::Create(Function))
+					->AddToken(FTextToken::Create(LOCTEXT("LevelBP_InvalidEvent_Error2", "within")))
+					->AddToken(FUObjectToken::Create(Player.GetEvaluationTemplate().GetSequence(Operand.SequenceID)))
+					->AddToken(FTextToken::Create(LOCTEXT("LevelBP_InvalidEvent_Error3", "because its signature is not compatible. Function signatures must have either 0 or 1 (non-ref) parameters.")));
 			}
 		}
 	}
 
-	void TriggerEvent(UObject* DirectorInstance, UFunction* Function, UObject* ObjectParamValue)
+	void TriggerEvent(UObject* DirectorInstance, UFunction* Function, UObject* ObjectParamValue, IMovieScenePlayer& Player, FMovieSceneSequenceID SequenceID)
 	{
 		// We know by now that the parameter type is compatible because it wouldn't be added to the array if not
 		if (UObjectProperty* ObjectParameter = Cast<UObjectProperty>(Function->PropertyLink))
@@ -308,11 +340,12 @@ struct FEventTriggerExecutionToken
 			return;
 		}
 
-
-		UE_LOG(LogMovieScene, Warning, TEXT("Failed to trigger event '%s' because its signature is not compatible. Function expects a '%s' parameter, but only object and interface parameters are supported."),
-			*Function->GetName(),
-			*Function->PropertyLink->GetClass()->GetName()
-			);
+		FMessageLog("PIE").Warning()
+			->AddToken(FTextToken::Create(LOCTEXT("LevelBP_InvalidObjectEvent_Error1", "Failed to trigger event")))
+			->AddToken(FUObjectToken::Create(Function))
+			->AddToken(FTextToken::Create(LOCTEXT("LevelBP_InvalidObjectEvent_Error2", "within")))
+			->AddToken(FUObjectToken::Create(Player.GetEvaluationTemplate().GetSequence(SequenceID)))
+			->AddToken(FTextToken::Create(FText::Format(LOCTEXT("LevelBP_InvalidObjectEvent_Error3", "because its signature is not compatible. Function expects a '%s' parameter, but only object and interface parameters are supported."), FText::FromName(Function->PropertyLink->GetClass()->GetFName()))));
 	}
 
 #if !NO_LOGGING
@@ -502,4 +535,6 @@ void FMovieSceneEventRepeaterTemplate::EvaluateSwept(const FMovieSceneEvaluation
 		TArray<FName> Events = { EventToTrigger };
 		ExecutionTokens.Add(FEventTriggerExecutionToken(MoveTemp(Events), EventReceivers));
 	}
-}
\ No newline at end of file
+}
+
+#undef LOCTEXT_NAMESPACE
\ No newline at end of file
