Sequencer: Material switcher and parameter tracks now play nicely with each other

  - Material switchers no longer re-assign materials when a MID is already assigned that's parented to the desired material
  - Material switchers now have a high evaluation priority than parameters to ensure that parameters animate the correct material on first frame animations

diff --git a/Engine/Source/Runtime/MovieSceneTracks/Private/Evaluation/MovieScenePrimitiveMaterialTemplate.cpp b/Engine/Source/Runtime/MovieSceneTracks/Private/Evaluation/MovieScenePrimitiveMaterialTemplate.cpp
--- a/Engine/Source/Runtime/MovieSceneTracks/Private/Evaluation/MovieScenePrimitiveMaterialTemplate.cpp
+++ b/Engine/Source/Runtime/MovieSceneTracks/Private/Evaluation/MovieScenePrimitiveMaterialTemplate.cpp
@@ -5,6 +5,7 @@
 #include "Sections/MovieScenePrimitiveMaterialSection.h"
 #include "UObject/StrongObjectPtr.h"
 #include "Materials/MaterialInterface.h"
+#include "Materials/MaterialInstanceDynamic.h"
 #include "Components/PrimitiveComponent.h"
 
 struct FSetMaterialToken : IMovieScenePreAnimatedToken
@@ -57,9 +58,17 @@ struct FPrimitiveMaterialExecToken : IMovieSceneExecutionToken
 			UPrimitiveComponent* PrimitiveComponent = Cast<UPrimitiveComponent>(WeakObject.Get());
 			if (PrimitiveComponent && MaterialIndex >= 0 && MaterialIndex < PrimitiveComponent->GetNumMaterials())
 			{
-				UMaterialInterface* ExistingMaterial = PrimitiveComponent->GetMaterial(MaterialIndex);
-				Player.SavePreAnimatedState(*PrimitiveComponent, TMovieSceneAnimTypeID<FPrimitiveMaterialExecToken>(), FSetMaterialTokenProducer(MaterialIndex, ExistingMaterial));
+				UMaterialInterface*        ExistingMaterial = PrimitiveComponent->GetMaterial(MaterialIndex);
+				UMaterialInstanceDynamic*  MaterialInstance = Cast<UMaterialInstanceDynamic>(ExistingMaterial);
 
+				if (MaterialInstance && MaterialInstance->Parent && MaterialInstance->Parent == NewMaterial)
+				{
+					// Do not re-assign materials when a dynamic instance is already assigned with the same parent (since that's basically the same material, just with animated parameters)
+					// This is required for supporting material switchers alongside parameter tracks
+					continue;
+				}
+				
+				Player.SavePreAnimatedState(*PrimitiveComponent, TMovieSceneAnimTypeID<FPrimitiveMaterialExecToken>(), FSetMaterialTokenProducer(MaterialIndex, ExistingMaterial));
 				if (NewMaterial != ExistingMaterial)
 				{
 					PrimitiveComponent->SetMaterial(MaterialIndex, NewMaterial);
diff --git a/Engine/Source/Runtime/MovieSceneTracks/Private/Tracks/MovieSceneMaterialTrack.cpp b/Engine/Source/Runtime/MovieSceneTracks/Private/Tracks/MovieSceneMaterialTrack.cpp
--- a/Engine/Source/Runtime/MovieSceneTracks/Private/Tracks/MovieSceneMaterialTrack.cpp
+++ b/Engine/Source/Runtime/MovieSceneTracks/Private/Tracks/MovieSceneMaterialTrack.cpp
@@ -3,7 +3,7 @@
 #include "Tracks/MovieSceneMaterialTrack.h"
 #include "MovieSceneCommonHelpers.h"
 #include "Evaluation/MovieSceneParameterTemplate.h"
-
+#include "Evaluation/MovieSceneEvaluationTrack.h"
 
 UMovieSceneMaterialTrack::UMovieSceneMaterialTrack(const FObjectInitializer& ObjectInitializer)
 	: Super(ObjectInitializer)
@@ -105,6 +105,11 @@ FMovieSceneEvalTemplatePtr UMovieSceneComponentMaterialTrack::CreateTemplateForS
 	return FMovieSceneComponentMaterialSectionTemplate(*CastChecked<UMovieSceneParameterSection>(&InSection), *this);
 }
 
+void UMovieSceneComponentMaterialTrack::PostCompile(FMovieSceneEvaluationTrack& OutTrack, const FMovieSceneTrackCompilerArgs& Args) const
+{
+	OutTrack.SetEvaluationPriority(EvaluationPriority);
+}
+
 
 #if WITH_EDITORONLY_DATA
 FText UMovieSceneComponentMaterialTrack::GetDefaultDisplayName() const
diff --git a/Engine/Source/Runtime/MovieSceneTracks/Private/Tracks/MovieScenePrimitiveMaterialTrack.cpp b/Engine/Source/Runtime/MovieSceneTracks/Private/Tracks/MovieScenePrimitiveMaterialTrack.cpp
--- a/Engine/Source/Runtime/MovieSceneTracks/Private/Tracks/MovieScenePrimitiveMaterialTrack.cpp
+++ b/Engine/Source/Runtime/MovieSceneTracks/Private/Tracks/MovieScenePrimitiveMaterialTrack.cpp
@@ -1,8 +1,10 @@
 // Copyright 1998-2019 Epic Games, Inc. All Rights Reserved.
 
 #include "Tracks/MovieScenePrimitiveMaterialTrack.h"
+#include "Tracks/MovieSceneMaterialTrack.h"
 #include "Sections/MovieScenePrimitiveMaterialSection.h"
 #include "Evaluation/MovieScenePrimitiveMaterialTemplate.h"
+#include "Evaluation/MovieSceneEvaluationTrack.h"
 
 
 UMovieScenePrimitiveMaterialTrack::UMovieScenePrimitiveMaterialTrack(const FObjectInitializer& ObjInit)
@@ -23,3 +25,9 @@ FMovieSceneEvalTemplatePtr UMovieScenePrimitiveMaterialTrack::CreateTemplateForS
 {
 	return FMovieScenePrimitiveMaterialTemplate(*CastChecked<UMovieScenePrimitiveMaterialSection>(&InSection), *this);
 }
+
+void UMovieScenePrimitiveMaterialTrack::PostCompile(FMovieSceneEvaluationTrack& OutTrack, const FMovieSceneTrackCompilerArgs& Args) const
+{
+	// Must evaluate before material parameter tracks
+	OutTrack.SetEvaluationPriority(UMovieSceneComponentMaterialTrack::EvaluationPriority + 1);
+}
\ No newline at end of file
diff --git a/Engine/Source/Runtime/MovieSceneTracks/Public/Tracks/MovieSceneMaterialTrack.h b/Engine/Source/Runtime/MovieSceneTracks/Public/Tracks/MovieSceneMaterialTrack.h
--- a/Engine/Source/Runtime/MovieSceneTracks/Public/Tracks/MovieSceneMaterialTrack.h
+++ b/Engine/Source/Runtime/MovieSceneTracks/Public/Tracks/MovieSceneMaterialTrack.h
@@ -69,7 +69,10 @@ public:
 
 	// UMovieSceneTrack interface
 
+	static const uint16 EvaluationPriority = 1000;
+
 	virtual FMovieSceneEvalTemplatePtr CreateTemplateForSection(const UMovieSceneSection& InSection) const override;
+	virtual void PostCompile(FMovieSceneEvaluationTrack& OutTrack, const FMovieSceneTrackCompilerArgs& Args) const override;
 	virtual FName GetTrackName() const { return FName( *FString::FromInt(MaterialIndex) ); }
 
 #if WITH_EDITORONLY_DATA
diff --git a/Engine/Source/Runtime/MovieSceneTracks/Public/Tracks/MovieScenePrimitiveMaterialTrack.h b/Engine/Source/Runtime/MovieSceneTracks/Public/Tracks/MovieScenePrimitiveMaterialTrack.h
--- a/Engine/Source/Runtime/MovieSceneTracks/Public/Tracks/MovieScenePrimitiveMaterialTrack.h
+++ b/Engine/Source/Runtime/MovieSceneTracks/Public/Tracks/MovieScenePrimitiveMaterialTrack.h
@@ -21,4 +21,5 @@ public:
 	/*~ UMovieSceneTrack interface */
 	virtual UMovieSceneSection* CreateNewSection() override;
 	virtual FMovieSceneEvalTemplatePtr CreateTemplateForSection(const UMovieSceneSection& InSection) const override;
+	virtual void PostCompile(FMovieSceneEvaluationTrack& OutTrack, const FMovieSceneTrackCompilerArgs& Args) const override;
 };
