Sequencer: In PostEditImport, reconstruct sections that have variable channel proxies based on property data. This fixes an issue where pasting a track/sections doesn't result in valid channel proxies

diff --git a/Engine/Plugins/Animation/LiveLink/Source/LiveLink/Private/MovieScene/MovieSceneLiveLinkSection.cpp b/Engine/Plugins/Animation/LiveLink/Source/LiveLink/Private/MovieScene/MovieSceneLiveLinkSection.cpp
--- a/Engine/Plugins/Animation/LiveLink/Source/LiveLink/Private/MovieScene/MovieSceneLiveLinkSection.cpp
+++ b/Engine/Plugins/Animation/LiveLink/Source/LiveLink/Private/MovieScene/MovieSceneLiveLinkSection.cpp
@@ -154,6 +154,13 @@ void UMovieSceneLiveLinkSection::Serialize(FArchive& Ar)
 	}
 }
 
+void UMovieSceneLiveLinkSection::PostEditImport()
+{
+	Super::PostEditImport();
+
+	UpdateChannelProxy();
+}
+
 void UMovieSceneLiveLinkSection::SetMask(const TArray<bool>& InChannelMask)
 {
 	ChannelMask = InChannelMask;
diff --git a/Engine/Plugins/Animation/LiveLink/Source/LiveLink/Public/MovieScene/MovieSceneLiveLinkSection.h b/Engine/Plugins/Animation/LiveLink/Source/LiveLink/Public/MovieScene/MovieSceneLiveLinkSection.h
--- a/Engine/Plugins/Animation/LiveLink/Source/LiveLink/Public/MovieScene/MovieSceneLiveLinkSection.h
+++ b/Engine/Plugins/Animation/LiveLink/Source/LiveLink/Public/MovieScene/MovieSceneLiveLinkSection.h
@@ -48,6 +48,7 @@ public:
 protected:
 
 	virtual void Serialize(FArchive& Ar) override;
+	virtual void PostEditImport() override;
 	void UpdateChannelProxy();
 
 };
diff --git a/Engine/Source/Runtime/MovieSceneTracks/Private/Sections/MovieScene3DTransformSection.cpp b/Engine/Source/Runtime/MovieSceneTracks/Private/Sections/MovieScene3DTransformSection.cpp
--- a/Engine/Source/Runtime/MovieSceneTracks/Private/Sections/MovieScene3DTransformSection.cpp
+++ b/Engine/Source/Runtime/MovieSceneTracks/Private/Sections/MovieScene3DTransformSection.cpp
@@ -305,6 +305,13 @@ void UMovieScene3DTransformSection::Serialize(FArchive& Ar)
 	}
 }
 
+void UMovieScene3DTransformSection::PostEditImport()
+{
+	Super::PostEditImport();
+
+	UpdateChannelProxy();
+}
+
 FMovieSceneTransformMask UMovieScene3DTransformSection::GetMask() const
 {
 	return TransformMask;
diff --git a/Engine/Source/Runtime/MovieSceneTracks/Private/Sections/MovieSceneBoolSection.cpp b/Engine/Source/Runtime/MovieSceneTracks/Private/Sections/MovieSceneBoolSection.cpp
--- a/Engine/Source/Runtime/MovieSceneTracks/Private/Sections/MovieSceneBoolSection.cpp
+++ b/Engine/Source/Runtime/MovieSceneTracks/Private/Sections/MovieSceneBoolSection.cpp
@@ -73,6 +73,13 @@ void UMovieSceneBoolSection::ReconstructChannelProxy()
 #endif
 }
 
+void UMovieSceneBoolSection::PostEditImport()
+{
+	Super::PostEditImport();
+
+	ReconstructChannelProxy();
+}
+
 void UMovieSceneBoolSection::PostLoad()
 {
 	if (!BoolCurve.GetDefault().IsSet() && DefaultValue_DEPRECATED)
diff --git a/Engine/Source/Runtime/MovieSceneTracks/Private/Sections/MovieSceneParameterSection.cpp b/Engine/Source/Runtime/MovieSceneTracks/Private/Sections/MovieSceneParameterSection.cpp
--- a/Engine/Source/Runtime/MovieSceneTracks/Private/Sections/MovieSceneParameterSection.cpp
+++ b/Engine/Source/Runtime/MovieSceneTracks/Private/Sections/MovieSceneParameterSection.cpp
@@ -29,7 +29,7 @@ UMovieSceneParameterSection::UMovieSceneParameterSection( const FObjectInitializ
 			GetLinkerCustomVersion(FSequencerObjectVersion::GUID) < FSequencerObjectVersion::WhenFinishedDefaultsToProjectDefault ? 
 			EMovieSceneCompletionMode::RestoreState : 
 			EMovieSceneCompletionMode::ProjectDefault);
-}
+}	
 
 void UMovieSceneParameterSection::Serialize(FArchive& Ar)
 {
@@ -41,6 +41,14 @@ void UMovieSceneParameterSection::Serialize(FArchive& Ar)
 	}
 }
 
+void UMovieSceneParameterSection::PostEditImport()
+{
+	Super::PostEditImport();
+
+	ReconstructChannelProxy();
+}
+
+
 void UMovieSceneParameterSection::ReconstructChannelProxy()
 {
 	FMovieSceneChannelProxyData Channels;
diff --git a/Engine/Source/Runtime/MovieSceneTracks/Private/Sections/MovieSceneVectorSection.cpp b/Engine/Source/Runtime/MovieSceneTracks/Private/Sections/MovieSceneVectorSection.cpp
--- a/Engine/Source/Runtime/MovieSceneTracks/Private/Sections/MovieSceneVectorSection.cpp
+++ b/Engine/Source/Runtime/MovieSceneTracks/Private/Sections/MovieSceneVectorSection.cpp
@@ -107,6 +107,13 @@ void UMovieSceneVectorSection::Serialize(FArchive& Ar)
 	}
 }
 
+void UMovieSceneVectorSection::PostEditImport()
+{
+	Super::PostEditImport();
+
+	RecreateChannelProxy();
+}
+
 void UMovieSceneVectorSection::RecreateChannelProxy()
 {
 	FMovieSceneChannelProxyData Channels;
diff --git a/Engine/Source/Runtime/MovieSceneTracks/Public/Sections/MovieScene3DTransformSection.h b/Engine/Source/Runtime/MovieSceneTracks/Public/Sections/MovieScene3DTransformSection.h
--- a/Engine/Source/Runtime/MovieSceneTracks/Public/Sections/MovieScene3DTransformSection.h
+++ b/Engine/Source/Runtime/MovieSceneTracks/Public/Sections/MovieScene3DTransformSection.h
@@ -275,6 +275,7 @@ public:
 protected:
 
 	virtual void Serialize(FArchive& Ar) override;
+	virtual void PostEditImport() override;
 	virtual TSharedPtr<FStructOnScope> GetKeyStruct(TArrayView<const FKeyHandle> KeyHandles) override;
 	virtual FMovieSceneEvalTemplatePtr GenerateTemplate() const override;
 
diff --git a/Engine/Source/Runtime/MovieSceneTracks/Public/Sections/MovieSceneBoolSection.h b/Engine/Source/Runtime/MovieSceneTracks/Public/Sections/MovieSceneBoolSection.h
--- a/Engine/Source/Runtime/MovieSceneTracks/Public/Sections/MovieSceneBoolSection.h
+++ b/Engine/Source/Runtime/MovieSceneTracks/Public/Sections/MovieSceneBoolSection.h
@@ -31,7 +31,7 @@ public:
 public:
 
 	//~ UObject interface
-
+	virtual void PostEditImport() override;
 	virtual void PostLoad() override;
 
 protected:
diff --git a/Engine/Source/Runtime/MovieSceneTracks/Public/Sections/MovieSceneParameterSection.h b/Engine/Source/Runtime/MovieSceneTracks/Public/Sections/MovieSceneParameterSection.h
--- a/Engine/Source/Runtime/MovieSceneTracks/Public/Sections/MovieSceneParameterSection.h
+++ b/Engine/Source/Runtime/MovieSceneTracks/Public/Sections/MovieSceneParameterSection.h
@@ -227,6 +227,7 @@ protected:
 
 	//~ UMovieSceneSection interface
 	virtual void Serialize(FArchive& Ar) override;
+	virtual void PostEditImport() override;
 
 	void ReconstructChannelProxy();
 
diff --git a/Engine/Source/Runtime/MovieSceneTracks/Public/Sections/MovieSceneVectorSection.h b/Engine/Source/Runtime/MovieSceneTracks/Public/Sections/MovieSceneVectorSection.h
--- a/Engine/Source/Runtime/MovieSceneTracks/Public/Sections/MovieSceneVectorSection.h
+++ b/Engine/Source/Runtime/MovieSceneTracks/Public/Sections/MovieSceneVectorSection.h
@@ -127,6 +127,7 @@ protected:
 
 	//~ UMovieSceneSection interface
 	virtual void Serialize(FArchive& Ar) override;
+	virtual void PostEditImport() override;
 	virtual TSharedPtr<FStructOnScope> GetKeyStruct(TArrayView<const FKeyHandle> KeyHandles) override;
 
 	MOVIESCENETRACKS_API void RecreateChannelProxy();
diff --git a/Engine/Source/Runtime/UMG/Private/Animation/MovieScene2DTransformSection.cpp b/Engine/Source/Runtime/UMG/Private/Animation/MovieScene2DTransformSection.cpp
--- a/Engine/Source/Runtime/UMG/Private/Animation/MovieScene2DTransformSection.cpp
+++ b/Engine/Source/Runtime/UMG/Private/Animation/MovieScene2DTransformSection.cpp
@@ -158,6 +158,13 @@ void UMovieScene2DTransformSection::Serialize(FArchive& Ar)
 	}
 }
 
+void UMovieScene2DTransformSection::PostEditImport()
+{
+	Super::PostEditImport();
+
+	UpdateChannelProxy();
+}
+
 
 FMovieScene2DTransformMask UMovieScene2DTransformSection::GetMask() const
 {
diff --git a/Engine/Source/Runtime/UMG/Public/Animation/MovieScene2DTransformSection.h b/Engine/Source/Runtime/UMG/Public/Animation/MovieScene2DTransformSection.h
--- a/Engine/Source/Runtime/UMG/Public/Animation/MovieScene2DTransformSection.h
+++ b/Engine/Source/Runtime/UMG/Public/Animation/MovieScene2DTransformSection.h
@@ -118,6 +118,7 @@ public:
 protected:
 
 	virtual void Serialize(FArchive& Ar) override;
+	virtual void PostEditImport() override;
 
 	void UpdateChannelProxy();
 
