Sequencer: Expose GetAnimSectionsAtTime

diff --git a/Engine/Source/Runtime/MovieSceneTracks/Public/Tracks/MovieSceneSkeletalAnimationTrack.h b/Engine/Source/Runtime/MovieSceneTracks/Public/Tracks/MovieSceneSkeletalAnimationTrack.h
--- a/Engine/Source/Runtime/MovieSceneTracks/Public/Tracks/MovieSceneSkeletalAnimationTrack.h
+++ b/Engine/Source/Runtime/MovieSceneTracks/Public/Tracks/MovieSceneSkeletalAnimationTrack.h
@@ -19,13 +19,13 @@ class UMovieSceneSkeletalAnimationTrack
 public:
 
 	/** Adds a new animation to this track */
-	virtual UMovieSceneSection* AddNewAnimationOnRow(FFrameNumber KeyTime, class UAnimSequenceBase* AnimSequence, int32 RowIndex);
+	MOVIESCENETRACKS_API virtual UMovieSceneSection* AddNewAnimationOnRow(FFrameNumber KeyTime, class UAnimSequenceBase* AnimSequence, int32 RowIndex);
 
 	/** Adds a new animation to this track on the next available/non-overlapping row */
-	virtual UMovieSceneSection* AddNewAnimation(FFrameNumber KeyTime, class UAnimSequenceBase* AnimSequence) { return AddNewAnimationOnRow(KeyTime, AnimSequence, INDEX_NONE); }
+	MOVIESCENETRACKS_API virtual UMovieSceneSection* AddNewAnimation(FFrameNumber KeyTime, class UAnimSequenceBase* AnimSequence) { return AddNewAnimationOnRow(KeyTime, AnimSequence, INDEX_NONE); }
 
 	/** Gets the animation sections at a certain time */
-	TArray<UMovieSceneSection*> GetAnimSectionsAtTime(FFrameNumber Time);
+	MOVIESCENETRACKS_API TArray<UMovieSceneSection*> GetAnimSectionsAtTime(FFrameNumber Time);
 
 public:
 
