Sequencer: Disable can collapse to track for multi channel tracks

diff --git a/Engine/Plugins/Compositing/Composure/Source/Composure/Private/MovieScene/MovieSceneComposurePostMoveSettingsSection.cpp b/Engine/Plugins/Compositing/Composure/Source/Composure/Private/MovieScene/MovieSceneComposurePostMoveSettingsSection.cpp
--- a/Engine/Plugins/Compositing/Composure/Source/Composure/Private/MovieScene/MovieSceneComposurePostMoveSettingsSection.cpp
+++ b/Engine/Plugins/Compositing/Composure/Source/Composure/Private/MovieScene/MovieSceneComposurePostMoveSettingsSection.cpp
@@ -17,24 +17,30 @@ struct FPostMoveSettingsChannelEditorData
 		MetaData[0].SetIdentifiers("Pivot.X", FCommonChannelData::ChannelX, PivotGroup);
 		MetaData[0].SortOrder = 0;
 		MetaData[0].Color = FCommonChannelData::RedChannelColor;
+		MetaData[0].bCanCollapseToTrack = false;
 
 		MetaData[1].SetIdentifiers("Pivot.Y", FCommonChannelData::ChannelY, PivotGroup);
 		MetaData[1].SortOrder = 1;
 		MetaData[1].Color = FCommonChannelData::GreenChannelColor;
+		MetaData[1].bCanCollapseToTrack = false;
 
 		MetaData[2].SetIdentifiers("Translation.X", FCommonChannelData::ChannelX, TranslationGroup);
 		MetaData[2].SortOrder = 2;
 		MetaData[2].Color = FCommonChannelData::RedChannelColor;
+		MetaData[2].bCanCollapseToTrack = false;
 
 		MetaData[3].SetIdentifiers("Translation.Y", FCommonChannelData::ChannelY, TranslationGroup);
 		MetaData[3].SortOrder = 3;
 		MetaData[3].Color = FCommonChannelData::GreenChannelColor;
+		MetaData[3].bCanCollapseToTrack = false;
 
 		MetaData[4].SetIdentifiers("Rotation", NSLOCTEXT("PostMoves", "Rotation", "Rotation"));
 		MetaData[4].SortOrder = 4;
+		MetaData[4].bCanCollapseToTrack = false;
 
 		MetaData[5].SetIdentifiers("Scale", NSLOCTEXT("PostMoves", "Scale", "Scale"));
 		MetaData[5].SortOrder = 5;
+		MetaData[5].bCanCollapseToTrack = false;
 
 		ExternalValues[0].OnGetExternalValue = ExtractPivotX;
 		ExternalValues[1].OnGetExternalValue = ExtractPivotY;
diff --git a/Engine/Source/Runtime/MovieSceneTracks/Private/Sections/MovieScene3DTransformSection.cpp b/Engine/Source/Runtime/MovieSceneTracks/Private/Sections/MovieScene3DTransformSection.cpp
--- a/Engine/Source/Runtime/MovieSceneTracks/Private/Sections/MovieScene3DTransformSection.cpp
+++ b/Engine/Source/Runtime/MovieSceneTracks/Private/Sections/MovieScene3DTransformSection.cpp
@@ -23,48 +23,57 @@ struct F3DTransformChannelEditorData
 			MetaData[0].bEnabled = EnumHasAllFlags(Mask, EMovieSceneTransformChannel::TranslationX);
 			MetaData[0].Color = FCommonChannelData::RedChannelColor;
 			MetaData[0].SortOrder = 0;
+			MetaData[0].bCanCollapseToTrack = false;
 
 			MetaData[1].SetIdentifiers("Location.Y", FCommonChannelData::ChannelY, LocationGroup);
 			MetaData[1].bEnabled = EnumHasAllFlags(Mask, EMovieSceneTransformChannel::TranslationY);
 			MetaData[1].Color = FCommonChannelData::GreenChannelColor;
 			MetaData[1].SortOrder = 1;
+			MetaData[1].bCanCollapseToTrack = false;
 
 			MetaData[2].SetIdentifiers("Location.Z", FCommonChannelData::ChannelZ, LocationGroup);
 			MetaData[2].bEnabled = EnumHasAllFlags(Mask, EMovieSceneTransformChannel::TranslationZ);
 			MetaData[2].Color = FCommonChannelData::BlueChannelColor;
 			MetaData[2].SortOrder = 2;
+			MetaData[2].bCanCollapseToTrack = false;
 		}
 		{
 			MetaData[3].SetIdentifiers("Rotation.X", NSLOCTEXT("MovieSceneTransformSection", "RotationX", "Roll"), RotationGroup);
 			MetaData[3].bEnabled = EnumHasAllFlags(Mask, EMovieSceneTransformChannel::RotationX);
 			MetaData[3].Color = FCommonChannelData::RedChannelColor;
 			MetaData[3].SortOrder = 3;
+			MetaData[3].bCanCollapseToTrack = false;
 
 			MetaData[4].SetIdentifiers("Rotation.Y", NSLOCTEXT("MovieSceneTransformSection", "RotationY", "Pitch"), RotationGroup);
 			MetaData[4].bEnabled = EnumHasAllFlags(Mask, EMovieSceneTransformChannel::RotationY);
 			MetaData[4].Color = FCommonChannelData::GreenChannelColor;
 			MetaData[4].SortOrder = 4;
+			MetaData[4].bCanCollapseToTrack = false;
 
 			MetaData[5].SetIdentifiers("Rotation.Z", NSLOCTEXT("MovieSceneTransformSection", "RotationZ", "Yaw"), RotationGroup);
 			MetaData[5].bEnabled = EnumHasAllFlags(Mask, EMovieSceneTransformChannel::RotationZ);
 			MetaData[5].Color = FCommonChannelData::BlueChannelColor;
 			MetaData[5].SortOrder = 5;
+			MetaData[5].bCanCollapseToTrack = false;
 		}
 		{
 			MetaData[6].SetIdentifiers("Scale.X", FCommonChannelData::ChannelX, ScaleGroup);
 			MetaData[6].bEnabled = EnumHasAllFlags(Mask, EMovieSceneTransformChannel::ScaleX);
 			MetaData[6].Color = FCommonChannelData::RedChannelColor;
 			MetaData[6].SortOrder = 6;
+			MetaData[6].bCanCollapseToTrack = false;
 
 			MetaData[7].SetIdentifiers("Scale.Y", FCommonChannelData::ChannelY, ScaleGroup);
 			MetaData[7].bEnabled = EnumHasAllFlags(Mask, EMovieSceneTransformChannel::ScaleY);
 			MetaData[7].Color = FCommonChannelData::GreenChannelColor;
 			MetaData[7].SortOrder = 7;
+			MetaData[7].bCanCollapseToTrack = false;
 
 			MetaData[8].SetIdentifiers("Scale.Z", FCommonChannelData::ChannelZ, ScaleGroup);
 			MetaData[8].bEnabled = EnumHasAllFlags(Mask, EMovieSceneTransformChannel::ScaleZ);
 			MetaData[8].Color = FCommonChannelData::BlueChannelColor;
 			MetaData[8].SortOrder = 8;
+			MetaData[8].bCanCollapseToTrack = false;
 		}
 		{
 			MetaData[9].SetIdentifiers("Weight", NSLOCTEXT("MovieSceneTransformSection", "Weight", "Weight"));
diff --git a/Engine/Source/Runtime/MovieSceneTracks/Private/Sections/MovieSceneColorSection.cpp b/Engine/Source/Runtime/MovieSceneTracks/Private/Sections/MovieSceneColorSection.cpp
--- a/Engine/Source/Runtime/MovieSceneTracks/Private/Sections/MovieSceneColorSection.cpp
+++ b/Engine/Source/Runtime/MovieSceneTracks/Private/Sections/MovieSceneColorSection.cpp
@@ -14,17 +14,21 @@ struct FColorSectionEditorData
 		MetaData[0].SetIdentifiers("Color.R", FCommonChannelData::ChannelR);
 		MetaData[0].SortOrder = 0;
 		MetaData[0].Color = FCommonChannelData::RedChannelColor;
+		MetaData[0].bCanCollapseToTrack = false;
 
 		MetaData[1].SetIdentifiers("Color.G", FCommonChannelData::ChannelG);
 		MetaData[1].SortOrder = 1;
 		MetaData[1].Color = FCommonChannelData::GreenChannelColor;
+		MetaData[1].bCanCollapseToTrack = false;
 
 		MetaData[2].SetIdentifiers("Color.B", FCommonChannelData::ChannelB);
 		MetaData[2].SortOrder = 2;
 		MetaData[2].Color = FCommonChannelData::BlueChannelColor;
+		MetaData[2].bCanCollapseToTrack = false;
 
 		MetaData[3].SetIdentifiers("Color.A", FCommonChannelData::ChannelA);
 		MetaData[3].SortOrder = 3;
+		MetaData[3].bCanCollapseToTrack = false;
 
 		ExternalValues[0].OnGetExternalValue = ExtractChannelR;
 		ExternalValues[1].OnGetExternalValue = ExtractChannelG;
diff --git a/Engine/Source/Runtime/MovieSceneTracks/Private/Sections/MovieSceneVectorSection.cpp b/Engine/Source/Runtime/MovieSceneTracks/Private/Sections/MovieSceneVectorSection.cpp
--- a/Engine/Source/Runtime/MovieSceneTracks/Private/Sections/MovieSceneVectorSection.cpp
+++ b/Engine/Source/Runtime/MovieSceneTracks/Private/Sections/MovieSceneVectorSection.cpp
@@ -16,17 +16,21 @@ struct FVectorSectionEditorData
 		MetaData[0].SetIdentifiers("Vector.X", FCommonChannelData::ChannelX);
 		MetaData[0].SortOrder = 0;
 		MetaData[0].Color = FCommonChannelData::RedChannelColor;
+		MetaData[0].bCanCollapseToTrack = false;
 
 		MetaData[1].SetIdentifiers("Vector.Y", FCommonChannelData::ChannelY);
 		MetaData[1].SortOrder = 1;
 		MetaData[1].Color = FCommonChannelData::GreenChannelColor;
+		MetaData[1].bCanCollapseToTrack = false;
 
 		MetaData[2].SetIdentifiers("Vector.Z", FCommonChannelData::ChannelZ);
 		MetaData[2].SortOrder = 2;
 		MetaData[2].Color = FCommonChannelData::BlueChannelColor;
+		MetaData[2].bCanCollapseToTrack = false;
 
 		MetaData[3].SetIdentifiers("Vector.W", FCommonChannelData::ChannelW);
 		MetaData[3].SortOrder = 3;
+		MetaData[3].bCanCollapseToTrack = false;
 
 		ExternalValues[0].OnGetExternalValue = [NumChannels](UObject& InObject, FTrackInstancePropertyBindings* Bindings) { return ExtractChannelX(InObject, Bindings, NumChannels); };
 		ExternalValues[1].OnGetExternalValue = [NumChannels](UObject& InObject, FTrackInstancePropertyBindings* Bindings) { return ExtractChannelY(InObject, Bindings, NumChannels); };
diff --git a/Engine/Source/Runtime/UMG/Private/Animation/MovieScene2DTransformSection.cpp b/Engine/Source/Runtime/UMG/Private/Animation/MovieScene2DTransformSection.cpp
--- a/Engine/Source/Runtime/UMG/Private/Animation/MovieScene2DTransformSection.cpp
+++ b/Engine/Source/Runtime/UMG/Private/Animation/MovieScene2DTransformSection.cpp
@@ -20,30 +20,37 @@ struct F2DTransformSectionEditorData
 		MetaData[0].SetIdentifiers("Translation.X", FCommonChannelData::ChannelX, TranslationGroup);
 		MetaData[0].bEnabled = EnumHasAllFlags(Mask, EMovieScene2DTransformChannel::TranslationX);
 		MetaData[0].SortOrder = 0;
+		MetaData[0].bCanCollapseToTrack = false;
 
 		MetaData[1].SetIdentifiers("Translation.Y", FCommonChannelData::ChannelY, TranslationGroup);
 		MetaData[1].bEnabled = EnumHasAllFlags(Mask, EMovieScene2DTransformChannel::TranslationY);
 		MetaData[1].SortOrder = 1;
+		MetaData[1].bCanCollapseToTrack = false;
 
 		MetaData[2].SetIdentifiers("Angle", NSLOCTEXT("MovieScene2DTransformSection", "AngleText", "Angle"), RotationGroup);
 		MetaData[2].bEnabled = EnumHasAllFlags(Mask, EMovieScene2DTransformChannel::Rotation);
 		MetaData[2].SortOrder = 2;
+		MetaData[2].bCanCollapseToTrack = false;
 
 		MetaData[3].SetIdentifiers("Scale.X", FCommonChannelData::ChannelX, ScaleGroup);
 		MetaData[3].bEnabled = EnumHasAllFlags(Mask, EMovieScene2DTransformChannel::ScaleX);
 		MetaData[3].SortOrder = 3;
-		
+		MetaData[3].bCanCollapseToTrack = false;
+
 		MetaData[4].SetIdentifiers("Scale.Y", FCommonChannelData::ChannelY, ScaleGroup);
 		MetaData[4].bEnabled = EnumHasAllFlags(Mask, EMovieScene2DTransformChannel::ScaleY);
 		MetaData[4].SortOrder = 4;
+		MetaData[4].bCanCollapseToTrack = false;
 
 		MetaData[5].SetIdentifiers("Shear.X", FCommonChannelData::ChannelX, ShearGroup);
 		MetaData[5].bEnabled = EnumHasAllFlags(Mask, EMovieScene2DTransformChannel::ShearX);
 		MetaData[5].SortOrder = 5;
-		
+		MetaData[5].bCanCollapseToTrack = false;
+
 		MetaData[6].SetIdentifiers("Shear.Y", FCommonChannelData::ChannelY, ShearGroup);
 		MetaData[6].bEnabled = EnumHasAllFlags(Mask, EMovieScene2DTransformChannel::ShearY);
 		MetaData[6].SortOrder = 6;
+		MetaData[6].bCanCollapseToTrack = false;
 
 		ExternalValues[0].OnGetExternalValue = ExtractTranslationX;
 		ExternalValues[1].OnGetExternalValue = ExtractTranslationY;
diff --git a/Engine/Source/Runtime/UMG/Private/Animation/MovieSceneMarginSection.cpp b/Engine/Source/Runtime/UMG/Private/Animation/MovieSceneMarginSection.cpp
--- a/Engine/Source/Runtime/UMG/Private/Animation/MovieSceneMarginSection.cpp
+++ b/Engine/Source/Runtime/UMG/Private/Animation/MovieSceneMarginSection.cpp
@@ -11,15 +11,19 @@ struct FMarginSectionEditorData
 	{
 		MetaData[0].SetIdentifiers("Left", NSLOCTEXT("MovieSceneMarginSection", "LeftText", "Left"));
 		MetaData[0].SortOrder = 0;
+		MetaData[0].bCanCollapseToTrack = false;
 
 		MetaData[1].SetIdentifiers("Top", NSLOCTEXT("MovieSceneMarginSection", "TopText", "Top"));
 		MetaData[1].SortOrder = 1;
+		MetaData[1].bCanCollapseToTrack = false;
 
 		MetaData[2].SetIdentifiers("Right", NSLOCTEXT("MovieSceneMarginSection", "RightText", "Right"));
 		MetaData[2].SortOrder = 2;
+		MetaData[2].bCanCollapseToTrack = false;
 
 		MetaData[3].SetIdentifiers("Bottom", NSLOCTEXT("MovieSceneMarginSection", "BottomText", "Bottom"));
 		MetaData[3].SortOrder = 3;
+		MetaData[3].bCanCollapseToTrack = false;
 
 		ExternalValues[0].OnGetExternalValue = ExtractLeftChannel;
 		ExternalValues[1].OnGetExternalValue = ExtractTopChannel;
