Sequencer: Null check

diff --git a/Engine/Source/Runtime/MovieSceneTracks/Private/Evaluation/MovieSceneObjectPropertyTemplate.cpp b/Engine/Source/Runtime/MovieSceneTracks/Private/Evaluation/MovieSceneObjectPropertyTemplate.cpp
--- a/Engine/Source/Runtime/MovieSceneTracks/Private/Evaluation/MovieSceneObjectPropertyTemplate.cpp
+++ b/Engine/Source/Runtime/MovieSceneTracks/Private/Evaluation/MovieSceneObjectPropertyTemplate.cpp
@@ -54,10 +54,11 @@ struct FObjectPropertyExecToken : IMovieSceneExecutionToken
 		{
 			return !TargetProperty->HasAnyPropertyFlags(CPF_NoClear);
 		}
-		else
+		else if (DesiredValue->GetClass() != nullptr)
 		{
 			return DesiredValue->GetClass()->IsChildOf(TargetProperty->PropertyClass);
 		}
+		return false;
 	}
 
 	UObject* NewObjectValue;
