Sequencer: Fix skeletal mesh getting ticked twice when animation blueprint is set.

diff --git a/Engine/Source/Runtime/MovieSceneTracks/Private/Evaluation/MovieSceneSkeletalAnimationTemplate.cpp b/Engine/Source/Runtime/MovieSceneTracks/Private/Evaluation/MovieSceneSkeletalAnimationTemplate.cpp
--- a/Engine/Source/Runtime/MovieSceneTracks/Private/Evaluation/MovieSceneSkeletalAnimationTemplate.cpp
+++ b/Engine/Source/Runtime/MovieSceneTracks/Private/Evaluation/MovieSceneSkeletalAnimationTemplate.cpp
@@ -242,7 +242,7 @@ namespace MovieScene
 			// If the skeletal component has already ticked this frame because tick prerequisites weren't set up yet or a new binding was created, forcibly tick this component to update.
 			// This resolves first frame issues where the skeletal component ticks first, then the sequencer binding is resolved which sets up tick prerequisites
 			// for the next frame.
-			if (SkeletalMeshComponent->PoseTickedThisFrame() || SequencerInstance != ExistingAnimInstance)
+			if (SkeletalMeshComponent->PoseTickedThisFrame() || (SequencerInstance && SequencerInstance != ExistingAnimInstance))
 			{
 				SkeletalMeshComponent->TickAnimation(0.f, false);
 
