Sequencer: Ensure that material switcher tracks restore state correctly for different material indices

Max.Chen

diff --git a/Engine/Source/Runtime/MovieSceneTracks/Private/Evaluation/MovieScenePrimitiveMaterialTemplate.cpp b/Engine/Source/Runtime/MovieSceneTracks/Private/Evaluation/MovieScenePrimitiveMaterialTemplate.cpp
--- a/Engine/Source/Runtime/MovieSceneTracks/Private/Evaluation/MovieScenePrimitiveMaterialTemplate.cpp
+++ b/Engine/Source/Runtime/MovieSceneTracks/Private/Evaluation/MovieScenePrimitiveMaterialTemplate.cpp
@@ -53,6 +53,9 @@ struct FPrimitiveMaterialExecToken : IMovieSceneExecutionToken
 
 	virtual void Execute(const FMovieSceneContext& Context, const FMovieSceneEvaluationOperand& Operand, FPersistentEvaluationData& PersistentData, IMovieScenePlayer& Player) override
 	{
+		static TMovieSceneAnimTypeIDContainer<int32> MaterialIndexToTypeID;
+		FMovieSceneAnimTypeID TypeID = MaterialIndexToTypeID.GetAnimTypeID(MaterialIndex);
+
 		for (TWeakObjectPtr<> WeakObject : Player.FindBoundObjects(Operand))
 		{
 			UPrimitiveComponent* PrimitiveComponent = Cast<UPrimitiveComponent>(WeakObject.Get());
@@ -68,7 +71,7 @@ struct FPrimitiveMaterialExecToken : IMovieSceneExecutionToken
 					continue;
 				}
 				
-				Player.SavePreAnimatedState(*PrimitiveComponent, TMovieSceneAnimTypeID<FPrimitiveMaterialExecToken>(), FSetMaterialTokenProducer(MaterialIndex, ExistingMaterial));
+				Player.SavePreAnimatedState(*PrimitiveComponent, TypeID, FSetMaterialTokenProducer(MaterialIndex, ExistingMaterial));
 				if (NewMaterial != ExistingMaterial)
 				{
 					PrimitiveComponent->SetMaterial(MaterialIndex, NewMaterial);
