Sequencer: Added AudioTrack::CreateNewSection()

diff --git a/Engine/Source/Runtime/MovieSceneTracks/Private/Tracks/MovieSceneAudioTrack.cpp b/Engine/Source/Runtime/MovieSceneTracks/Private/Tracks/MovieSceneAudioTrack.cpp
--- a/Engine/Source/Runtime/MovieSceneTracks/Private/Tracks/MovieSceneAudioTrack.cpp
+++ b/Engine/Source/Runtime/MovieSceneTracks/Private/Tracks/MovieSceneAudioTrack.cpp
@@ -121,5 +121,10 @@ FMovieSceneTrackRowSegmentBlenderPtr UMovieSceneAudioTrack::GetRowSegmentBlender
 	return FBlender();
 }
 
+UMovieSceneSection* UMovieSceneAudioTrack::CreateNewSection()
+{
+	return NewObject<UMovieSceneAudioSection>(this, NAME_None, RF_Transactional);
+}
+
 
 #undef LOCTEXT_NAMESPACE
diff --git a/Engine/Source/Runtime/MovieSceneTracks/Public/Tracks/MovieSceneAudioTrack.h b/Engine/Source/Runtime/MovieSceneTracks/Public/Tracks/MovieSceneAudioTrack.h
--- a/Engine/Source/Runtime/MovieSceneTracks/Public/Tracks/MovieSceneAudioTrack.h
+++ b/Engine/Source/Runtime/MovieSceneTracks/Public/Tracks/MovieSceneAudioTrack.h
@@ -54,6 +54,7 @@ public:
 	virtual const TArray<UMovieSceneSection*>& GetAllSections() const override;
 	virtual bool SupportsMultipleRows() const override;
 	virtual FMovieSceneTrackRowSegmentBlenderPtr GetRowSegmentBlender() const override;
+	virtual UMovieSceneSection* CreateNewSection() override;
 
 private:
 
