Sequencer: Add function to get the active camera cut camera. Allow camera cut evaluation, but disable binding to the player when CanUpdateCameraCut is false.

diff --git a/Engine/Source/Runtime/LevelSequence/Private/LevelSequencePlayer.cpp b/Engine/Source/Runtime/LevelSequence/Private/LevelSequencePlayer.cpp
--- a/Engine/Source/Runtime/LevelSequence/Private/LevelSequencePlayer.cpp
+++ b/Engine/Source/Runtime/LevelSequence/Private/LevelSequencePlayer.cpp
@@ -209,6 +209,11 @@ void ULevelSequencePlayer::UpdateCameraCut(UObject* CameraObject, UObject* Unloc
 
 	CachedCameraComponent = CameraComponent;
 
+	if (!CanUpdateCameraCut())
+	{
+		return;
+	}
+
 	if (CameraObject == ViewTarget)
 	{
 		if ( bJumpCut )
diff --git a/Engine/Source/Runtime/LevelSequence/Public/LevelSequencePlayer.h b/Engine/Source/Runtime/LevelSequence/Public/LevelSequencePlayer.h
--- a/Engine/Source/Runtime/LevelSequence/Public/LevelSequencePlayer.h
+++ b/Engine/Source/Runtime/LevelSequence/Public/LevelSequencePlayer.h
@@ -136,6 +136,10 @@ public:
 	UPROPERTY(BlueprintAssignable, Category="Game|Cinematic")
 	FOnLevelSequencePlayerCameraCutEvent OnCameraCut;
 
+	/** Get the active camera cut camera */
+	UFUNCTION(BlueprintCallable, Category="Game|Cinematic")
+	UCameraComponent* GetActiveCameraComponent() const { return CachedCameraComponent.Get(); }
+
 public:
 
 	/**
diff --git a/Engine/Source/Runtime/MovieSceneTracks/Private/Evaluation/MovieSceneCameraCutTemplate.cpp b/Engine/Source/Runtime/MovieSceneTracks/Private/Evaluation/MovieSceneCameraCutTemplate.cpp
--- a/Engine/Source/Runtime/MovieSceneTracks/Private/Evaluation/MovieSceneCameraCutTemplate.cpp
+++ b/Engine/Source/Runtime/MovieSceneTracks/Private/Evaluation/MovieSceneCameraCutTemplate.cpp
@@ -103,12 +103,6 @@ struct FCameraCutExecutionToken : IMovieSceneExecutionToken
 	/** Execute this token, operating on all objects referenced by 'Operand' */
 	virtual void Execute(const FMovieSceneContext& Context, const FMovieSceneEvaluationOperand& Operand, FPersistentEvaluationData& PersistentData, IMovieScenePlayer& Player) override
 	{
-		if (!Player.CanUpdateCameraCut())
-		{
-			Player.UpdateCameraCut(nullptr, nullptr);
-			return;
-		}
-
 		MOVIESCENE_DETAILED_SCOPE_CYCLE_COUNTER(MovieSceneEval_CameraCutTrack_TokenExecute)
 
 		FMovieSceneSequenceID SequenceID = Operand.SequenceID;
