Sequencer: Don't play audio if offset into audio clip is beyond bounds of audio duration

diff --git a/Engine/Source/Runtime/MovieSceneTracks/Private/Evaluation/MovieSceneAudioTemplate.cpp b/Engine/Source/Runtime/MovieSceneTracks/Private/Evaluation/MovieSceneAudioTemplate.cpp
--- a/Engine/Source/Runtime/MovieSceneTracks/Private/Evaluation/MovieSceneAudioTemplate.cpp
+++ b/Engine/Source/Runtime/MovieSceneTracks/Private/Evaluation/MovieSceneAudioTemplate.cpp
@@ -331,8 +331,11 @@ void FMovieSceneAudioSectionTemplateData::EnsureAudioIsPlaying(UAudioComponent&
 			AudioComponent.bIsUISound = false;
 		}
 
-		const float AudioTime = (Context.GetTime() / Context.GetFrameRate()) - SectionStartTimeSeconds + FMath::Max((float)Context.GetFrameRate().AsSeconds(AudioStartOffset), 0.f);
-		AudioComponent.Play(AudioTime);
+		const float AudioTime = (Context.GetTime() / Context.GetFrameRate()) - SectionStartTimeSeconds + (float)Context.GetFrameRate().AsSeconds(AudioStartOffset);
+		if (AudioTime > 0.f && AudioTime < AudioComponent.Sound->Duration)
+		{
+			AudioComponent.Play(AudioTime);
+		}
 
 		if (Context.GetStatus() == EMovieScenePlayerStatus::Scrubbing)
 		{
