When playing Skeletal Animation if we are playing outside of a section that previously had keep state it keeps itself in the Animation Proxy Blend node since we don't get any notification (like from PreAnimationToken.RestoreState) that it's no longer alive. So what happens is that though the animation isn't active it still blends with it's last evaluated state.  (If we RestoreState we reset the AnimSequence so the previous Proxy Blend node is deleted).

The simple fix that is working well after some testing is to simply always zero out the weights in the AnimSequenceInstance or the AnimMontag(needed to expose SetDesiredWeight) since they then get set anyway when we cycle through the Animation Parameters and call SetFinalAnimation.

diff --git a/Engine/Source/Runtime/Engine/Classes/Animation/AnimMontage.h b/Engine/Source/Runtime/Engine/Classes/Animation/AnimMontage.h
--- a/Engine/Source/Runtime/Engine/Classes/Animation/AnimMontage.h
+++ b/Engine/Source/Runtime/Engine/Classes/Animation/AnimMontage.h
@@ -402,7 +402,8 @@ public:
 
 	/** Set the weight */
 	void SetWeight(float InValue) { Blend.SetAlpha(InValue); }
-
+	/** Set the Desired Weight */
+	void SetDesiredWeight(float InValue) { Blend.SetDesiredValue(InValue); }
 private:
 	/** Followers this Montage will synchronize */
 	TArray<struct FAnimMontageInstance*> MontageSyncFollowers;
diff --git a/Engine/Source/Runtime/MovieSceneTracks/Private/Evaluation/MovieSceneSkeletalAnimationTemplate.cpp b/Engine/Source/Runtime/MovieSceneTracks/Private/Evaluation/MovieSceneSkeletalAnimationTemplate.cpp
--- a/Engine/Source/Runtime/MovieSceneTracks/Private/Evaluation/MovieSceneSkeletalAnimationTemplate.cpp
+++ b/Engine/Source/Runtime/MovieSceneTracks/Private/Evaluation/MovieSceneSkeletalAnimationTemplate.cpp
@@ -200,9 +200,9 @@ namespace MovieScene
 
 			OriginalStack.SavePreAnimatedState(Player, *SkeletalMeshComponent, GetAnimControlTypeID(), FPreAnimatedAnimationTokenProducer());
 
-			const UAnimInstance* ExistingAnimInstance = SkeletalMeshComponent->GetAnimInstance();
+			UAnimInstance* ExistingAnimInstance = SkeletalMeshComponent->GetAnimInstance();
 
-			const UAnimSequencerInstance* SequencerInstance = UAnimCustomInstance::BindToSkeletalMeshComponent<UAnimSequencerInstance>(SkeletalMeshComponent);
+			UAnimSequencerInstance* SequencerInstance = UAnimCustomInstance::BindToSkeletalMeshComponent<UAnimSequencerInstance>(SkeletalMeshComponent);
 
 			const bool bPreviewPlayback = ShouldUsePreviewPlayback(Player, *SkeletalMeshComponent);
 
@@ -221,6 +221,25 @@ namespace MovieScene
 										(DeltaTime == 0.0f && PlayerStatus != EMovieScenePlayerStatus::Stopped); 
 		
 			static const bool bLooping = false;
+			//Need to zero all weights first since we may be blending animation that are keeping state but are no longer active.
+			
+			if(SequencerInstance)
+			{
+				SequencerInstance->ResetNodes();
+			}
+			else if (ExistingAnimInstance)
+			{
+				for (const TPair<FObjectKey, FMontagePlayerPerSectionData >& Pair : MontageData)
+				{
+					int32 InstanceId = Pair.Value.MontageInstanceId;
+					FAnimMontageInstance* MontageInstanceToUpdate = ExistingAnimInstance->GetMontageInstanceForID(InstanceId);
+					if (MontageInstanceToUpdate)
+					{
+						MontageInstanceToUpdate->SetDesiredWeight(0.0f);
+						MontageInstanceToUpdate->SetWeight(0.0f);
+					}
+				}
+			}
 			for (const FMinimalAnimParameters& AnimParams : InFinalValue.AllAnimations)
 			{
 				Player.PreAnimatedState.SetCaptureEntity(AnimParams.EvaluationScope.Key, AnimParams.EvaluationScope.CompletionMode);
