Jittery Sequencer Animations with RateScale > 0.0

Two changes fixed this. First we invert out the RateScale when calculating the PlayRate inside  FAnimNode_SequenceEvaluator::UpdateAssetPlayer, since it get's multiplied back in the CreateTickRecordForNode function.   This stops the jittering but the animation playback isn't scaled. To fix that we modify the Sequencer sections PlayRate by the asset'ss RateScale.

diff --git a/Engine/Source/Runtime/AnimGraphRuntime/Private/AnimNodes/AnimNode_SequenceEvaluator.cpp b/Engine/Source/Runtime/AnimGraphRuntime/Private/AnimNodes/AnimNode_SequenceEvaluator.cpp
--- a/Engine/Source/Runtime/AnimGraphRuntime/Private/AnimNodes/AnimNode_SequenceEvaluator.cpp
+++ b/Engine/Source/Runtime/AnimGraphRuntime/Private/AnimNodes/AnimNode_SequenceEvaluator.cpp
@@ -72,7 +72,8 @@ void FAnimNode_SequenceEvaluator::UpdateAssetPlayer(const FAnimationUpdateContex
  			}
 			
 			const float DeltaTime = Context.GetDeltaTime();
-			const float PlayRate = FMath::IsNearlyZero(DeltaTime) ? 0.f : (TimeJump / DeltaTime);
+			const float RateScale = Sequence->RateScale;
+			const float PlayRate = FMath::IsNearlyZero(DeltaTime) || FMath::IsNearlyZero(RateScale) ? 0.f : (TimeJump / (DeltaTime * RateScale));
 			CreateTickRecordForNode(Context, Sequence, bShouldLoop, PlayRate);
 		}
 		else
diff --git a/Engine/Source/Runtime/MovieSceneTracks/Private/Evaluation/MovieSceneSkeletalAnimationTemplate.cpp b/Engine/Source/Runtime/MovieSceneTracks/Private/Evaluation/MovieSceneSkeletalAnimationTemplate.cpp
--- a/Engine/Source/Runtime/MovieSceneTracks/Private/Evaluation/MovieSceneSkeletalAnimationTemplate.cpp
+++ b/Engine/Source/Runtime/MovieSceneTracks/Private/Evaluation/MovieSceneSkeletalAnimationTemplate.cpp
@@ -429,8 +429,8 @@ void FMovieSceneSkeletalAnimationSectionTemplate::Evaluate(const FMovieSceneEval
 float FMovieSceneSkeletalAnimationSectionTemplateParameters::MapTimeToAnimation(FFrameTime InPosition, FFrameRate InFrameRate) const
 {
 	InPosition = FMath::Clamp(InPosition, FFrameTime(SectionStartTime), FFrameTime(SectionEndTime-1));
-
-	const float SectionPlayRate = PlayRate;
+	
+	const float SectionPlayRate = PlayRate * Animation->RateScale;
 	const float AnimPlayRate = FMath::IsNearlyZero(SectionPlayRate) ? 1.0f : SectionPlayRate;
 
 	const float SeqLength = GetSequenceLength() - InFrameRate.AsSeconds(StartFrameOffset + EndFrameOffset);
