Sequencer: Force animation playback to be Custom when playing back Sequencer Animation Asset even if the character is trying to force itself back to BP.

diff --git a/Engine/Plugins/VirtualProduction/Takes/Source/TakeTrackRecorders/Private/MovieSceneAnimationTrackRecorder.cpp b/Engine/Plugins/VirtualProduction/Takes/Source/TakeTrackRecorders/Private/MovieSceneAnimationTrackRecorder.cpp
--- a/Engine/Plugins/VirtualProduction/Takes/Source/TakeTrackRecorders/Private/MovieSceneAnimationTrackRecorder.cpp
+++ b/Engine/Plugins/VirtualProduction/Takes/Source/TakeTrackRecorders/Private/MovieSceneAnimationTrackRecorder.cpp
@@ -157,6 +157,7 @@ void UMovieSceneAnimationTrackRecorder::CreateTrackImpl()
 			{
 				AnimTrack->AddNewAnimation(FFrameNumber(0), AnimSequence.Get());
 				MovieSceneSection = Cast<UMovieSceneSkeletalAnimationSection>(AnimTrack->GetAllSections()[0]);
+				MovieSceneSection->Params.bForceCustomMode = true;
 			}
 		}
 	}
diff --git a/Engine/Source/Runtime/MovieSceneTracks/Private/Evaluation/MovieSceneSkeletalAnimationTemplate.cpp b/Engine/Source/Runtime/MovieSceneTracks/Private/Evaluation/MovieSceneSkeletalAnimationTemplate.cpp
--- a/Engine/Source/Runtime/MovieSceneTracks/Private/Evaluation/MovieSceneSkeletalAnimationTemplate.cpp
+++ b/Engine/Source/Runtime/MovieSceneTracks/Private/Evaluation/MovieSceneSkeletalAnimationTemplate.cpp
@@ -123,7 +123,7 @@ struct FPreAnimatedAnimationTokenProducer : IMovieScenePreAnimatedTokenProducer
 
 struct FMinimalAnimParameters
 {
-	FMinimalAnimParameters(UAnimSequenceBase* InAnimation, float InEvalTime, float InBlendWeight, const FMovieSceneEvaluationScope& InScope, FName InSlotName, FObjectKey InSection, bool InSkipAnimationNotifiers)
+	FMinimalAnimParameters(UAnimSequenceBase* InAnimation, float InEvalTime, float InBlendWeight, const FMovieSceneEvaluationScope& InScope, FName InSlotName, FObjectKey InSection, bool InSkipAnimationNotifiers, bool InForceCustomMode)
 		: Animation(InAnimation)
 		, EvalTime(InEvalTime)
 		, BlendWeight(InBlendWeight)
@@ -131,6 +131,7 @@ struct FMinimalAnimParameters
 		, SlotName(InSlotName)
 		, Section(InSection)
 		, bSkipAnimNotifiers(InSkipAnimationNotifiers)
+		, bForceCustomMode(InForceCustomMode)
 	{}
 	
 	UAnimSequenceBase* Animation;
@@ -140,6 +141,7 @@ struct FMinimalAnimParameters
 	FName SlotName;
 	FObjectKey Section;
 	bool bSkipAnimNotifiers;
+	bool bForceCustomMode;
 };
 
 /** Montage player per section data */
@@ -197,7 +199,6 @@ namespace MovieScene
 			{
 				return;
 			}
-
 			OriginalStack.SavePreAnimatedState(Player, *SkeletalMeshComponent, GetAnimControlTypeID(), FPreAnimatedAnimationTokenProducer());
 
 			UAnimInstance* ExistingAnimInstance = SkeletalMeshComponent->GetAnimInstance();
@@ -248,13 +249,16 @@ namespace MovieScene
 				{
 					PreviewSetAnimPosition(PersistentData, Player, SkeletalMeshComponent,
 						AnimParams.SlotName, AnimParams.Section, AnimParams.Animation, AnimParams.EvalTime, AnimParams.BlendWeight,
-						bLooping, bFireNotifies && !AnimParams.bSkipAnimNotifiers, DeltaTime, Player.GetPlaybackStatus() == EMovieScenePlayerStatus::Playing, bResetDynamics);
+						bLooping, bFireNotifies && !AnimParams.bSkipAnimNotifiers, DeltaTime, Player.GetPlaybackStatus() == EMovieScenePlayerStatus::Playing, 
+						bResetDynamics, AnimParams.bForceCustomMode);
 				}
 				else
 				{
 					SetAnimPosition(PersistentData, Player, SkeletalMeshComponent,
 						AnimParams.SlotName, AnimParams.Section, AnimParams.Animation, AnimParams.EvalTime, AnimParams.BlendWeight,
-						bLooping, Player.GetPlaybackStatus() == EMovieScenePlayerStatus::Playing, bFireNotifies && !AnimParams.bSkipAnimNotifiers);
+						bLooping, Player.GetPlaybackStatus() == EMovieScenePlayerStatus::Playing, bFireNotifies && !AnimParams.bSkipAnimNotifiers,
+						AnimParams.bForceCustomMode
+					);
 				}
 			}
 
@@ -293,13 +297,16 @@ namespace MovieScene
 			return nullptr;
 		}
 
-		void SetAnimPosition(FPersistentEvaluationData& PersistentData, IMovieScenePlayer& Player, USkeletalMeshComponent* SkeletalMeshComponent, FName SlotName, FObjectKey Section, UAnimSequenceBase* InAnimSequence, float InPosition, float Weight, bool bLooping, bool bPlaying, bool bFireNotifies)
+		void SetAnimPosition(FPersistentEvaluationData& PersistentData, IMovieScenePlayer& Player, USkeletalMeshComponent* SkeletalMeshComponent, FName SlotName, FObjectKey Section, UAnimSequenceBase* InAnimSequence, float InPosition, float Weight, bool bLooping, bool bPlaying, bool bFireNotifies, bool bForceCustomMode)
 		{
 			if (!CanPlayAnimation(SkeletalMeshComponent, InAnimSequence))
 			{
 				return;
 			}
-
+			if (bForceCustomMode)
+			{
+				SkeletalMeshComponent->SetAnimationMode(EAnimationMode::AnimationCustomMode);
+			}
 			UAnimSequencerInstance* SequencerInst = Cast<UAnimSequencerInstance>(SkeletalMeshComponent->GetAnimInstance());
 			if (SequencerInst)
 			{
@@ -333,13 +340,16 @@ namespace MovieScene
 			}
 		}
 
-		void PreviewSetAnimPosition(FPersistentEvaluationData& PersistentData, IMovieScenePlayer& Player, USkeletalMeshComponent* SkeletalMeshComponent, FName SlotName, FObjectKey Section, UAnimSequenceBase* InAnimSequence, float InPosition, float Weight, bool bLooping, bool bFireNotifies, float DeltaTime, bool bPlaying, bool bResetDynamics)
+		void PreviewSetAnimPosition(FPersistentEvaluationData& PersistentData, IMovieScenePlayer& Player, USkeletalMeshComponent* SkeletalMeshComponent, FName SlotName, FObjectKey Section, UAnimSequenceBase* InAnimSequence, float InPosition, float Weight, bool bLooping, bool bFireNotifies, float DeltaTime, bool bPlaying, bool bResetDynamics, bool bForceCustomMode)
 		{
 			if (!CanPlayAnimation(SkeletalMeshComponent, InAnimSequence))
 			{
 				return;
 			}
-
+			if (bForceCustomMode)
+			{
+				SkeletalMeshComponent->SetAnimationMode(EAnimationMode::AnimationCustomMode);
+			}
 			UAnimSequencerInstance* SequencerInst = Cast<UAnimSequencerInstance>(SkeletalMeshComponent->GetAnimInstance());
 			if (SequencerInst)
 			{
@@ -420,7 +430,8 @@ void FMovieSceneSkeletalAnimationSectionTemplate::Evaluate(const FMovieSceneEval
 
 		// Add the blendable to the accumulator
 		FMinimalAnimParameters AnimParams(
-			Params.Animation, EvalTime, Weight, ExecutionTokens.GetCurrentScope(), Params.SlotName, GetSourceSection(), Params.bSkipAnimNotifiers
+			Params.Animation, EvalTime, Weight, ExecutionTokens.GetCurrentScope(), Params.SlotName, GetSourceSection(), Params.bSkipAnimNotifiers, 
+			Params.bForceCustomMode
 		);
 		ExecutionTokens.BlendToken(ActuatorTypeID, TBlendableToken<MovieScene::FBlendedAnimation>(AnimParams, BlendType.Get(), 1.f));
 	}
diff --git a/Engine/Source/Runtime/MovieSceneTracks/Private/Sections/MovieSceneSkeletalAnimationSection.cpp b/Engine/Source/Runtime/MovieSceneTracks/Private/Sections/MovieSceneSkeletalAnimationSection.cpp
--- a/Engine/Source/Runtime/MovieSceneTracks/Private/Sections/MovieSceneSkeletalAnimationSection.cpp
+++ b/Engine/Source/Runtime/MovieSceneTracks/Private/Sections/MovieSceneSkeletalAnimationSection.cpp
@@ -27,6 +27,7 @@ FMovieSceneSkeletalAnimationParams::FMovieSceneSkeletalAnimationParams()
 	SlotName = DefaultSlotName;
 	Weight.SetDefault(1.f);
 	bSkipAnimNotifiers = false;
+	bForceCustomMode = false;
 }
 
 UMovieSceneSkeletalAnimationSection::UMovieSceneSkeletalAnimationSection( const FObjectInitializer& ObjectInitializer )
diff --git a/Engine/Source/Runtime/MovieSceneTracks/Public/Sections/MovieSceneSkeletalAnimationSection.h b/Engine/Source/Runtime/MovieSceneTracks/Public/Sections/MovieSceneSkeletalAnimationSection.h
--- a/Engine/Source/Runtime/MovieSceneTracks/Public/Sections/MovieSceneSkeletalAnimationSection.h
+++ b/Engine/Source/Runtime/MovieSceneTracks/Public/Sections/MovieSceneSkeletalAnimationSection.h
@@ -54,6 +54,10 @@ struct FMovieSceneSkeletalAnimationParams
 	UPROPERTY(EditAnywhere, Category = "Animation")
 	bool bSkipAnimNotifiers;
 
+	/** If on animation sequence will always play when active even if the animation is controlled by a Blueprint or Anim Instance Class*/
+	UPROPERTY(EditAnywhere, Category = "Animation")
+	bool bForceCustomMode;
+
 	UPROPERTY()
 	float StartOffset_DEPRECATED;
 
