Layering Support for Channel Key Editing via IKeyArea.

With the previous work here we supporting Layering with external values correctly working with the TrackEditors which work when setting keys via the details panel, auto key, or the hot keys. However if you pressed the key button icon next to the channel in the sequencer UI it would only now always just set the current value which could be different than what it should be based upon the external value (like if you moved the sphere in the jira ticket).

So we added support for an extra OnGetCurrentValueAndWeight to the TMovieSceneExternalValue in order to modify the value based upone the global diff/weight as we did before. This specialization only happens for <float> and currently only for the Transform/Vector sections (other like ComposurePostMovieSettings seemed margiinally useful and we don't support ExternalValues for generic float chanels for some reason).

diff --git a/Engine/Source/Editor/MovieSceneTools/Private/Channels/BuiltInChannelEditors.cpp b/Engine/Source/Editor/MovieSceneTools/Private/Channels/BuiltInChannelEditors.cpp
--- a/Engine/Source/Editor/MovieSceneTools/Private/Channels/BuiltInChannelEditors.cpp
+++ b/Engine/Source/Editor/MovieSceneTools/Private/Channels/BuiltInChannelEditors.cpp
@@ -22,32 +22,47 @@
 #define LOCTEXT_NAMESPACE "BuiltInChannelEditors"
 
 
-FKeyHandle AddOrUpdateKey(FMovieSceneFloatChannel* Channel, const TMovieSceneExternalValue<float>& ExternalValue, FFrameNumber InTime, ISequencer& Sequencer, const FGuid& InObjectBindingID, FTrackInstancePropertyBindings* PropertyBindings)
+FKeyHandle AddOrUpdateKey(FMovieSceneFloatChannel* Channel, UMovieSceneSection* SectionToKey, const TMovieSceneExternalValue<float>& ExternalValue, FFrameNumber InTime, ISequencer& Sequencer, const FGuid& InObjectBindingID, FTrackInstancePropertyBindings* PropertyBindings)
 {
 	TOptional<float> Value;
-
-	// Add a key for the current value of the valid first object we can find
-	if (ExternalValue.OnGetExternalValue && InObjectBindingID.IsValid())
+	float CurrentValue = 0.0f, CurrentWeight = 1.0f;
+	if ((ExternalValue.OnGetExternalValue && InObjectBindingID.IsValid()))
 	{
 		for (TWeakObjectPtr<> WeakObject : Sequencer.FindBoundObjects(InObjectBindingID, Sequencer.GetFocusedTemplateID()))
 		{
 			if (UObject* Object = WeakObject.Get())
 			{
 				Value = ExternalValue.OnGetExternalValue(*Object, PropertyBindings);
+				if (Value.IsSet() && ExternalValue.OnGetCurrentValueAndWeight && SectionToKey)
+				{
+					CurrentValue = Value.Get(0.0f);
+					ExternalValue.OnGetCurrentValueAndWeight(Object, SectionToKey, InTime, Sequencer.GetFocusedTickResolution(), Sequencer.GetEvaluationTemplate(),CurrentValue,CurrentWeight);
+				}
 				break;
 			}
 		}
 	}
 
 	float NewValue = Channel->GetDefault().Get(0.f);
-	Channel->Evaluate(InTime, NewValue);
-	
+	bool bWasEvaluated = Channel->Evaluate(InTime, NewValue);
+	if (Value.IsSet()) //need to get the diff between Value(Global) and CurrentValue and apply that to the local
+	{
+		if (bWasEvaluated)
+		{
+			float CurrentGlobalValue = Value.GetValue();
+			NewValue = (Value.Get(0.0f) - CurrentValue) * CurrentWeight + NewValue;
+		}
+		else //Nothing set (key or default) on channel so use external value
+		{
+			NewValue = Value.Get(0.0f);
+		}
+	}
 
 	using namespace MovieScene;
 	return AddKeyToChannel(Channel, InTime, NewValue, Sequencer.GetKeyInterpolation());
 }
 
-FKeyHandle AddOrUpdateKey(FMovieSceneActorReferenceData* Channel, FFrameNumber InTime, ISequencer& Sequencer, const FGuid& InObjectBindingID, FTrackInstancePropertyBindings* PropertyBindings)
+FKeyHandle AddOrUpdateKey(FMovieSceneActorReferenceData* Channel, UMovieSceneSection* SectionToKey, FFrameNumber InTime, ISequencer& Sequencer, const FGuid& InObjectBindingID, FTrackInstancePropertyBindings* PropertyBindings)
 {
 	AActor* CurrentActor = nullptr;
 
diff --git a/Engine/Source/Editor/MovieSceneTools/Private/Channels/BuiltInChannelEditors.h b/Engine/Source/Editor/MovieSceneTools/Private/Channels/BuiltInChannelEditors.h
--- a/Engine/Source/Editor/MovieSceneTools/Private/Channels/BuiltInChannelEditors.h
+++ b/Engine/Source/Editor/MovieSceneTools/Private/Channels/BuiltInChannelEditors.h
@@ -34,8 +34,8 @@ class ISectionLayoutBuilder;
 template<typename> class  TArrayView;
 
 /** Overrides for adding or updating a key for non-standard channels */
-FKeyHandle AddOrUpdateKey(FMovieSceneFloatChannel* Channel, const TMovieSceneExternalValue<float>& EditorData, FFrameNumber InTime, ISequencer& Sequencer, const FGuid& InObjectBindingID, FTrackInstancePropertyBindings* PropertyBindings);
-FKeyHandle AddOrUpdateKey(FMovieSceneActorReferenceData* Channel, FFrameNumber InTime, ISequencer& Sequencer, const FGuid& InObjectBindingID, FTrackInstancePropertyBindings* PropertyBindings);
+FKeyHandle AddOrUpdateKey(FMovieSceneFloatChannel* Channel, UMovieSceneSection* SectionToKey,  const TMovieSceneExternalValue<float>& EditorData, FFrameNumber InTime, ISequencer& Sequencer, const FGuid& InObjectBindingID, FTrackInstancePropertyBindings* PropertyBindings);
+FKeyHandle AddOrUpdateKey(FMovieSceneActorReferenceData* Channel, UMovieSceneSection* SectionToKey, FFrameNumber InTime, ISequencer& Sequencer, const FGuid& InObjectBindingID, FTrackInstancePropertyBindings* PropertyBindings);
 
 /** Key editor overrides */
 bool CanCreateKeyEditor(const FMovieSceneBoolChannel*       Channel);
diff --git a/Engine/Source/Editor/MovieSceneTools/Private/CurveKeyEditors/SequencerKeyEditor.h b/Engine/Source/Editor/MovieSceneTools/Private/CurveKeyEditors/SequencerKeyEditor.h
--- a/Engine/Source/Editor/MovieSceneTools/Private/CurveKeyEditors/SequencerKeyEditor.h
+++ b/Engine/Source/Editor/MovieSceneTools/Private/CurveKeyEditors/SequencerKeyEditor.h
@@ -74,7 +74,17 @@ struct TSequencerKeyEditor
 		if (Channel && Sequencer && OwningSection)
 		{
 			const FFrameTime CurrentTime = MovieScene::ClampToDiscreteRange(Sequencer->GetLocalTime().Time, OwningSection->GetRange());
-			EvaluateChannel(Channel, CurrentTime, Result);
+			//If we have no keys and no default, key with the external value if it exists
+			if (!EvaluateChannel(Channel, CurrentTime, Result))
+			{
+				if (TOptional<ValueType> ExternalValue = GetExternalValue())
+				{
+					if (ExternalValue.IsSet())
+					{
+						Result = ExternalValue.GetValue();
+					}
+				}
+			}
 		}
 
 		return Result;
diff --git a/Engine/Source/Editor/MovieSceneTools/Private/MovieSceneToolHelpers.cpp b/Engine/Source/Editor/MovieSceneTools/Private/MovieSceneToolHelpers.cpp
--- a/Engine/Source/Editor/MovieSceneTools/Private/MovieSceneToolHelpers.cpp
+++ b/Engine/Source/Editor/MovieSceneTools/Private/MovieSceneToolHelpers.cpp
@@ -1939,14 +1939,5 @@ FMovieSceneEvaluationTrack* MovieSceneToolHelpers::GetEvaluationTrack(ISequencer
 			return EvalTrack;
 		}
 	}
-	Sequencer->NotifyMovieSceneDataChanged(EMovieSceneDataChangeType::TrackValueChangedRefreshImmediately);
-	Template = Sequencer->GetEvaluationTemplate().FindTemplate(Sequencer->GetFocusedTemplateID());
-	if (Template)
-	{
-		if (FMovieSceneEvaluationTrack* EvalTrack = Template->FindTrack(TrackSignature))
-		{
-			return EvalTrack;
-		}
-	}
 	return nullptr;
 }
diff --git a/Engine/Source/Editor/MovieSceneTools/Private/TrackEditors/TransformTrackEditor.cpp b/Engine/Source/Editor/MovieSceneTools/Private/TrackEditors/TransformTrackEditor.cpp
--- a/Engine/Source/Editor/MovieSceneTools/Private/TrackEditors/TransformTrackEditor.cpp
+++ b/Engine/Source/Editor/MovieSceneTools/Private/TrackEditors/TransformTrackEditor.cpp
@@ -852,7 +852,7 @@ bool F3DTransformTrackEditor::ModifyGeneratedKeysByCurrentAndWeight(UObject *Obj
 			
 	FVector CurrentPos; FRotator CurrentRot;
 	FVector CurrentScale;
-	for (const FTransform& Transform : InterrogationData.Iterate<FTransform>(UMovieScene3DTransformTrack::GetInterrogationKey()))
+	for (const FTransform& Transform : InterrogationData.Iterate<FTransform>(UMovieScene3DTransformSection::GetInterrogationKey()))
 	{
 		CurrentPos = Transform.GetTranslation();
 		CurrentRot = Transform.Rotator();
diff --git a/Engine/Source/Editor/MovieSceneTools/Public/KeyframeTrackEditor.h b/Engine/Source/Editor/MovieSceneTools/Public/KeyframeTrackEditor.h
--- a/Engine/Source/Editor/MovieSceneTools/Public/KeyframeTrackEditor.h
+++ b/Engine/Source/Editor/MovieSceneTools/Public/KeyframeTrackEditor.h
@@ -35,30 +35,6 @@ struct IImpl
 
 };
 
-template<typename ChannelType, typename ValueType>
-struct TSetDefaultImpl : IImpl
-{
-	int32 ChannelIndex;
-	ValueType ValueToSet;
-
-	TSetDefaultImpl(int32 InChannelIndex, const ValueType& InValue)
-		: ChannelIndex(InChannelIndex), ValueToSet(InValue)
-	{}
-
-	virtual void ApplyDefault(UMovieSceneSection* Section, FMovieSceneChannelProxy& Proxy) const override
-	{
-		ChannelType* Channel = Proxy.GetChannel<ChannelType>(ChannelIndex);
-		if (Channel && Channel->GetData().GetTimes().Num() == 0)
-		{
-			if (Section->TryModify())
-			{
-				using namespace MovieScene;
-				SetChannelDefault(Channel, ValueToSet);
-			}
-		}
-	}
-};
-
 //If this implementation changes need to change the specializations below.
 template<typename ChannelType, typename ValueType>
 struct TAddKeyImpl : IImpl
@@ -286,14 +262,8 @@ struct FMovieSceneChannelValueSetter
 	static FMovieSceneChannelValueSetter Create(int32 ChannelIndex, ValueType InNewValue, bool bAddKey)
 	{
 		FMovieSceneChannelValueSetter NewValue;
-		if (bAddKey)
-		{
-			NewValue.Impl = TAddKeyImpl<ChannelType, typename TDecay<ValueType>::Type>(ChannelIndex, Forward<ValueType>(InNewValue));
-		}
-		else
-		{
-			NewValue.Impl = TSetDefaultImpl<ChannelType, typename TDecay<ValueType>::Type>(ChannelIndex, Forward<ValueType>(InNewValue));
-		}
+		NewValue.Impl = TAddKeyImpl<ChannelType, typename TDecay<ValueType>::Type>(ChannelIndex, Forward<ValueType>(InNewValue));
+		
 		return MoveTemp(NewValue);
 	}
 
diff --git a/Engine/Source/Editor/Sequencer/Private/IKeyArea.cpp b/Engine/Source/Editor/Sequencer/Private/IKeyArea.cpp
--- a/Engine/Source/Editor/Sequencer/Private/IKeyArea.cpp
+++ b/Engine/Source/Editor/Sequencer/Private/IKeyArea.cpp
@@ -69,7 +69,7 @@ FKeyHandle IKeyArea::AddOrUpdateKey(FFrameNumber Time, const FGuid& ObjectBindin
 	if (EditorInterface && Channel)
 	{
 		FTrackInstancePropertyBindings* BindingsPtr = PropertyBindings.IsSet() ? &PropertyBindings.GetValue() : nullptr;
-		return EditorInterface->AddOrUpdateKey_Raw(Channel, RawExtendedData, Time, InSequencer, ObjectBindingID, BindingsPtr);
+		return EditorInterface->AddOrUpdateKey_Raw(Channel, WeakOwningSection.Get(), RawExtendedData, Time, InSequencer, ObjectBindingID, BindingsPtr);
 	}
 
 	return FKeyHandle();
diff --git a/Engine/Source/Editor/Sequencer/Private/SequencerEdMode.cpp b/Engine/Source/Editor/Sequencer/Private/SequencerEdMode.cpp
--- a/Engine/Source/Editor/Sequencer/Private/SequencerEdMode.cpp
+++ b/Engine/Source/Editor/Sequencer/Private/SequencerEdMode.cpp
@@ -472,7 +472,7 @@ void FSequencerEdMode::GetLocationAtTime(FMovieSceneEvaluationTrack* Track, UObj
 	FMovieSceneContext Context(FMovieSceneEvaluationRange(KeyTime, Sequencer->GetFocusedTickResolution()));
 	Track->Interrogate(Context, InterrogationData, Object);
 
-	for (const FTransform& Transform : InterrogationData.Iterate<FTransform>(UMovieScene3DTransformTrack::GetInterrogationKey()))
+	for (const FTransform& Transform : InterrogationData.Iterate<FTransform>(UMovieScene3DTransformSection::GetInterrogationKey()))
 	{
 		KeyPos = Transform.GetTranslation();
 		KeyRot = Transform.GetRotation().Rotator();
diff --git a/Engine/Source/Editor/Sequencer/Public/ISequencerChannelInterface.h b/Engine/Source/Editor/Sequencer/Public/ISequencerChannelInterface.h
--- a/Engine/Source/Editor/Sequencer/Public/ISequencerChannelInterface.h
+++ b/Engine/Source/Editor/Sequencer/Public/ISequencerChannelInterface.h
@@ -51,6 +51,7 @@ struct ISequencerChannelInterface
 	 * Add (or update) a key to the specified channel using it's current value at that time, or some external value specified by the extended editor data
 	 *
 	 * @param Channel               The channel to add a key to
+	 * @param SectionToKey          The SectionToKey
 	 * @param ExtendedEditorData    A pointer to the extended editor data for this channel of type TMovieSceneChannelTraits<>::ExtendedEditorDataType
 	 * @param InTime                The time at which to add a key
 	 * @param InSequencer           The currently active sequencer
@@ -58,7 +59,7 @@ struct ISequencerChannelInterface
 	 * @param PropertyBindings      (Optional) Property bindings where this channel exists on a property track
 	 * @return A handle to the new or updated key
 	 */
-	virtual FKeyHandle AddOrUpdateKey_Raw(FMovieSceneChannel* Channel, const void* ExtendedEditorData, FFrameNumber InTime, ISequencer& InSequencer, const FGuid& ObjectBindingID, FTrackInstancePropertyBindings* PropertyBindings) const = 0;
+	virtual FKeyHandle AddOrUpdateKey_Raw(FMovieSceneChannel* Channel, UMovieSceneSection* SectionToKey, const void* ExtendedEditorData, FFrameNumber InTime, ISequencer& InSequencer, const FGuid& ObjectBindingID, FTrackInstancePropertyBindings* PropertyBindings) const = 0;
 
 	/**
 	 * Copy all the keys specified in KeyMask to the specified clipboard
diff --git a/Engine/Source/Editor/Sequencer/Public/SequencerChannelInterface.h b/Engine/Source/Editor/Sequencer/Public/SequencerChannelInterface.h
--- a/Engine/Source/Editor/Sequencer/Public/SequencerChannelInterface.h
+++ b/Engine/Source/Editor/Sequencer/Public/SequencerChannelInterface.h
@@ -180,6 +180,7 @@ struct TSequencerChannelInterfaceBase<ChannelType, false> : TSequencerChannelInt
 	 * Add (or update) a key to the specified channel using it's current value at that time, or some external value specified by the extended editor data
 	 *
 	 * @param Channel               The channel to add a key to
+	 * @param SectionToKey          The Section to Key
 	 * @param ExtendedEditorData    A pointer to the extended editor data for this channel of type TMovieSceneChannelTraits<>::ExtendedEditorDataType
 	 * @param InTime                The time at which to add a key
 	 * @param InSequencer           The currently active sequencer
@@ -187,10 +188,10 @@ struct TSequencerChannelInterfaceBase<ChannelType, false> : TSequencerChannelInt
 	 * @param PropertyBindings      (Optional) Property bindings where this channel exists on a property track
 	 * @return A handle to the new or updated key
 	 */
-	virtual FKeyHandle AddOrUpdateKey_Raw(FMovieSceneChannel* InChannel, const void* ExtendedEditorData, FFrameNumber InTime, ISequencer& Sequencer, const FGuid& ObjectBindingID, FTrackInstancePropertyBindings* PropertyBindings) const override
+	virtual FKeyHandle AddOrUpdateKey_Raw(FMovieSceneChannel* InChannel, UMovieSceneSection* SectionToKey, const void* ExtendedEditorData, FFrameNumber InTime, ISequencer& Sequencer, const FGuid& ObjectBindingID, FTrackInstancePropertyBindings* PropertyBindings) const override
 	{
 		using namespace Sequencer;
-		return AddOrUpdateKey(static_cast<ChannelType*>(InChannel), InTime, Sequencer, ObjectBindingID, PropertyBindings);
+		return AddOrUpdateKey(static_cast<ChannelType*>(InChannel), SectionToKey, InTime, Sequencer, ObjectBindingID, PropertyBindings);
 	}
 };
 
@@ -206,6 +207,7 @@ struct TSequencerChannelInterfaceBase<ChannelType, true> : TSequencerChannelInte
 	 * Add (or update) a key to the specified channel using it's current value at that time, or some external value specified by the extended editor data
 	 *
 	 * @param Channel               The channel to add a key to
+	 * @param SectionToKey          The Section to Key
 	 * @param ExtendedEditorData    A pointer to the extended editor data for this channel of type TMovieSceneChannelTraits<>::ExtendedEditorDataType
 	 * @param InTime                The time at which to add a key
 	 * @param InSequencer           The currently active sequencer
@@ -213,7 +215,7 @@ struct TSequencerChannelInterfaceBase<ChannelType, true> : TSequencerChannelInte
 	 * @param PropertyBindings      (Optional) Property bindings where this channel exists on a property track
 	 * @return A handle to the new or updated key
 	 */
-	virtual FKeyHandle AddOrUpdateKey_Raw(FMovieSceneChannel* InChannel, const void* ExtendedEditorData, FFrameNumber InTime, ISequencer& Sequencer, const FGuid& ObjectBindingID, FTrackInstancePropertyBindings* PropertyBindings) const override
+	virtual FKeyHandle AddOrUpdateKey_Raw(FMovieSceneChannel* InChannel, UMovieSceneSection* SectionToKey, const void* ExtendedEditorData, FFrameNumber InTime, ISequencer& Sequencer, const FGuid& ObjectBindingID, FTrackInstancePropertyBindings* PropertyBindings) const override
 	{
 		using namespace Sequencer;
 
@@ -221,7 +223,7 @@ struct TSequencerChannelInterfaceBase<ChannelType, true> : TSequencerChannelInte
 		check(ExtendedEditorData);
 
 		const auto* TypedEditorData = static_cast<const ExtendedEditorDataType*>(ExtendedEditorData);
-		return AddOrUpdateKey(static_cast<ChannelType*>(InChannel), *TypedEditorData, InTime, Sequencer, ObjectBindingID, PropertyBindings);
+		return AddOrUpdateKey(static_cast<ChannelType*>(InChannel), SectionToKey, *TypedEditorData, InTime, Sequencer, ObjectBindingID, PropertyBindings);
 	}
 };
 
diff --git a/Engine/Source/Editor/Sequencer/Public/SequencerChannelTraits.h b/Engine/Source/Editor/Sequencer/Public/SequencerChannelTraits.h
--- a/Engine/Source/Editor/Sequencer/Public/SequencerChannelTraits.h
+++ b/Engine/Source/Editor/Sequencer/Public/SequencerChannelTraits.h
@@ -183,6 +183,7 @@ namespace Sequencer
 	 * Add or update a key for this channel's current value
 	 *
 	 * @param InChannel          The channel to create a key for
+	 * @param InSectionToKey     The Section to key
 	 * @param InTime             The time at which to add a key
 	 * @param InSequencer        The currently active sequencer
 	 * @param InObjectBindingID  The object binding ID that this section's track is bound to
@@ -192,6 +193,7 @@ namespace Sequencer
 	template<typename ChannelType>
 	FKeyHandle AddOrUpdateKey(
 		ChannelType*                    InChannel,
+		UMovieSceneSection*             InSectionToKey,
 		FFrameNumber                    InTime,
 		ISequencer&                     InSequencer,
 		const FGuid&                    InObjectBindingID,
@@ -206,6 +208,7 @@ namespace Sequencer
 	 * Add or update a key for this channel's current value, using an external value if possible
 	 *
 	 * @param InChannel          The channel to create a key for
+	 * @param InSectionToKey     The Section to key
 	 * @param InExternalValue    The external value definition
 	 * @param InTime             The time at which to add a key
 	 * @param InSequencer        The currently active sequencer
@@ -216,6 +219,7 @@ namespace Sequencer
 	template<typename ChannelType, typename ValueType>
 	FKeyHandle AddOrUpdateKey(
 		ChannelType*                               InChannel,
+		UMovieSceneSection*             SectionToKey,
 		const TMovieSceneExternalValue<ValueType>& InExternalValue,
 		FFrameNumber                               InTime,
 		ISequencer&                                InSequencer,
diff --git a/Engine/Source/Runtime/MovieScene/Private/MovieSceneCommonHelpers.cpp b/Engine/Source/Runtime/MovieScene/Private/MovieSceneCommonHelpers.cpp
--- a/Engine/Source/Runtime/MovieScene/Private/MovieSceneCommonHelpers.cpp
+++ b/Engine/Source/Runtime/MovieScene/Private/MovieSceneCommonHelpers.cpp
@@ -13,6 +13,7 @@
 #include "Sound/SoundWave.h"
 #include "Sound/SoundCue.h"
 #include "Sound/SoundNodeWavePlayer.h"
+#include "MovieSceneTrack.h"
 
 UMovieSceneSection* MovieSceneHelpers::FindSectionAtTime( const TArray<UMovieSceneSection*>& Sections, FFrameNumber Time )
 {
@@ -30,7 +31,6 @@ UMovieSceneSection* MovieSceneHelpers::FindSectionAtTime( const TArray<UMovieSce
 	return nullptr;
 }
 
-
 UMovieSceneSection* MovieSceneHelpers::FindNearestSectionAtTime( const TArray<UMovieSceneSection*>& Sections, FFrameNumber Time )
 {
 	TArray<UMovieSceneSection*> OverlappingSections, NonOverlappingSections;
@@ -258,6 +258,50 @@ float MovieSceneHelpers::GetSoundDuration(USoundBase* Sound)
 	return Duration == INDEFINITELY_LOOPING_DURATION ? SoundWave->Duration : Duration;
 }
 
+
+float MovieSceneHelpers::CalculateWeightForBlending(UMovieSceneSection* SectionToKey, FFrameNumber Time)
+{
+	float Weight = 1.0f;
+	UMovieSceneTrack* Track = SectionToKey->GetTypedOuter<UMovieSceneTrack>();
+	FOptionalMovieSceneBlendType BlendType = SectionToKey->GetBlendType();
+	if (Track && BlendType.IsValid() && (BlendType.Get() == EMovieSceneBlendType::Additive || BlendType.Get() == EMovieSceneBlendType::Absolute))
+	{
+		//if additive weight is just the inverse of any weight on it
+		if (BlendType.Get() == EMovieSceneBlendType::Additive)
+		{
+			float TotalWeightValue = SectionToKey->GetTotalWeightValue(Time);
+			Weight = !FMath::IsNearlyZero(TotalWeightValue) ? 1.0f / TotalWeightValue : 0.0f;
+		}
+		else
+		{
+
+			const TArray<UMovieSceneSection*>& Sections = Track->GetAllSections();
+			TArray<UMovieSceneSection*, TInlineAllocator<4>> OverlappingSections;
+			for (UMovieSceneSection* Section : Sections)
+			{
+				if (Section->GetRange().Contains(Time))
+				{
+					OverlappingSections.Add(Section);
+				}
+			}
+			//if absolute need to calculate weight based upon other sections weights (+ implicit absolute weights)
+			int TotalNumOfAbsoluteSections = 1;
+			for (UMovieSceneSection* Section : OverlappingSections)
+			{
+				FOptionalMovieSceneBlendType NewBlendType = Section->GetBlendType();
+
+				if (Section != SectionToKey && NewBlendType.IsValid() && NewBlendType.Get() == EMovieSceneBlendType::Absolute)
+				{
+					++TotalNumOfAbsoluteSections;
+				}
+			}
+			float TotalWeightValue = SectionToKey->GetTotalWeightValue(Time);
+			Weight = !FMath::IsNearlyZero(TotalWeightValue) ? float(TotalNumOfAbsoluteSections) / TotalWeightValue : 0.0f;
+		}
+	}
+	return Weight;
+}
+
 FTrackInstancePropertyBindings::FTrackInstancePropertyBindings( FName InPropertyName, const FString& InPropertyPath, const FName& InFunctionName, const FName& InNotifyFunctionName )
     : PropertyPath( InPropertyPath )
 	, NotifyFunctionName(InNotifyFunctionName)
diff --git a/Engine/Source/Runtime/MovieScene/Public/Channels/MovieSceneChannelEditorData.h b/Engine/Source/Runtime/MovieScene/Public/Channels/MovieSceneChannelEditorData.h
--- a/Engine/Source/Runtime/MovieScene/Public/Channels/MovieSceneChannelEditorData.h
+++ b/Engine/Source/Runtime/MovieScene/Public/Channels/MovieSceneChannelEditorData.h
@@ -8,6 +8,7 @@
 #include "Math/Color.h"
 #include "Templates/Function.h"
 #include "MovieSceneCommonHelpers.h"
+#include "Evaluation/MovieSceneEvaluationTemplateInstance.h"
 
 #if WITH_EDITOR
 
@@ -95,6 +96,9 @@ struct TMovieSceneExternalValue
 
 	/** Function to invoke to get the current value of the property of an object */
 	TFunction<TOptional<T>(UObject&, FTrackInstancePropertyBindings*)> OnGetExternalValue;
+
+	/** Optional Function To Get Current Value and Weight, needed for setting keys on blended sections */
+	TFunction<void (UObject*, UMovieSceneSection*,  FFrameNumber, FFrameRate, FMovieSceneRootEvaluationTemplateInstance&, T&, float&) > OnGetCurrentValueAndWeight;
 };
 
 
diff --git a/Engine/Source/Runtime/MovieScene/Public/MovieSceneCommonHelpers.h b/Engine/Source/Runtime/MovieScene/Public/MovieSceneCommonHelpers.h
--- a/Engine/Source/Runtime/MovieScene/Public/MovieSceneCommonHelpers.h
+++ b/Engine/Source/Runtime/MovieScene/Public/MovieSceneCommonHelpers.h
@@ -118,6 +118,14 @@ public:
 	 * Sort predicate that sorts overlapping sections by row primarily, then by overlap priority
 	 */
 	static bool SortOverlappingSections(const UMovieSceneSection* A, const UMovieSceneSection* B);
+
+	/*
+	* Get weight needed to modify the global difference in order to correctly key this section due to it possibly being blended by other sections.
+	* @param Section The Section who's weight we are calculating.
+	* @param  Time we are at.
+	* @return Returns the weight that needs to be applied to the global difference to correctly key this section.
+	*/
+	static float CalculateWeightForBlending(UMovieSceneSection* SectionToKey, FFrameNumber Time);
 };
 
 /**
diff --git a/Engine/Source/Runtime/MovieSceneTracks/Private/Evaluation/MovieScene3DTransformTemplate.cpp b/Engine/Source/Runtime/MovieSceneTracks/Private/Evaluation/MovieScene3DTransformTemplate.cpp
--- a/Engine/Source/Runtime/MovieSceneTracks/Private/Evaluation/MovieScene3DTransformTemplate.cpp
+++ b/Engine/Source/Runtime/MovieSceneTracks/Private/Evaluation/MovieScene3DTransformTemplate.cpp
@@ -92,7 +92,7 @@ struct FComponentTransformActuator : TMovieSceneBlendingActuator<F3DTransformTra
 
 	virtual void Actuate(FMovieSceneInterrogationData& InterrogationData, const F3DTransformTrackToken& InValue, const TBlendableTokenStack<F3DTransformTrackToken>& OriginalStack, const FMovieSceneContext& Context) const override
 	{
-		InterrogationData.Add(FTransform(InValue.Rotation.Quaternion(), InValue.Translation, InValue.Scale), UMovieScene3DTransformTrack::GetInterrogationKey());
+		InterrogationData.Add(FTransform(InValue.Rotation.Quaternion(), InValue.Translation, InValue.Scale), UMovieScene3DTransformSection::GetInterrogationKey());
 	}
 };
 
diff --git a/Engine/Source/Runtime/MovieSceneTracks/Private/Sections/MovieScene3DTransformSection.cpp b/Engine/Source/Runtime/MovieSceneTracks/Private/Sections/MovieScene3DTransformSection.cpp
--- a/Engine/Source/Runtime/MovieSceneTracks/Private/Sections/MovieScene3DTransformSection.cpp
+++ b/Engine/Source/Runtime/MovieSceneTracks/Private/Sections/MovieScene3DTransformSection.cpp
@@ -8,6 +8,10 @@
 #include "Channels/MovieSceneChannelProxy.h"
 #include "GameFramework/Actor.h"
 #include "EulerTransform.h"
+#include "Compilation/MovieSceneTemplateInterrogation.h"
+#include "Evaluation/MovieSceneEvaluationTrack.h"
+#include "Evaluation/MovieSceneEvaluationTemplateInstance.h"
+#include "Evaluation/MovieScenePropertyTemplate.h"
 
 #if WITH_EDITOR
 
@@ -89,6 +93,17 @@ struct F3DTransformChannelEditorData
 		ExternalValues[6].OnGetExternalValue = ExtractScaleX;
 		ExternalValues[7].OnGetExternalValue = ExtractScaleY;
 		ExternalValues[8].OnGetExternalValue = ExtractScaleZ;
+
+		ExternalValues[0].OnGetCurrentValueAndWeight = GetTranslationXValueAndWeight;
+		ExternalValues[1].OnGetCurrentValueAndWeight = GetTranslationYValueAndWeight;
+		ExternalValues[2].OnGetCurrentValueAndWeight = GetTranslationZValueAndWeight;
+		ExternalValues[3].OnGetCurrentValueAndWeight = GetRotationXValueAndWeight;
+		ExternalValues[4].OnGetCurrentValueAndWeight = GetRotationYValueAndWeight;
+		ExternalValues[5].OnGetCurrentValueAndWeight = GetRotationZValueAndWeight;
+		ExternalValues[6].OnGetCurrentValueAndWeight = GetScaleXValueAndWeight;
+		ExternalValues[7].OnGetCurrentValueAndWeight = GetScaleYValueAndWeight;
+		ExternalValues[8].OnGetCurrentValueAndWeight = GetScaleZValueAndWeight;
+
 	}
 
 	static TOptional<FVector> GetTranslation(UObject& InObject, FTrackInstancePropertyBindings* Bindings)
@@ -187,6 +202,62 @@ struct F3DTransformChannelEditorData
 		return TOptional<FVector>();
 	}
 
+	static void GetValueAndWeight(UObject* Object, UMovieSceneSection*  SectionToKey, int32 Index, FFrameNumber KeyTime, FFrameRate TickResolution, FMovieSceneRootEvaluationTemplateInstance& RootTemplate,
+		float& OutValue, float& OutWeight)
+	{
+		UMovieSceneTrack* Track = SectionToKey->GetTypedOuter<UMovieSceneTrack>();
+		FMovieSceneEvaluationTrack EvalTrack = Track->GenerateTrackTemplate();
+		FMovieSceneInterrogationData InterrogationData;
+		RootTemplate.CopyActuators(InterrogationData.GetAccumulator());
+
+		FMovieSceneContext Context(FMovieSceneEvaluationRange(KeyTime, TickResolution));
+		EvalTrack.Interrogate(Context, InterrogationData, Object);
+
+		FVector CurrentPos; FRotator CurrentRot;
+		FVector CurrentScale;
+		for (const FTransform& Transform : InterrogationData.Iterate<FTransform>(UMovieScene3DTransformSection::GetInterrogationKey()))
+		{
+			CurrentPos = Transform.GetTranslation();
+			CurrentRot = Transform.Rotator();
+			CurrentScale = Transform.GetScale3D();
+			break;
+		}
+
+		switch (Index)
+		{
+		case 0:
+			OutValue = CurrentPos.X;
+			break;
+		case 1:
+			OutValue = CurrentPos.Y;
+			break;
+		case 2:
+			OutValue = CurrentPos.Z;
+			break;
+		case 3:
+			OutValue = CurrentRot.Roll;
+			break;
+		case 4:
+			OutValue = CurrentRot.Pitch;
+			break;
+		case 5:
+			OutValue = CurrentRot.Yaw;
+			break;
+		case 6:
+			OutValue = CurrentScale.X;
+			break;
+		case 7:
+			OutValue = CurrentScale.Y;
+			break;
+		case 8:
+			OutValue = CurrentScale.Z;
+			break;
+
+		}
+		OutWeight = MovieSceneHelpers::CalculateWeightForBlending(SectionToKey, KeyTime);
+	}
+
+
 	static TOptional<float> ExtractTranslationX(UObject& InObject, FTrackInstancePropertyBindings* Bindings)
 	{
 		TOptional<FVector> Translation = GetTranslation(InObject, Bindings);
@@ -235,6 +306,53 @@ struct F3DTransformChannelEditorData
 		return Scale.IsSet() ? Scale->Z : TOptional<float>();
 	}
 
+	static void GetTranslationXValueAndWeight(UObject* Object, UMovieSceneSection*  SectionToKey, FFrameNumber KeyTime, FFrameRate TickResolution, FMovieSceneRootEvaluationTemplateInstance& RootTemplate,
+		float& OutValue, float& OutWeight)
+	{
+		GetValueAndWeight(Object, SectionToKey, 0, KeyTime, TickResolution, RootTemplate, OutValue, OutWeight);
+	}
+	static void GetTranslationYValueAndWeight(UObject* Object, UMovieSceneSection*  SectionToKey,  FFrameNumber KeyTime, FFrameRate TickResolution, FMovieSceneRootEvaluationTemplateInstance& RootTemplate,
+		float& OutValue, float& OutWeight)
+	{
+		GetValueAndWeight(Object, SectionToKey, 1, KeyTime, TickResolution, RootTemplate, OutValue, OutWeight);
+	}
+	static void GetTranslationZValueAndWeight(UObject* Object, UMovieSceneSection*  SectionToKey, FFrameNumber KeyTime, FFrameRate TickResolution, FMovieSceneRootEvaluationTemplateInstance& RootTemplate,
+		float& OutValue, float& OutWeight)
+	{
+		GetValueAndWeight(Object, SectionToKey, 2, KeyTime, TickResolution, RootTemplate, OutValue, OutWeight);
+	}
+	static void GetRotationXValueAndWeight(UObject* Object, UMovieSceneSection*  SectionToKey, FFrameNumber KeyTime, FFrameRate TickResolution, FMovieSceneRootEvaluationTemplateInstance& RootTemplate,
+		float& OutValue, float& OutWeight)
+	{
+		GetValueAndWeight(Object, SectionToKey, 3, KeyTime, TickResolution, RootTemplate, OutValue, OutWeight);
+	}
+	static void GetRotationYValueAndWeight(UObject* Object, UMovieSceneSection*  SectionToKey, FFrameNumber KeyTime, FFrameRate TickResolution, FMovieSceneRootEvaluationTemplateInstance& RootTemplate,
+		float& OutValue, float& OutWeight)
+	{
+		GetValueAndWeight(Object, SectionToKey, 4, KeyTime, TickResolution, RootTemplate, OutValue, OutWeight);
+	}
+	static void GetRotationZValueAndWeight(UObject* Object, UMovieSceneSection*  SectionToKey, FFrameNumber KeyTime, FFrameRate TickResolution, FMovieSceneRootEvaluationTemplateInstance& RootTemplate,
+		float& OutValue, float& OutWeight)
+	{
+		GetValueAndWeight(Object, SectionToKey, 5, KeyTime, TickResolution, RootTemplate, OutValue, OutWeight);
+	}
+	static void GetScaleXValueAndWeight(UObject* Object, UMovieSceneSection*  SectionToKey, FFrameNumber KeyTime, FFrameRate TickResolution, FMovieSceneRootEvaluationTemplateInstance& RootTemplate,
+		float& OutValue, float& OutWeight)
+	{
+		GetValueAndWeight(Object, SectionToKey, 6, KeyTime, TickResolution, RootTemplate, OutValue, OutWeight);
+	}
+	static void GetScaleYValueAndWeight(UObject* Object, UMovieSceneSection*  SectionToKey, FFrameNumber KeyTime, FFrameRate TickResolution, FMovieSceneRootEvaluationTemplateInstance& RootTemplate,
+		float& OutValue, float& OutWeight)
+	{
+		GetValueAndWeight(Object, SectionToKey, 7, KeyTime, TickResolution, RootTemplate, OutValue, OutWeight);
+	}
+	static void GetScaleZValueAndWeight(UObject* Object, UMovieSceneSection*  SectionToKey, FFrameNumber KeyTime, FFrameRate TickResolution, FMovieSceneRootEvaluationTemplateInstance& RootTemplate,
+		float& OutValue, float& OutWeight)
+	{
+		GetValueAndWeight(Object, SectionToKey, 8, KeyTime, TickResolution, RootTemplate, OutValue, OutWeight);
+	}
+
+
 	FMovieSceneChannelMetaData      MetaData[10];
 	TMovieSceneExternalValue<float> ExternalValues[10];
 };
@@ -597,3 +715,8 @@ void UMovieScene3DTransformSection::SetBlendType(EMovieSceneBlendType InBlendTyp
 
 	}
 }
+FMovieSceneInterrogationKey UMovieScene3DTransformSection::GetInterrogationKey()
+{
+	static FMovieSceneAnimTypeID TypeID = FMovieSceneAnimTypeID::Unique();
+	return TypeID;
+}
\ No newline at end of file
diff --git a/Engine/Source/Runtime/MovieSceneTracks/Private/Sections/MovieSceneCameraCutSection.cpp b/Engine/Source/Runtime/MovieSceneTracks/Private/Sections/MovieSceneCameraCutSection.cpp
--- a/Engine/Source/Runtime/MovieSceneTracks/Private/Sections/MovieSceneCameraCutSection.cpp
+++ b/Engine/Source/Runtime/MovieSceneTracks/Private/Sections/MovieSceneCameraCutSection.cpp
@@ -36,7 +36,7 @@ FMovieSceneEvalTemplatePtr UMovieSceneCameraCutSection::GenerateTemplate() const
 					FMovieSceneInterrogationData Container;
 					TransformTrackTemplate.Interrogate(Context, Container);
 
-					for (auto It = Container.Iterate<FTransform>(UMovieScene3DTransformTrack::GetInterrogationKey()); It; ++It)
+					for (auto It = Container.Iterate<FTransform>(UMovieScene3DTransformSection::GetInterrogationKey()); It; ++It)
 					{
 						CutTransform = *It;
 						break;
diff --git a/Engine/Source/Runtime/MovieSceneTracks/Private/Sections/MovieSceneVectorSection.cpp b/Engine/Source/Runtime/MovieSceneTracks/Private/Sections/MovieSceneVectorSection.cpp
--- a/Engine/Source/Runtime/MovieSceneTracks/Private/Sections/MovieSceneVectorSection.cpp
+++ b/Engine/Source/Runtime/MovieSceneTracks/Private/Sections/MovieSceneVectorSection.cpp
@@ -4,6 +4,10 @@
 #include "UObject/StructOnScope.h"
 #include "UObject/SequencerObjectVersion.h"
 #include "Channels/MovieSceneChannelProxy.h"
+#include "Compilation/MovieSceneTemplateInterrogation.h"
+#include "Evaluation/MovieSceneEvaluationTrack.h"
+#include "Evaluation/MovieSceneEvaluationTemplateInstance.h"
+#include "Evaluation/MovieScenePropertyTemplate.h"
 
 /* FMovieSceneVectorKeyStruct interface
  *****************************************************************************/
@@ -36,6 +40,17 @@ struct FVectorSectionEditorData
 		ExternalValues[1].OnGetExternalValue = [NumChannels](UObject& InObject, FTrackInstancePropertyBindings* Bindings) { return ExtractChannelY(InObject, Bindings, NumChannels); };
 		ExternalValues[2].OnGetExternalValue = [NumChannels](UObject& InObject, FTrackInstancePropertyBindings* Bindings) { return ExtractChannelZ(InObject, Bindings, NumChannels); };
 		ExternalValues[3].OnGetExternalValue = [NumChannels](UObject& InObject, FTrackInstancePropertyBindings* Bindings) { return ExtractChannelW(InObject, Bindings, NumChannels); };
+
+
+		ExternalValues[0].OnGetCurrentValueAndWeight = [NumChannels](UObject* Object, UMovieSceneSection*  SectionToKey, FFrameNumber KeyTime, FFrameRate TickResolution, FMovieSceneRootEvaluationTemplateInstance& RootTemplate,
+			float& OutValue, float& OutWeight) { GetChannelValueAndWeight(NumChannels, 0, Object, SectionToKey, KeyTime, TickResolution, RootTemplate, OutValue, OutWeight); };
+		ExternalValues[1].OnGetCurrentValueAndWeight = [NumChannels](UObject* Object, UMovieSceneSection*  SectionToKey, FFrameNumber KeyTime, FFrameRate TickResolution, FMovieSceneRootEvaluationTemplateInstance& RootTemplate,
+			float& OutValue, float& OutWeight) { GetChannelValueAndWeight(NumChannels, 1, Object, SectionToKey, KeyTime, TickResolution, RootTemplate, OutValue, OutWeight); };
+		ExternalValues[2].OnGetCurrentValueAndWeight = [NumChannels](UObject* Object, UMovieSceneSection*  SectionToKey, FFrameNumber KeyTime, FFrameRate TickResolution, FMovieSceneRootEvaluationTemplateInstance& RootTemplate,
+			float& OutValue, float& OutWeight) { GetChannelValueAndWeight(NumChannels, 2,Object, SectionToKey, KeyTime, TickResolution, RootTemplate, OutValue, OutWeight); };
+		ExternalValues[3].OnGetCurrentValueAndWeight = [NumChannels](UObject* Object, UMovieSceneSection*  SectionToKey, FFrameNumber KeyTime, FFrameRate TickResolution, FMovieSceneRootEvaluationTemplateInstance& RootTemplate,
+			float& OutValue, float& OutWeight) { GetChannelValueAndWeight(NumChannels, 3, Object, SectionToKey, KeyTime, TickResolution, RootTemplate, OutValue, OutWeight); };
+
 	}
 
 	static FVector4 GetPropertyValue(UObject& InObject, FTrackInstancePropertyBindings& Bindings, int32 NumChannels)
@@ -72,8 +87,104 @@ struct FVectorSectionEditorData
 	{
 		return Bindings ? GetPropertyValue(InObject, *Bindings, NumChannels).W : TOptional<float>();
 	}
+
+	static void GetChannelValueAndWeight(int32 NumChannels, int32 Index, UObject* Object, UMovieSceneSection*  SectionToKey,  FFrameNumber KeyTime, FFrameRate TickResolution, FMovieSceneRootEvaluationTemplateInstance& RootTemplate,
+		float& OutValue, float& OutWeight)
+	{
+		OutValue = 0.0f;
+		OutWeight = 1.0f;
+		if (Index >= NumChannels)
+		{
+			return;
+		}
+
+		UMovieSceneTrack* Track = SectionToKey->GetTypedOuter<UMovieSceneTrack>();
+
+		if (Track)
+		{
+			FMovieSceneEvaluationTrack EvalTrack = Track->GenerateTrackTemplate();
+			FMovieSceneInterrogationData InterrogationData;
+			RootTemplate.CopyActuators(InterrogationData.GetAccumulator());
+
+			FMovieSceneContext Context(FMovieSceneEvaluationRange(KeyTime, TickResolution));
+			EvalTrack.Interrogate(Context, InterrogationData, Object);
+
+			switch (NumChannels)
+			{
+			case 2:
+			{
+				FVector2D Val(0.0f, 0.0f);
+				for (const FVector2D& InVector : InterrogationData.Iterate<FVector2D>(FMovieScenePropertySectionTemplate::GetVector2DInterrogationKey()))
+				{
+					Val = InVector;
+					break;
+				}
+				switch (Index)
+				{
+				case 0:
+					OutValue = Val.X;
+					break;
+				case 1:
+					OutValue = Val.Y;
+					break;
+				}
+			}
+			break;
+			case 3:
+			{
+				FVector Val(0.0f, 0.0f, 0.0f);
+				for (const FVector& InVector : InterrogationData.Iterate<FVector>(FMovieScenePropertySectionTemplate::GetVectorInterrogationKey()))
+				{
+					Val = InVector;
+					break;
+				}
+				switch (Index)
+				{
+				case 0:
+					OutValue = Val.X;
+					break;
+				case 1:
+					OutValue = Val.Y;
+					break;
+				case 2:
+					OutValue = Val.Z;
+					break;
+				}
+			}
+			break;
+			case 4:
+			{
+				FVector4 Val(0.0f, 0.0f, 0.0f, 0.0f);
+				for (const FVector4& InVector : InterrogationData.Iterate<FVector4>(FMovieScenePropertySectionTemplate::GetVector4InterrogationKey()))
+				{
+					Val = InVector;
+					break;
+				}
+				switch (Index)
+				{
+				case 0:
+					OutValue = Val.X;
+					break;
+				case 1:
+					OutValue = Val.Y;
+					break;
+				case 2:
+					OutValue = Val.Z;
+					break;
+				case 3:
+					OutValue = Val.W;
+					break;
+				}
+			}
+			break;
+			}
+		}
+		OutWeight = MovieSceneHelpers::CalculateWeightForBlending(SectionToKey, KeyTime);
+	}
+
 	FMovieSceneChannelMetaData      MetaData[4];
 	TMovieSceneExternalValue<float> ExternalValues[4];
+
 };
 #endif
 
diff --git a/Engine/Source/Runtime/MovieSceneTracks/Private/Tracks/MovieScene3DTransformTrack.cpp b/Engine/Source/Runtime/MovieSceneTracks/Private/Tracks/MovieScene3DTransformTrack.cpp
--- a/Engine/Source/Runtime/MovieSceneTracks/Private/Tracks/MovieScene3DTransformTrack.cpp
+++ b/Engine/Source/Runtime/MovieSceneTracks/Private/Tracks/MovieScene3DTransformTrack.cpp
@@ -33,12 +33,6 @@ UMovieSceneSection* UMovieScene3DTransformTrack::CreateNewSection()
 	return NewObject<UMovieScene3DTransformSection>(this, NAME_None, RF_Transactional);
 }
 
-FMovieSceneInterrogationKey UMovieScene3DTransformTrack::GetInterrogationKey()
-{
-	static FMovieSceneAnimTypeID TypeID = FMovieSceneAnimTypeID::Unique();
-	return TypeID;
-}
-
 
 #if WITH_EDITOR
 
diff --git a/Engine/Source/Runtime/MovieSceneTracks/Private/Tracks/MovieScenePropertyTrack.cpp b/Engine/Source/Runtime/MovieSceneTracks/Private/Tracks/MovieScenePropertyTrack.cpp
--- a/Engine/Source/Runtime/MovieSceneTracks/Private/Tracks/MovieScenePropertyTrack.cpp
+++ b/Engine/Source/Runtime/MovieSceneTracks/Private/Tracks/MovieScenePropertyTrack.cpp
@@ -168,30 +168,9 @@ UMovieSceneSection* UMovieScenePropertyTrack::FindOrExtendSection(FFrameNumber T
 		}
 		//we need to calculate weight also possibly
 		FOptionalMovieSceneBlendType BlendType = SectionToKey->GetBlendType();
-		if (bCalculateWeight && BlendType.IsValid() && (BlendType.Get() == EMovieSceneBlendType::Additive || BlendType.Get() == EMovieSceneBlendType::Absolute))
+		if (bCalculateWeight)
 		{
-			//if additive weight is just the inverse of any weight on it
-			if (BlendType.Get() == EMovieSceneBlendType::Additive)
-			{
-				float TotalWeightValue = SectionToKey->GetTotalWeightValue(Time);
-				Weight = !FMath::IsNearlyZero(TotalWeightValue) ? 1.0f / TotalWeightValue : 0.0f;
-			}
-			else
-			{
-				//if absolute need to calculate weight based upon other sections weights (+ implicit absolute weights)
-				int TotalNumOfAbsoluteSections = 1;
-				for (UMovieSceneSection* Section : OverlappingSections)
-				{
-					FOptionalMovieSceneBlendType NewBlendType = Section->GetBlendType();
-
-					if (Section != SectionToKey  && NewBlendType.IsValid() && NewBlendType.Get() == EMovieSceneBlendType::Absolute)
-					{
-						++TotalNumOfAbsoluteSections;
-					}
-				}
-				float TotalWeightValue = SectionToKey->GetTotalWeightValue(Time);
-				Weight = !FMath::IsNearlyZero(TotalWeightValue) ? float(TotalNumOfAbsoluteSections) / TotalWeightValue : 0.0f;
-			}
+			Weight = MovieSceneHelpers::CalculateWeightForBlending(SectionToKey, Time);
 		}
 		return SectionToKey;
 	}
diff --git a/Engine/Source/Runtime/MovieSceneTracks/Public/Sections/MovieScene3DTransformSection.h b/Engine/Source/Runtime/MovieSceneTracks/Public/Sections/MovieScene3DTransformSection.h
--- a/Engine/Source/Runtime/MovieSceneTracks/Public/Sections/MovieScene3DTransformSection.h
+++ b/Engine/Source/Runtime/MovieSceneTracks/Public/Sections/MovieScene3DTransformSection.h
@@ -311,6 +311,12 @@ private:
 	UPROPERTY(EditAnywhere, DisplayName = "Use Quaternion Interpolation", Category = "Rotation")
 	bool bUseQuaternionInterpolation;
 
+public:
+	/**
+	 * Access the interrogation key for transform data - any interrgation data stored with this key is guaranteed to be of type 'FTransform'
+	 */
+	MOVIESCENETRACKS_API static FMovieSceneInterrogationKey GetInterrogationKey();
+
 #if WITH_EDITORONLY_DATA
 
 public:
diff --git a/Engine/Source/Runtime/MovieSceneTracks/Public/Tracks/MovieScene3DTransformTrack.h b/Engine/Source/Runtime/MovieSceneTracks/Public/Tracks/MovieScene3DTransformTrack.h
--- a/Engine/Source/Runtime/MovieSceneTracks/Public/Tracks/MovieScene3DTransformTrack.h
+++ b/Engine/Source/Runtime/MovieSceneTracks/Public/Tracks/MovieScene3DTransformTrack.h
@@ -65,12 +65,6 @@ public:
 	virtual bool SupportsType(TSubclassOf<UMovieSceneSection> SectionClass) const override;
 	virtual UMovieSceneSection* CreateNewSection() override;
 
-	/**
-	 * Access the interrogation key for transform data - any interrgation data stored with this key is guaranteed to be of type 'FTransform'
-	 */
-	MOVIESCENETRACKS_API static FMovieSceneInterrogationKey GetInterrogationKey();
-
-
 #if WITH_EDITOR
 
 	MOVIESCENETRACKS_API TArray<FTrajectoryKey> GetTrajectoryData(FFrameNumber Time, int32 MaxNumDataPoints) const;
diff --git a/Engine/Source/Runtime/UMG/Private/Animation/MovieScene2DTransformSection.cpp b/Engine/Source/Runtime/UMG/Private/Animation/MovieScene2DTransformSection.cpp
--- a/Engine/Source/Runtime/UMG/Private/Animation/MovieScene2DTransformSection.cpp
+++ b/Engine/Source/Runtime/UMG/Private/Animation/MovieScene2DTransformSection.cpp
@@ -4,7 +4,10 @@
 #include "UObject/SequencerObjectVersion.h"
 #include "Channels/MovieSceneChannelProxy.h"
 #include "Slate/WidgetTransform.h"
-
+#include "Compilation/MovieSceneTemplateInterrogation.h"
+#include "Evaluation/MovieSceneEvaluationTrack.h"
+#include "Evaluation/MovieSceneEvaluationTemplateInstance.h"
+#include "Evaluation/MovieScenePropertyTemplate.h"
 
 #if WITH_EDITOR
 
@@ -59,6 +62,73 @@ struct F2DTransformSectionEditorData
 		ExternalValues[4].OnGetExternalValue = ExtractScaleY;
 		ExternalValues[5].OnGetExternalValue = ExtractShearX;
 		ExternalValues[6].OnGetExternalValue = ExtractShearY;
+
+
+		ExternalValues[0].OnGetCurrentValueAndWeight = GetTranslationXValueAndWeight;
+		ExternalValues[1].OnGetCurrentValueAndWeight = GetTranslationYValueAndWeight;
+		ExternalValues[2].OnGetCurrentValueAndWeight = GetRotationValueAndWeight;
+		ExternalValues[3].OnGetCurrentValueAndWeight = GetScaleXValueAndWeight;
+		ExternalValues[4].OnGetCurrentValueAndWeight = GetScaleYValueAndWeight;
+		ExternalValues[5].OnGetCurrentValueAndWeight = GetShearXValueAndWeight;
+		ExternalValues[6].OnGetCurrentValueAndWeight = GetShearYValueAndWeight;
+
+
+	}
+
+	static void GetValueAndWeight(UObject* Object, UMovieSceneSection*  SectionToKey, int32 Index, FFrameNumber KeyTime, FFrameRate TickResolution, FMovieSceneRootEvaluationTemplateInstance& RootTemplate,
+		float& OutValue, float& OutWeight)
+	{
+		UMovieSceneTrack* Track = SectionToKey->GetTypedOuter<UMovieSceneTrack>();
+		FMovieSceneEvaluationTrack EvalTrack = Track->GenerateTrackTemplate();
+		FMovieSceneInterrogationData InterrogationData;
+		RootTemplate.CopyActuators(InterrogationData.GetAccumulator());
+
+		FMovieSceneContext Context(FMovieSceneEvaluationRange(KeyTime, TickResolution));
+		EvalTrack.Interrogate(Context, InterrogationData, Object);
+
+
+		FVector2D Translation;
+		FVector2D Scale;
+		FVector2D Shear;
+		float Angle = 0.0f;
+
+
+		for (const FWidgetTransform& Transform : InterrogationData.Iterate<FWidgetTransform>(UMovieScene2DTransformSection::GetWidgetTransformInterrogationKey()))
+		{
+			Translation = Transform.Translation;
+			Scale = Transform.Scale;
+			Shear = Transform.Shear;
+			Angle = Transform.Angle;
+			break;
+		}
+
+		switch (Index)
+		{
+		case 0:
+			OutValue = Translation.X;
+			break;
+		case 1:
+			OutValue = Translation.Y;
+			break;
+		case 2:
+			OutValue = Angle;
+			break;
+		case 3:
+			OutValue = Scale.X;
+			break;
+		case 4:
+			OutValue = Scale.Y;
+			break;
+		case 5:
+			OutValue = Shear.X;
+			break;
+		case 6:
+			OutValue = Shear.Y;
+			break;
+	
+
+		}
+		OutWeight = MovieSceneHelpers::CalculateWeightForBlending(SectionToKey, KeyTime);
 	}
 
 	static TOptional<float> ExtractTranslationX(UObject& InObject, FTrackInstancePropertyBindings* Bindings)
@@ -93,6 +163,43 @@ struct F2DTransformSectionEditorData
 		return Bindings ? Bindings->GetCurrentValue<FWidgetTransform>(InObject).Shear.Y : TOptional<float>();
 	}
 
+
+	static void GetTranslationXValueAndWeight(UObject* Object, UMovieSceneSection*  SectionToKey, FFrameNumber KeyTime, FFrameRate TickResolution, FMovieSceneRootEvaluationTemplateInstance& RootTemplate,
+		float& OutValue, float& OutWeight)
+	{
+		GetValueAndWeight(Object, SectionToKey, 0, KeyTime, TickResolution, RootTemplate, OutValue, OutWeight);
+	}
+	static void GetTranslationYValueAndWeight(UObject* Object, UMovieSceneSection*  SectionToKey, FFrameNumber KeyTime, FFrameRate TickResolution, FMovieSceneRootEvaluationTemplateInstance& RootTemplate,
+		float& OutValue, float& OutWeight)
+	{
+		GetValueAndWeight(Object, SectionToKey, 1, KeyTime, TickResolution, RootTemplate, OutValue, OutWeight);
+	}
+	static void GetRotationValueAndWeight(UObject* Object, UMovieSceneSection*  SectionToKey, FFrameNumber KeyTime, FFrameRate TickResolution, FMovieSceneRootEvaluationTemplateInstance& RootTemplate,
+		float& OutValue, float& OutWeight)
+	{
+		GetValueAndWeight(Object, SectionToKey, 2, KeyTime, TickResolution, RootTemplate, OutValue, OutWeight);
+	}
+	static void GetScaleXValueAndWeight(UObject* Object, UMovieSceneSection*  SectionToKey, FFrameNumber KeyTime, FFrameRate TickResolution, FMovieSceneRootEvaluationTemplateInstance& RootTemplate,
+		float& OutValue, float& OutWeight)
+	{
+		GetValueAndWeight(Object, SectionToKey, 3, KeyTime, TickResolution, RootTemplate, OutValue, OutWeight);
+	}
+	static void GetScaleYValueAndWeight(UObject* Object, UMovieSceneSection*  SectionToKey, FFrameNumber KeyTime, FFrameRate TickResolution, FMovieSceneRootEvaluationTemplateInstance& RootTemplate,
+		float& OutValue, float& OutWeight)
+	{
+		GetValueAndWeight(Object, SectionToKey, 4, KeyTime, TickResolution, RootTemplate, OutValue, OutWeight);
+	}
+	static void GetShearXValueAndWeight(UObject* Object, UMovieSceneSection*  SectionToKey, FFrameNumber KeyTime, FFrameRate TickResolution, FMovieSceneRootEvaluationTemplateInstance& RootTemplate,
+		float& OutValue, float& OutWeight)
+	{
+		GetValueAndWeight(Object, SectionToKey, 5, KeyTime, TickResolution, RootTemplate, OutValue, OutWeight);
+	}
+	static void GetShearYValueAndWeight(UObject* Object, UMovieSceneSection*  SectionToKey, FFrameNumber KeyTime, FFrameRate TickResolution, FMovieSceneRootEvaluationTemplateInstance& RootTemplate,
+		float& OutValue, float& OutWeight)
+	{
+		GetValueAndWeight(Object, SectionToKey, 6, KeyTime, TickResolution, RootTemplate, OutValue, OutWeight);
+	}
+
 	FMovieSceneChannelMetaData      MetaData[7];
 	TMovieSceneExternalValue<float> ExternalValues[7];
 };
@@ -229,3 +336,10 @@ FMovieScene2DTransformMask UMovieScene2DTransformSection::GetMaskByName(const FN
 
 	return EMovieScene2DTransformChannel::AllTransform;
 }
+
+const FMovieSceneInterrogationKey UMovieScene2DTransformSection::GetWidgetTransformInterrogationKey()
+{
+	const static FMovieSceneAnimTypeID TypeID = FMovieSceneAnimTypeID::Unique();
+	return TypeID;
+}
+
diff --git a/Engine/Source/Runtime/UMG/Private/Animation/MovieScene2DTransformTemplate.cpp b/Engine/Source/Runtime/UMG/Private/Animation/MovieScene2DTransformTemplate.cpp
--- a/Engine/Source/Runtime/UMG/Private/Animation/MovieScene2DTransformTemplate.cpp
+++ b/Engine/Source/Runtime/UMG/Private/Animation/MovieScene2DTransformTemplate.cpp
@@ -68,3 +68,48 @@ void FMovieScene2DTransformSectionTemplate::Evaluate(const FMovieSceneEvaluation
 		ExecutionTokens.BlendToken(ActuatorTypeID, TBlendableToken<FWidgetTransform>(AnimatedData, BlendType, Weight));
 	}
 }
+
+void FMovieScene2DTransformSectionTemplate::Interrogate(const FMovieSceneContext& Context, FMovieSceneInterrogationData& Container, UObject* BindingOverride) const
+{
+	const FFrameTime Time = Context.GetTime();
+	MovieScene::TMultiChannelValue<float, 7> AnimatedData;
+
+	EMovieScene2DTransformChannel ChannelMask = Mask.GetChannels();
+
+	// Only activate channels if the curve has data associated with it
+	auto EvalChannel = [&AnimatedData, Time, ChannelMask](uint8 ChanneIndex, EMovieScene2DTransformChannel ChannelType, const FMovieSceneFloatChannel& Channel)
+	{
+		float Value = 0.f;
+		if (EnumHasAllFlags(ChannelMask, ChannelType) && Channel.Evaluate(Time, Value))
+		{
+			AnimatedData.Set(ChanneIndex, Value);
+		}
+	};
+
+	EvalChannel(0, EMovieScene2DTransformChannel::TranslationX, Translation[0]);
+	EvalChannel(1, EMovieScene2DTransformChannel::TranslationY, Translation[1]);
+
+	EvalChannel(2, EMovieScene2DTransformChannel::ScaleX, Scale[0]);
+	EvalChannel(3, EMovieScene2DTransformChannel::ScaleY, Scale[1]);
+
+	EvalChannel(4, EMovieScene2DTransformChannel::ShearX, Shear[0]);
+	EvalChannel(5, EMovieScene2DTransformChannel::ShearY, Shear[1]);
+
+	EvalChannel(6, EMovieScene2DTransformChannel::Rotation, Rotation);
+
+	FMovieSceneAnimTypeID TypeID = GetPropertyTypeID();
+	static FMovieSceneBlendingActuatorID ActuatorTypeID(TypeID);
+	if (!Container.GetAccumulator().FindActuator<FWidgetTransform>(ActuatorTypeID))
+	{
+		PropertyTemplate::FSectionData SectionData;
+		SectionData.Initialize(PropertyData.PropertyName, PropertyData.PropertyPath, PropertyData.FunctionName, PropertyData.NotifyFunctionName);
+		Container.GetAccumulator().DefineActuator(ActuatorTypeID, MakeShared<TPropertyActuator<FWidgetTransform>>(SectionData));
+	}
+
+	if (!AnimatedData.IsEmpty())
+	{
+		// Add the blendable to the accumulator
+		float Weight = EvaluateEasing(Context.GetTime());
+		Container.GetAccumulator().BlendToken(FMovieSceneEvaluationOperand(), ActuatorTypeID, FMovieSceneEvaluationScope(), Context, TBlendableToken<FWidgetTransform>(AnimatedData, BlendType, Weight));
+	}
+}
diff --git a/Engine/Source/Runtime/UMG/Private/Animation/MovieScene2DTransformTemplate.h b/Engine/Source/Runtime/UMG/Private/Animation/MovieScene2DTransformTemplate.h
--- a/Engine/Source/Runtime/UMG/Private/Animation/MovieScene2DTransformTemplate.h
+++ b/Engine/Source/Runtime/UMG/Private/Animation/MovieScene2DTransformTemplate.h
@@ -10,6 +10,7 @@
 #include "Evaluation/MovieScenePropertyTemplate.h"
 #include "Evaluation/Blending/MovieSceneMultiChannelBlending.h"
 #include "Slate/WidgetTransform.h"
+#include "Evaluation/MovieScenePropertyTemplate.h"
 #include "MovieScene2DTransformTemplate.generated.h"
 
 class UMovieScene2DTransformSection;
@@ -31,6 +32,8 @@ private:
 
 	virtual void Evaluate(const FMovieSceneEvaluationOperand& Operand, const FMovieSceneContext& Context, const FPersistentEvaluationData& PersistentData, FMovieSceneExecutionTokens& ExecutionTokens) const override;
 
+	virtual void Interrogate(const FMovieSceneContext& Context, FMovieSceneInterrogationData& Container, UObject* BindingOverride) const override;
+
 	/** Translation curves */
 	UPROPERTY()
 	FMovieSceneFloatChannel Translation[2];
@@ -55,6 +58,13 @@ private:
 	FMovieScene2DTransformMask Mask;
 };
 
+template<>
+inline void TPropertyActuator<FWidgetTransform>::Actuate(FMovieSceneInterrogationData& InterrogationData, typename TCallTraits<FWidgetTransform>::ParamType InValue, const TBlendableTokenStack<FWidgetTransform>& OriginalStack, const FMovieSceneContext& Context) const
+{
+	FWidgetTransform Value = InValue;
+	InterrogationData.Add(Value, UMovieScene2DTransformSection::GetWidgetTransformInterrogationKey());
+};
+
 /** Access the unique runtime type identifier for a widget transform. */
 template<> UMG_API FMovieSceneAnimTypeID GetBlendingDataType<FWidgetTransform>();
 
@@ -80,4 +90,4 @@ namespace MovieScene
 		Out.Shear = FVector2D(In[4], In[5]);
 		Out.Angle = In[6];
 	}
-}
\ No newline at end of file
+}
diff --git a/Engine/Source/Runtime/UMG/Private/Animation/MovieSceneMarginSection.cpp b/Engine/Source/Runtime/UMG/Private/Animation/MovieSceneMarginSection.cpp
--- a/Engine/Source/Runtime/UMG/Private/Animation/MovieSceneMarginSection.cpp
+++ b/Engine/Source/Runtime/UMG/Private/Animation/MovieSceneMarginSection.cpp
@@ -2,6 +2,10 @@
 
 #include "Animation/MovieSceneMarginSection.h"
 #include "Channels/MovieSceneChannelProxy.h"
+#include "Compilation/MovieSceneTemplateInterrogation.h"
+#include "Evaluation/MovieSceneEvaluationTrack.h"
+#include "Evaluation/MovieSceneEvaluationTemplateInstance.h"
+#include "Evaluation/MovieScenePropertyTemplate.h"
 
 #if WITH_EDITOR
 
@@ -29,6 +33,11 @@ struct FMarginSectionEditorData
 		ExternalValues[1].OnGetExternalValue = ExtractTopChannel;
 		ExternalValues[2].OnGetExternalValue = ExtractRightChannel;
 		ExternalValues[3].OnGetExternalValue = ExtractBottomChannel;
+
+		ExternalValues[0].OnGetCurrentValueAndWeight = GetLeftChannelValueAndWeight;
+		ExternalValues[1].OnGetCurrentValueAndWeight = GetTopChannelValueAndWeight;
+		ExternalValues[2].OnGetCurrentValueAndWeight = GetRightChannelValueAndWeight;
+		ExternalValues[3].OnGetCurrentValueAndWeight = GetBottomChannelValueAndWeight;
 	}
 
 	static TOptional<float> ExtractLeftChannel(UObject& InObject, FTrackInstancePropertyBindings* Bindings)
@@ -47,6 +56,65 @@ struct FMarginSectionEditorData
 	{
 		return Bindings ? Bindings->GetCurrentValue<FMargin>(InObject).Bottom : TOptional<float>();
 	}
+	static void GetLeftChannelValueAndWeight(UObject* Object, UMovieSceneSection*  SectionToKey, FFrameNumber KeyTime, FFrameRate TickResolution, FMovieSceneRootEvaluationTemplateInstance& RootTemplate,
+		float& OutValue, float& OutWeight)
+	{
+		GetValueAndWeight(Object, SectionToKey, 0, KeyTime, TickResolution, RootTemplate, OutValue, OutWeight);
+	}
+	static void GetTopChannelValueAndWeight(UObject* Object, UMovieSceneSection*  SectionToKey, FFrameNumber KeyTime, FFrameRate TickResolution, FMovieSceneRootEvaluationTemplateInstance& RootTemplate,
+		float& OutValue, float& OutWeight)
+	{
+		GetValueAndWeight(Object, SectionToKey, 1, KeyTime, TickResolution, RootTemplate, OutValue, OutWeight);
+	}
+	static void GetRightChannelValueAndWeight(UObject* Object, UMovieSceneSection*  SectionToKey, FFrameNumber KeyTime, FFrameRate TickResolution, FMovieSceneRootEvaluationTemplateInstance& RootTemplate,
+		float& OutValue, float& OutWeight)
+	{
+		GetValueAndWeight(Object, SectionToKey, 2, KeyTime, TickResolution, RootTemplate, OutValue, OutWeight);
+	}
+	static void GetBottomChannelValueAndWeight(UObject* Object, UMovieSceneSection*  SectionToKey, FFrameNumber KeyTime, FFrameRate TickResolution, FMovieSceneRootEvaluationTemplateInstance& RootTemplate,
+		float& OutValue, float& OutWeight)
+	{
+		GetValueAndWeight(Object, SectionToKey, 3, KeyTime, TickResolution, RootTemplate, OutValue, OutWeight);
+	}
+	static void GetValueAndWeight(UObject* Object, UMovieSceneSection*  SectionToKey, int32 Index, FFrameNumber KeyTime, FFrameRate TickResolution, FMovieSceneRootEvaluationTemplateInstance& RootTemplate,
+		float& OutValue, float& OutWeight)
+	{
+		UMovieSceneTrack* Track = SectionToKey->GetTypedOuter<UMovieSceneTrack>();
+		FMovieSceneEvaluationTrack EvalTrack = Track->GenerateTrackTemplate();
+		FMovieSceneInterrogationData InterrogationData;
+		RootTemplate.CopyActuators(InterrogationData.GetAccumulator());
+
+		FMovieSceneContext Context(FMovieSceneEvaluationRange(KeyTime, TickResolution));
+		EvalTrack.Interrogate(Context, InterrogationData, Object);
+
+		float Left = 0.0f, Right = 0.0f, Top = 0.0f, Bottom = 0.0f;
+
+		for (const FMargin& Margin : InterrogationData.Iterate<FMargin>(UMovieSceneMarginSection::GetMarginInterrogationKey()))
+		{
+			Left = Margin.Left;
+			Right = Margin.Right;
+			Top = Margin.Top;
+			Bottom = Margin.Bottom;
+			break;
+		}
+
+		switch (Index)
+		{
+		case 0:
+			OutValue = Left;
+			break;
+		case 1:
+			OutValue = Top;
+			break;
+		case 2:
+			OutValue = Right;
+			break;
+		case 3:
+			OutValue = Bottom;
+			break;
+		}
+		OutWeight = MovieSceneHelpers::CalculateWeightForBlending(SectionToKey, KeyTime);
+	}
 
 	FMovieSceneChannelMetaData      MetaData[4];
 	TMovieSceneExternalValue<float> ExternalValues[4];
@@ -80,4 +148,10 @@ UMovieSceneMarginSection::UMovieSceneMarginSection( const FObjectInitializer& Ob
 #endif
 	
 	ChannelProxy = MakeShared<FMovieSceneChannelProxy>(MoveTemp(Channels));
+}
+
+const FMovieSceneInterrogationKey UMovieSceneMarginSection::GetMarginInterrogationKey()
+{
+	const static FMovieSceneAnimTypeID TypeID = FMovieSceneAnimTypeID::Unique();
+	return TypeID;
 }
\ No newline at end of file
diff --git a/Engine/Source/Runtime/UMG/Private/Animation/MovieSceneMarginTemplate.cpp b/Engine/Source/Runtime/UMG/Private/Animation/MovieSceneMarginTemplate.cpp
--- a/Engine/Source/Runtime/UMG/Private/Animation/MovieSceneMarginTemplate.cpp
+++ b/Engine/Source/Runtime/UMG/Private/Animation/MovieSceneMarginTemplate.cpp
@@ -57,3 +57,47 @@ void FMovieSceneMarginSectionTemplate::Evaluate(const FMovieSceneEvaluationOpera
 		ExecutionTokens.BlendToken(ActuatorTypeID, TBlendableToken<FMargin>(AnimatedData, BlendType, Weight));
 	}
 }
+
+
+void FMovieSceneMarginSectionTemplate::Interrogate(const FMovieSceneContext& Context, FMovieSceneInterrogationData& Container, UObject* BindingOverride) const
+{
+	const FFrameTime Time = Context.GetTime();
+	MovieScene::TMultiChannelValue<float, 4> AnimatedData;
+
+	float Value = 0.f;
+	if (LeftCurve.Evaluate(Time, Value))
+	{
+		AnimatedData.Set(0, Value);
+	}
+
+	if (TopCurve.Evaluate(Time, Value))
+	{
+		AnimatedData.Set(1, Value);
+	}
+
+	if (RightCurve.Evaluate(Time, Value))
+	{
+		AnimatedData.Set(2, Value);
+	}
+
+	if (BottomCurve.Evaluate(Time, Value))
+	{
+		AnimatedData.Set(3, Value);
+	}
+
+	FMovieSceneAnimTypeID TypeID = GetPropertyTypeID();
+	static FMovieSceneBlendingActuatorID ActuatorTypeID(TypeID);
+	if (!Container.GetAccumulator().FindActuator<FMargin>(ActuatorTypeID))
+	{
+		PropertyTemplate::FSectionData SectionData;
+		SectionData.Initialize(PropertyData.PropertyName, PropertyData.PropertyPath, PropertyData.FunctionName, PropertyData.NotifyFunctionName);
+		Container.GetAccumulator().DefineActuator(ActuatorTypeID, MakeShared<TPropertyActuator<FMargin>>(SectionData));
+	}
+
+	if (!AnimatedData.IsEmpty())
+	{
+		// Add the blendable to the accumulator
+		float Weight = EvaluateEasing(Context.GetTime());
+		Container.GetAccumulator().BlendToken(FMovieSceneEvaluationOperand(), ActuatorTypeID, FMovieSceneEvaluationScope(), Context, TBlendableToken<FMargin>(AnimatedData, BlendType, Weight));
+	}
+}
diff --git a/Engine/Source/Runtime/UMG/Private/Animation/MovieSceneMarginTemplate.h b/Engine/Source/Runtime/UMG/Private/Animation/MovieSceneMarginTemplate.h
--- a/Engine/Source/Runtime/UMG/Private/Animation/MovieSceneMarginTemplate.h
+++ b/Engine/Source/Runtime/UMG/Private/Animation/MovieSceneMarginTemplate.h
@@ -9,6 +9,7 @@
 #include "Evaluation/MovieScenePropertyTemplate.h"
 #include "Evaluation/Blending/MovieSceneMultiChannelBlending.h"
 #include "Layout/Margin.h"
+#include "Animation/MovieSceneMarginSection.h"
 #include "MovieSceneMarginTemplate.generated.h"
 
 class UMovieSceneMarginSection;
@@ -33,6 +34,7 @@ private:
 
 	virtual UScriptStruct& GetScriptStructImpl() const override { return *StaticStruct(); }
 	virtual void Evaluate(const FMovieSceneEvaluationOperand& Operand, const FMovieSceneContext& Context, const FPersistentEvaluationData& PersistentData, FMovieSceneExecutionTokens& ExecutionTokens) const override;
+	virtual void Interrogate(const FMovieSceneContext& Context, FMovieSceneInterrogationData& Container, UObject* BindingOverride) const override;
 
 	UPROPERTY()
 	FMovieSceneFloatChannel TopCurve;
@@ -50,6 +52,13 @@ private:
 	EMovieSceneBlendType BlendType;
 };
 
+template<>
+inline void TPropertyActuator<FMargin>::Actuate(FMovieSceneInterrogationData& InterrogationData, typename TCallTraits<FMargin>::ParamType InValue, const TBlendableTokenStack<FMargin>& OriginalStack, const FMovieSceneContext& Context) const
+{
+	FMargin Value = InValue;
+	InterrogationData.Add(Value, UMovieSceneMarginSection::GetMarginInterrogationKey());
+};
+
 /** Access the unique runtime type identifier for a margin. */
 template<> UMG_API FMovieSceneAnimTypeID GetBlendingDataType<FMargin>();
 
diff --git a/Engine/Source/Runtime/UMG/Public/Animation/MovieScene2DTransformSection.h b/Engine/Source/Runtime/UMG/Public/Animation/MovieScene2DTransformSection.h
--- a/Engine/Source/Runtime/UMG/Public/Animation/MovieScene2DTransformSection.h
+++ b/Engine/Source/Runtime/UMG/Public/Animation/MovieScene2DTransformSection.h
@@ -121,6 +121,8 @@ protected:
 	virtual void PostEditImport() override;
 
 	void UpdateChannelProxy();
+public:
+	const static FMovieSceneInterrogationKey GetWidgetTransformInterrogationKey();
 
 public:
 
diff --git a/Engine/Source/Runtime/UMG/Public/Animation/MovieSceneMarginSection.h b/Engine/Source/Runtime/UMG/Public/Animation/MovieSceneMarginSection.h
--- a/Engine/Source/Runtime/UMG/Public/Animation/MovieSceneMarginSection.h
+++ b/Engine/Source/Runtime/UMG/Public/Animation/MovieSceneMarginSection.h
@@ -20,6 +20,8 @@ class UMovieSceneMarginSection
 	: public UMovieSceneSection
 {
 	GENERATED_UCLASS_BODY()
+public:
+	const static FMovieSceneInterrogationKey GetMarginInterrogationKey();
 
 public:
 
