Sequencer: Set params as blueprinttype

diff --git a/Engine/Source/Runtime/MovieSceneTracks/Public/Sections/MovieSceneSkeletalAnimationSection.h b/Engine/Source/Runtime/MovieSceneTracks/Public/Sections/MovieSceneSkeletalAnimationSection.h
--- a/Engine/Source/Runtime/MovieSceneTracks/Public/Sections/MovieSceneSkeletalAnimationSection.h
+++ b/Engine/Source/Runtime/MovieSceneTracks/Public/Sections/MovieSceneSkeletalAnimationSection.h
@@ -9,7 +9,7 @@
 #include "Channels/MovieSceneFloatChannel.h"
 #include "MovieSceneSkeletalAnimationSection.generated.h"
 
-USTRUCT()
+USTRUCT(BlueprintType)
 struct FMovieSceneSkeletalAnimationParams
 {
 	GENERATED_BODY()
@@ -23,35 +23,35 @@ struct FMovieSceneSkeletalAnimationParams
 	float GetSequenceLength() const { return Animation != nullptr ? Animation->SequenceLength : 0.f; }
 
 	/** The animation this section plays */
-	UPROPERTY(EditAnywhere, Category="Animation", meta=(AllowedClasses = "AnimSequence, AnimComposite"))
+	UPROPERTY(BlueprintReadWrite, EditAnywhere, Category="Animation", meta=(AllowedClasses = "AnimSequence, AnimComposite"))
 	UAnimSequenceBase* Animation;
 
 	/** The offset into the beginning of the animation clip */
-	UPROPERTY(EditAnywhere, Category="Animation")
+	UPROPERTY(BlueprintReadWrite, EditAnywhere, Category="Animation")
 	FFrameNumber StartFrameOffset;
 
 	/** The offset into the end of the animation clip */
-	UPROPERTY(EditAnywhere, Category="Animation")
+	UPROPERTY(BlueprintReadWrite, EditAnywhere, Category="Animation")
 	FFrameNumber EndFrameOffset;
 
 	/** The playback rate of the animation clip */
-	UPROPERTY(EditAnywhere, Category="Animation")
+	UPROPERTY(BlueprintReadWrite, EditAnywhere, Category="Animation")
 	float PlayRate;
 
 	/** Reverse the playback of the animation clip */
-	UPROPERTY(EditAnywhere, Category="Animation")
+	UPROPERTY(BlueprintReadWrite, EditAnywhere, Category="Animation")
 	uint32 bReverse:1;
 
 	/** The slot name to use for the animation */
-	UPROPERTY( EditAnywhere, Category = "Animation" )
+	UPROPERTY(BlueprintReadWrite, EditAnywhere, Category = "Animation" )
 	FName SlotName;
 
 	/** The weight curve for this animation section */
-	UPROPERTY( )
+	UPROPERTY()
 	FMovieSceneFloatChannel Weight;
 
 	/** If on will skip sending animation notifies */
-	UPROPERTY(EditAnywhere, Category = "Animation")
+	UPROPERTY(BlueprintReadWrite, EditAnywhere, Category = "Animation")
 	bool bSkipAnimNotifiers;
 
 	/** If on animation sequence will always play when active even if the animation is controlled by a Blueprint or Anim Instance Class*/
@@ -76,7 +76,7 @@ class UMovieSceneSkeletalAnimationSection
 
 public:
 
-	UPROPERTY(EditAnywhere, Category = "Animation", meta=(ShowOnlyInnerProperties))
+	UPROPERTY(BlueprintReadWrite, EditAnywhere, Category = "Animation", meta=(ShowOnlyInnerProperties))
 	FMovieSceneSkeletalAnimationParams Params;
 
 	/** Get Frame Time as Animation Time*/
