Added layering support to FLinearColor to fix keying bugs.  Also added Margins and 2D Transforms layering support to it's key editors so we can now properly key additve sections for them.

diff --git a/Engine/Source/Editor/MovieSceneTools/Private/TrackEditors/PropertyTrackEditors/ColorPropertyTrackEditor.cpp b/Engine/Source/Editor/MovieSceneTools/Private/TrackEditors/PropertyTrackEditors/ColorPropertyTrackEditor.cpp
--- a/Engine/Source/Editor/MovieSceneTools/Private/TrackEditors/PropertyTrackEditors/ColorPropertyTrackEditor.cpp
+++ b/Engine/Source/Editor/MovieSceneTools/Private/TrackEditors/PropertyTrackEditors/ColorPropertyTrackEditor.cpp
@@ -7,7 +7,7 @@
 #include "MatineeImportTools.h"
 #include "Matinee/InterpTrackLinearColorProp.h"
 #include "Matinee/InterpTrackColorProp.h"
-
+#include "Evaluation/MovieScenePropertyTemplate.h"
 
 FName FColorPropertyTrackEditor::RedName( "R" );
 FName FColorPropertyTrackEditor::GreenName( "G" );
@@ -122,3 +122,34 @@ void FColorPropertyTrackEditor::BuildTrackContextMenu( FMenuBuilder& MenuBuilder
 	FKeyframeTrackEditor::BuildTrackContextMenu(MenuBuilder, Track);
 }
 
+
+bool FColorPropertyTrackEditor::ModifyGeneratedKeysByCurrentAndWeight(UObject *Object, UMovieSceneTrack *Track, UMovieSceneSection* SectionToKey, FFrameNumber KeyTime, FGeneratedTrackKeys& GeneratedTotalKeys, float Weight) const
+{
+	FFrameRate TickResolution = GetSequencer()->GetFocusedTickResolution();
+
+	UMovieSceneColorTrack* ColorTrack = Cast<UMovieSceneColorTrack>(Track);
+	FMovieSceneEvaluationTrack EvalTrack = Track->GenerateTrackTemplate();
+
+	if (ColorTrack)
+	{
+		FMovieSceneInterrogationData InterrogationData;
+		GetSequencer()->GetEvaluationTemplate().CopyActuators(InterrogationData.GetAccumulator());
+
+		FMovieSceneContext Context(FMovieSceneEvaluationRange(KeyTime, GetSequencer()->GetFocusedTickResolution()));
+		EvalTrack.Interrogate(Context, InterrogationData, Object);
+
+		FLinearColor Val(0.0f, 0.0f, 0.0f, 0.0f);
+		for (const FLinearColor& InColor : InterrogationData.Iterate<FLinearColor>(FMovieScenePropertySectionTemplate::GetColorInterrogationKey()))
+		{
+			Val = InColor;
+			break;
+		}
+		FMovieSceneChannelProxy& Proxy = SectionToKey->GetChannelProxy();
+		GeneratedTotalKeys[0]->ModifyByCurrentAndWeight(Proxy, KeyTime, (void *)&Val.R, Weight);
+		GeneratedTotalKeys[1]->ModifyByCurrentAndWeight(Proxy, KeyTime, (void *)&Val.G, Weight);
+		GeneratedTotalKeys[2]->ModifyByCurrentAndWeight(Proxy, KeyTime, (void *)&Val.B, Weight);
+		GeneratedTotalKeys[3]->ModifyByCurrentAndWeight(Proxy, KeyTime, (void *)&Val.A, Weight);
+		return true;
+	}
+	return false;
+}
diff --git a/Engine/Source/Editor/MovieSceneTools/Private/TrackEditors/PropertyTrackEditors/ColorPropertyTrackEditor.h b/Engine/Source/Editor/MovieSceneTools/Private/TrackEditors/PropertyTrackEditors/ColorPropertyTrackEditor.h
--- a/Engine/Source/Editor/MovieSceneTools/Private/TrackEditors/PropertyTrackEditors/ColorPropertyTrackEditor.h
+++ b/Engine/Source/Editor/MovieSceneTools/Private/TrackEditors/PropertyTrackEditors/ColorPropertyTrackEditor.h
@@ -58,6 +58,7 @@ protected:
 	//~ FPropertyTrackEditor interface
 
 	virtual void GenerateKeysFromPropertyChanged(const FPropertyChangedParams& PropertyChangedParams, FGeneratedTrackKeys& OutGeneratedKeys) override;
+	virtual bool ModifyGeneratedKeysByCurrentAndWeight(UObject *Object, UMovieSceneTrack *Track, UMovieSceneSection* SectionToKey, FFrameNumber KeyTime, FGeneratedTrackKeys& GeneratedTotalKeys, float Weight) const override;
 
 protected:
 
diff --git a/Engine/Source/Editor/MovieSceneTools/Public/KeyframeTrackEditor.h b/Engine/Source/Editor/MovieSceneTools/Public/KeyframeTrackEditor.h
--- a/Engine/Source/Editor/MovieSceneTools/Public/KeyframeTrackEditor.h
+++ b/Engine/Source/Editor/MovieSceneTools/Public/KeyframeTrackEditor.h
@@ -40,10 +40,11 @@ template<typename ChannelType, typename ValueType>
 struct TAddKeyImpl : IImpl
 {
 	int32 ChannelIndex;
+	bool bAddKey;
 	ValueType ValueToSet;
 
-	TAddKeyImpl(int32 InChannelIndex, const ValueType& InValue)
-		: ChannelIndex(InChannelIndex), ValueToSet(InValue)
+	TAddKeyImpl(int32 InChannelIndex, bool bInAddKey, const ValueType& InValue)
+		: ChannelIndex(InChannelIndex), bAddKey(bInAddKey), ValueToSet(InValue)
 	{}
 
 	virtual bool Apply(UMovieSceneSection* Section, FMovieSceneChannelProxy& Proxy, FFrameNumber InTime, EMovieSceneKeyInterpolation InterpolationMode, bool bKeyEvenIfUnchanged, bool bKeyEvenIfEmpty) const override
@@ -52,7 +53,7 @@ struct TAddKeyImpl : IImpl
 		using namespace MovieScene;
 
 		ChannelType* Channel = Proxy.GetChannel<ChannelType>(ChannelIndex);
-		if (Channel)
+		if (bAddKey && Channel)
 		{
 			bool bShouldKeyChannel = bKeyEvenIfUnchanged;
 			if (!bShouldKeyChannel)
@@ -99,10 +100,11 @@ template<>
 struct TAddKeyImpl<FMovieSceneFloatChannel, float> : IImpl
 {
 	int32 ChannelIndex;
+	bool bAddKey;
 	float ValueToSet;
 
-	TAddKeyImpl(int32 InChannelIndex, const float& InValue)
-		: ChannelIndex(InChannelIndex), ValueToSet(InValue)
+	TAddKeyImpl(int32 InChannelIndex, bool bInAddKey,  const float& InValue)
+		: ChannelIndex(InChannelIndex), bAddKey(bInAddKey), ValueToSet(InValue)
 	{}
 
 	virtual bool Apply(UMovieSceneSection* Section, FMovieSceneChannelProxy& Proxy, FFrameNumber InTime, EMovieSceneKeyInterpolation InterpolationMode, bool bKeyEvenIfUnchanged, bool bKeyEvenIfEmpty) const override
@@ -111,7 +113,7 @@ struct TAddKeyImpl<FMovieSceneFloatChannel, float> : IImpl
 		using namespace MovieScene;
 
 		FMovieSceneFloatChannel* Channel = Proxy.GetChannel<FMovieSceneFloatChannel>(ChannelIndex);
-		if (Channel)
+		if (bAddKey && Channel)
 		{
 			bool bShouldKeyChannel = bKeyEvenIfUnchanged;
 			if (!bShouldKeyChannel)
@@ -176,10 +178,11 @@ template<>
 struct TAddKeyImpl<FMovieSceneIntegerChannel, int32> : IImpl
 {
 	int32 ChannelIndex;
+	bool bAddKey;
 	int32 ValueToSet;
 
-	TAddKeyImpl(int32 InChannelIndex, const int32& InValue)
-		: ChannelIndex(InChannelIndex), ValueToSet(InValue)
+	TAddKeyImpl(int32 InChannelIndex, bool bInAddKey, const int32& InValue)
+		: ChannelIndex(InChannelIndex),bAddKey(bInAddKey), ValueToSet(InValue)
 	{}
 
 	virtual bool Apply(UMovieSceneSection* Section, FMovieSceneChannelProxy& Proxy, FFrameNumber InTime, EMovieSceneKeyInterpolation InterpolationMode, bool bKeyEvenIfUnchanged, bool bKeyEvenIfEmpty) const override
@@ -188,7 +191,7 @@ struct TAddKeyImpl<FMovieSceneIntegerChannel, int32> : IImpl
 		using namespace MovieScene;
 
 		FMovieSceneIntegerChannel* Channel = Proxy.GetChannel<FMovieSceneIntegerChannel>(ChannelIndex);
-		if (Channel)
+		if (bAddKey && Channel)
 		{
 			bool bShouldKeyChannel = bKeyEvenIfUnchanged;
 			if (!bShouldKeyChannel)
@@ -262,7 +265,7 @@ struct FMovieSceneChannelValueSetter
 	static FMovieSceneChannelValueSetter Create(int32 ChannelIndex, ValueType InNewValue, bool bAddKey)
 	{
 		FMovieSceneChannelValueSetter NewValue;
-		NewValue.Impl = TAddKeyImpl<ChannelType, typename TDecay<ValueType>::Type>(ChannelIndex, Forward<ValueType>(InNewValue));
+		NewValue.Impl = TAddKeyImpl<ChannelType, typename TDecay<ValueType>::Type>(ChannelIndex, bAddKey, Forward<ValueType>(InNewValue));
 		
 		return MoveTemp(NewValue);
 	}
diff --git a/Engine/Source/Editor/UMGEditor/Private/Animation/MarginTrackEditor.cpp b/Engine/Source/Editor/UMGEditor/Private/Animation/MarginTrackEditor.cpp
--- a/Engine/Source/Editor/UMGEditor/Private/Animation/MarginTrackEditor.cpp
+++ b/Engine/Source/Editor/UMGEditor/Private/Animation/MarginTrackEditor.cpp
@@ -3,19 +3,16 @@
 #include "Animation/MarginTrackEditor.h"
 #include "ISectionLayoutBuilder.h"
 
-
 FName FMarginTrackEditor::LeftName( "Left" );
 FName FMarginTrackEditor::TopName( "Top" );
 FName FMarginTrackEditor::RightName( "Right" );
 FName FMarginTrackEditor::BottomName( "Bottom" );
 
-
 TSharedRef<ISequencerTrackEditor> FMarginTrackEditor::CreateTrackEditor( TSharedRef<ISequencer> InSequencer )
 {
 	return MakeShareable( new FMarginTrackEditor( InSequencer ) );
 }
 
-
 void FMarginTrackEditor::GenerateKeysFromPropertyChanged( const FPropertyChangedParams& PropertyChangedParams, FGeneratedTrackKeys& OutGeneratedKeys)
 {
 	FPropertyPath StructPath = PropertyChangedParams.StructPathToKey;
@@ -33,3 +30,34 @@ void FMarginTrackEditor::GenerateKeysFromPropertyChanged( const FPropertyChanged
 	OutGeneratedKeys.Add(FMovieSceneChannelValueSetter::Create<FMovieSceneFloatChannel>(2, Margin.Right,  bKeyRight));
 	OutGeneratedKeys.Add(FMovieSceneChannelValueSetter::Create<FMovieSceneFloatChannel>(3, Margin.Bottom, bKeyBottom));
 }
+
+bool FMarginTrackEditor::ModifyGeneratedKeysByCurrentAndWeight(UObject *Object, UMovieSceneTrack *Track, UMovieSceneSection* SectionToKey, FFrameNumber KeyTime, FGeneratedTrackKeys& GeneratedTotalKeys, float Weight) const
+{
+	FFrameRate TickResolution = GetSequencer()->GetFocusedTickResolution();
+
+	UMovieSceneMarginTrack* MarginTrack = Cast<UMovieSceneMarginTrack>(Track);
+	FMovieSceneEvaluationTrack EvalTrack = Track->GenerateTrackTemplate();
+
+	if (MarginTrack)
+	{
+		FMovieSceneInterrogationData InterrogationData;
+		GetSequencer()->GetEvaluationTemplate().CopyActuators(InterrogationData.GetAccumulator());
+
+		FMovieSceneContext Context(FMovieSceneEvaluationRange(KeyTime, GetSequencer()->GetFocusedTickResolution()));
+		EvalTrack.Interrogate(Context, InterrogationData, Object);
+
+		FMargin Val(0.0f, 0.0f, 0.0f, 0.0f);
+		for (const FMargin& InMargin : InterrogationData.Iterate<FMargin>(UMovieSceneMarginSection::GetMarginInterrogationKey()))
+		{
+			Val = InMargin;
+			break;
+		}
+		FMovieSceneChannelProxy& Proxy = SectionToKey->GetChannelProxy();
+		GeneratedTotalKeys[0]->ModifyByCurrentAndWeight(Proxy, KeyTime, (void *)&Val.Left, Weight);
+		GeneratedTotalKeys[1]->ModifyByCurrentAndWeight(Proxy, KeyTime, (void *)&Val.Top, Weight);
+		GeneratedTotalKeys[2]->ModifyByCurrentAndWeight(Proxy, KeyTime, (void *)&Val.Right, Weight);
+		GeneratedTotalKeys[3]->ModifyByCurrentAndWeight(Proxy, KeyTime, (void *)&Val.Bottom, Weight);
+		return true;
+	}
+	return false;
+}
\ No newline at end of file
diff --git a/Engine/Source/Editor/UMGEditor/Private/Animation/MarginTrackEditor.h b/Engine/Source/Editor/UMGEditor/Private/Animation/MarginTrackEditor.h
--- a/Engine/Source/Editor/UMGEditor/Private/Animation/MarginTrackEditor.h
+++ b/Engine/Source/Editor/UMGEditor/Private/Animation/MarginTrackEditor.h
@@ -49,6 +49,7 @@ protected:
 	//~ FPropertyTrackEditor Interface
 
 	virtual void GenerateKeysFromPropertyChanged( const FPropertyChangedParams& PropertyChangedParams, FGeneratedTrackKeys& OutGeneratedKeys) override;
+	virtual bool ModifyGeneratedKeysByCurrentAndWeight(UObject *Object, UMovieSceneTrack *Track, UMovieSceneSection* SectionToKey, FFrameNumber KeyTime, FGeneratedTrackKeys& GeneratedTotalKeys, float Weight) const override;
 
 private:
 
diff --git a/Engine/Source/Editor/UMGEditor/Private/Animation/Sequencer2DTransformTrackEditor.cpp b/Engine/Source/Editor/UMGEditor/Private/Animation/Sequencer2DTransformTrackEditor.cpp
--- a/Engine/Source/Editor/UMGEditor/Private/Animation/Sequencer2DTransformTrackEditor.cpp
+++ b/Engine/Source/Editor/UMGEditor/Private/Animation/Sequencer2DTransformTrackEditor.cpp
@@ -5,7 +5,6 @@
 #include "Slate/WidgetTransform.h"
 #include "ISectionLayoutBuilder.h"
 
-
 FName F2DTransformTrackEditor::TranslationName( "Translation" );
 FName F2DTransformTrackEditor::ScaleName( "Scale" );
 FName F2DTransformTrackEditor::ShearName( "Shear" );
@@ -25,7 +24,6 @@ TSharedRef<ISequencerSection> F2DTransformTrackEditor::MakeSectionInterface(UMov
 	return MakeShared<F2DTransformSection>(SectionObject, GetSequencer());
 }
 
-
 void F2DTransformTrackEditor::GenerateKeysFromPropertyChanged( const FPropertyChangedParams& PropertyChangedParams, FGeneratedTrackKeys& OutGeneratedKeys )
 {
 	FPropertyPath StructPath = PropertyChangedParams.StructPathToKey;
@@ -98,4 +96,38 @@ void F2DTransformTrackEditor::GenerateKeysFromPropertyChanged( const FPropertyCh
 	OutGeneratedKeys.Add(FMovieSceneChannelValueSetter::Create<FMovieSceneFloatChannel>(4, Transform.Scale.Y,       bKeyScaleY));
 	OutGeneratedKeys.Add(FMovieSceneChannelValueSetter::Create<FMovieSceneFloatChannel>(5, Transform.Shear.X,       bKeyShearX));
 	OutGeneratedKeys.Add(FMovieSceneChannelValueSetter::Create<FMovieSceneFloatChannel>(6, Transform.Shear.Y,       bKeyShearY));
+}
+
+bool F2DTransformTrackEditor::ModifyGeneratedKeysByCurrentAndWeight(UObject *Object, UMovieSceneTrack *Track, UMovieSceneSection* SectionToKey, FFrameNumber KeyTime, FGeneratedTrackKeys& GeneratedTotalKeys, float Weight) const
+{
+	FFrameRate TickResolution = GetSequencer()->GetFocusedTickResolution();
+
+	UMovieScene2DTransformTrack* TransformTrack = Cast<UMovieScene2DTransformTrack>(Track);
+	FMovieSceneEvaluationTrack EvalTrack = Track->GenerateTrackTemplate();
+
+	if (TransformTrack)
+	{
+		FMovieSceneInterrogationData InterrogationData;
+		GetSequencer()->GetEvaluationTemplate().CopyActuators(InterrogationData.GetAccumulator());
+
+		FMovieSceneContext Context(FMovieSceneEvaluationRange(KeyTime, GetSequencer()->GetFocusedTickResolution()));
+		EvalTrack.Interrogate(Context, InterrogationData, Object);
+
+		FWidgetTransform Val;
+		for (const FWidgetTransform& InWidget : InterrogationData.Iterate<FWidgetTransform>(UMovieScene2DTransformSection::GetWidgetTransformInterrogationKey()))
+		{
+			Val = InWidget;
+			break;
+		}
+		FMovieSceneChannelProxy& Proxy = SectionToKey->GetChannelProxy();
+		GeneratedTotalKeys[0]->ModifyByCurrentAndWeight(Proxy, KeyTime, (void *)&Val.Translation.X, Weight);
+		GeneratedTotalKeys[1]->ModifyByCurrentAndWeight(Proxy, KeyTime, (void *)&Val.Translation.Y, Weight);
+		GeneratedTotalKeys[2]->ModifyByCurrentAndWeight(Proxy, KeyTime, (void *)&Val.Angle, Weight);
+		GeneratedTotalKeys[3]->ModifyByCurrentAndWeight(Proxy, KeyTime, (void *)&Val.Scale.X, Weight);
+		GeneratedTotalKeys[4]->ModifyByCurrentAndWeight(Proxy, KeyTime, (void *)&Val.Scale.Y, Weight);
+		GeneratedTotalKeys[5]->ModifyByCurrentAndWeight(Proxy, KeyTime, (void *)&Val.Shear.X, Weight);
+		GeneratedTotalKeys[6]->ModifyByCurrentAndWeight(Proxy, KeyTime, (void *)&Val.Shear.Y, Weight);
+		return true;
+	}
+	return false;
 }
\ No newline at end of file
diff --git a/Engine/Source/Editor/UMGEditor/Private/Animation/Sequencer2DTransformTrackEditor.h b/Engine/Source/Editor/UMGEditor/Private/Animation/Sequencer2DTransformTrackEditor.h
--- a/Engine/Source/Editor/UMGEditor/Private/Animation/Sequencer2DTransformTrackEditor.h
+++ b/Engine/Source/Editor/UMGEditor/Private/Animation/Sequencer2DTransformTrackEditor.h
@@ -51,6 +51,7 @@ protected:
 	// FPropertyTrackEditor interface
 
 	virtual void GenerateKeysFromPropertyChanged( const FPropertyChangedParams& PropertyChangedParams, FGeneratedTrackKeys& OutGeneratedKeys ) override;
+	virtual bool ModifyGeneratedKeysByCurrentAndWeight(UObject *Object, UMovieSceneTrack *Track, UMovieSceneSection* SectionToKey, FFrameNumber KeyTime, FGeneratedTrackKeys& GeneratedTotalKeys, float Weight) const override;
 
 private:
 	static FName TranslationName;
diff --git a/Engine/Source/Runtime/MovieScene/Private/Evaluation/Blending/MovieSceneMultiChannelBlending.cpp b/Engine/Source/Runtime/MovieScene/Private/Evaluation/Blending/MovieSceneMultiChannelBlending.cpp
--- a/Engine/Source/Runtime/MovieScene/Private/Evaluation/Blending/MovieSceneMultiChannelBlending.cpp
+++ b/Engine/Source/Runtime/MovieScene/Private/Evaluation/Blending/MovieSceneMultiChannelBlending.cpp
@@ -9,3 +9,4 @@ template<> FMovieSceneAnimTypeID GetBlendingDataType<FVector>()		{ static FMovie
 template<> FMovieSceneAnimTypeID GetBlendingDataType<FVector4>()	{ static FMovieSceneAnimTypeID TypeID = FMovieSceneAnimTypeID::Unique(); return TypeID; }
 template<> FMovieSceneAnimTypeID GetBlendingDataType<FTransform>()	{ static FMovieSceneAnimTypeID TypeID = FMovieSceneAnimTypeID::Unique(); return TypeID; }
 template<> FMovieSceneAnimTypeID GetBlendingDataType<FEulerTransform>()	{ static FMovieSceneAnimTypeID TypeID = FMovieSceneAnimTypeID::Unique(); return TypeID; }
+template<> FMovieSceneAnimTypeID GetBlendingDataType<FLinearColor>()   { static FMovieSceneAnimTypeID TypeId = FMovieSceneAnimTypeID::Unique(); return TypeId; }
\ No newline at end of file
diff --git a/Engine/Source/Runtime/MovieScene/Private/Evaluation/MovieScenePropertyTemplate.cpp b/Engine/Source/Runtime/MovieScene/Private/Evaluation/MovieScenePropertyTemplate.cpp
--- a/Engine/Source/Runtime/MovieScene/Private/Evaluation/MovieScenePropertyTemplate.cpp
+++ b/Engine/Source/Runtime/MovieScene/Private/Evaluation/MovieScenePropertyTemplate.cpp
@@ -61,6 +61,12 @@ const FMovieSceneInterrogationKey FMovieScenePropertySectionTemplate::GetVectorI
 	return TypeID;
 }
 const FMovieSceneInterrogationKey FMovieScenePropertySectionTemplate::GetVector2DInterrogationKey()
+{
+	const static FMovieSceneAnimTypeID TypeID = FMovieSceneAnimTypeID::Unique();
+	return TypeID;
+}
+
+const FMovieSceneInterrogationKey FMovieScenePropertySectionTemplate::GetColorInterrogationKey()
 {
 	const static FMovieSceneAnimTypeID TypeID = FMovieSceneAnimTypeID::Unique();
 	return TypeID;
diff --git a/Engine/Source/Runtime/MovieScene/Public/Evaluation/Blending/MovieSceneMultiChannelBlending.h b/Engine/Source/Runtime/MovieScene/Public/Evaluation/Blending/MovieSceneMultiChannelBlending.h
--- a/Engine/Source/Runtime/MovieScene/Public/Evaluation/Blending/MovieSceneMultiChannelBlending.h
+++ b/Engine/Source/Runtime/MovieScene/Public/Evaluation/Blending/MovieSceneMultiChannelBlending.h
@@ -307,6 +307,10 @@ namespace MovieScene
 	{
 		Out = { In.Location.X, In.Location.Y, In.Location.Z, In.Rotation.Roll, In.Rotation.Pitch, In.Rotation.Yaw, In.Scale.X, In.Scale.Y, In.Scale.Z };
 	}
+	inline void MultiChannelFromData(const FLinearColor& In, TMultiChannelValue<float, 4>& Out)
+	{
+		Out = { In.R, In.G, In.B, In.A };
+	}
 
 	inline void ResolveChannelsToData(const TMultiChannelValue<double, 1>& In, 	int32& Out)			{ Out = In[0]; }
 	inline void ResolveChannelsToData(const TMultiChannelValue<float, 1>& In, 	float& Out)			{ Out = In[0]; }
@@ -329,6 +333,10 @@ namespace MovieScene
 			FVector(In[6], In[7], In[8])
 			);
 	}
+	inline void ResolveChannelsToData(const TMultiChannelValue<float, 4>& In, FLinearColor& Out)
+	{
+		Out = FLinearColor(In[0], In[1], In[2], In[3]);
+	}
 }
 
 /** Define runtime type identifiers for the built-in C++ data types */
@@ -339,6 +347,7 @@ template<> MOVIESCENE_API FMovieSceneAnimTypeID GetBlendingDataType<FVector>();
 template<> MOVIESCENE_API FMovieSceneAnimTypeID GetBlendingDataType<FVector4>();
 template<> MOVIESCENE_API FMovieSceneAnimTypeID GetBlendingDataType<FTransform>();
 template<> MOVIESCENE_API FMovieSceneAnimTypeID GetBlendingDataType<FEulerTransform>();
+template<> MOVIESCENE_API FMovieSceneAnimTypeID GetBlendingDataType<FLinearColor>();
 
 /** Define working data types for blending calculations */
 template<> struct TBlendableTokenTraits<int32> 			{ typedef MovieScene::TMaskedBlendable<double,1> 	WorkingDataType; };
@@ -348,5 +357,4 @@ template<> struct TBlendableTokenTraits<FVector> 		{ typedef MovieScene::TMasked
 template<> struct TBlendableTokenTraits<FVector4> 		{ typedef MovieScene::TMaskedBlendable<float, 4> 	WorkingDataType; };
 template<> struct TBlendableTokenTraits<FTransform> 	{ typedef MovieScene::TMaskedBlendable<float, 9> 	WorkingDataType; };
 template<> struct TBlendableTokenTraits<FEulerTransform>{ typedef MovieScene::TMaskedBlendable<float, 9> 	WorkingDataType; };
-
-
+template<> struct TBlendableTokenTraits<FLinearColor>   { typedef MovieScene::TMaskedBlendable<float, 4>    WorkingDataType; };
\ No newline at end of file
diff --git a/Engine/Source/Runtime/MovieScene/Public/Evaluation/MovieScenePropertyTemplate.h b/Engine/Source/Runtime/MovieScene/Public/Evaluation/MovieScenePropertyTemplate.h
--- a/Engine/Source/Runtime/MovieScene/Public/Evaluation/MovieScenePropertyTemplate.h
+++ b/Engine/Source/Runtime/MovieScene/Public/Evaluation/MovieScenePropertyTemplate.h
@@ -241,6 +241,7 @@ public:
 	const static FMovieSceneInterrogationKey GetVector4InterrogationKey();
 	const static FMovieSceneInterrogationKey GetVectorInterrogationKey();
 	const static FMovieSceneInterrogationKey GetVector2DInterrogationKey();
+	const static FMovieSceneInterrogationKey GetColorInterrogationKey();
 
 protected:
 
diff --git a/Engine/Source/Runtime/MovieSceneTracks/Private/Evaluation/MovieSceneColorTemplate.cpp b/Engine/Source/Runtime/MovieSceneTracks/Private/Evaluation/MovieSceneColorTemplate.cpp
--- a/Engine/Source/Runtime/MovieSceneTracks/Private/Evaluation/MovieSceneColorTemplate.cpp
+++ b/Engine/Source/Runtime/MovieSceneTracks/Private/Evaluation/MovieSceneColorTemplate.cpp
@@ -16,34 +16,6 @@
 
 DECLARE_CYCLE_STAT(TEXT("Color Track Token Execute"), MovieSceneEval_ColorTrack_TokenExecute, STATGROUP_MovieSceneEval);
 
-/** Access the unique runtime type identifier for a widget transform. */
-template<> FMovieSceneAnimTypeID GetBlendingDataType<FLinearColor>()
-{
-	static FMovieSceneAnimTypeID TypeId = FMovieSceneAnimTypeID::Unique();
-	return TypeId;
-}
-
-/** Inform the blending accumulator to use a 7 channel float to blend margins */
-template<> struct TBlendableTokenTraits<FLinearColor>
-{
-	typedef MovieScene::TMaskedBlendable<float, 4> WorkingDataType;
-};
-
-namespace MovieScene
-{
-	/** Convert a color into a 4 channel float */
-	inline void MultiChannelFromData(const FLinearColor& In, TMultiChannelValue<float, 4>& Out)
-	{
-		Out = { In.R, In.G, In.B, In.A };
-	}
-
-	/** Convert a 4 channel float into a color */
-	inline void ResolveChannelsToData(const TMultiChannelValue<float, 4>& In, FLinearColor& Out)
-	{
-		Out = FLinearColor(In[0], In[1], In[2], In[3]);
-	}
-}
-
 enum class EColorType : uint8
 {
 	/** FSlateColor */
@@ -228,6 +200,12 @@ struct FColorTokenActuator : TMovieSceneBlendingActuator<FLinearColor>
 			FColorToken(InFinalValue).Apply(*InObject, PropertyBindings);
 		}
 	}
+	virtual void Actuate(FMovieSceneInterrogationData& InterrogationData, typename TCallTraits<FLinearColor>::ParamType InValue, const TBlendableTokenStack<FLinearColor>& OriginalStack, const FMovieSceneContext& Context) const
+	{
+		FLinearColor Value = InValue;
+		InterrogationData.Add(Value, FMovieScenePropertySectionTemplate::GetColorInterrogationKey());
+	};
+
 };
 
 FMovieSceneColorSectionTemplate::FMovieSceneColorSectionTemplate(const UMovieSceneColorSection& Section, const UMovieSceneColorTrack& Track)
@@ -272,3 +250,35 @@ void FMovieSceneColorSectionTemplate::Evaluate(const FMovieSceneEvaluationOperan
 		ExecutionTokens.BlendToken(ActuatorTypeID, TBlendableToken<FLinearColor>(AnimationData, BlendType, Weight));
 	}
 }
+
+void FMovieSceneColorSectionTemplate::Interrogate(const FMovieSceneContext& Context, FMovieSceneInterrogationData& Container, UObject* BindingOverride) const
+{
+	const FFrameTime Time = Context.GetTime();
+	MovieScene::TMultiChannelValue<float, 4> AnimationData;
+
+	for (uint8 Index = 0; Index < 4; ++Index)
+	{
+		float ChannelValue = 0.f;
+		if (Curves[Index].Evaluate(Time, ChannelValue))
+		{
+			AnimationData.Set(Index, ChannelValue);
+		}
+	}
+
+
+	FMovieSceneAnimTypeID TypeID = GetPropertyTypeID();
+	static FMovieSceneBlendingActuatorID ActuatorTypeID(TypeID);
+	if (!Container.GetAccumulator().FindActuator<FLinearColor>(ActuatorTypeID))
+	{
+		PropertyTemplate::FSectionData SectionData;
+		SectionData.Initialize(PropertyData.PropertyName, PropertyData.PropertyPath, PropertyData.FunctionName, PropertyData.NotifyFunctionName);
+		Container.GetAccumulator().DefineActuator(ActuatorTypeID, MakeShared<TPropertyActuator<FLinearColor>>(SectionData));
+	}
+
+	if (!AnimationData.IsEmpty())
+	{
+		// Add the blendable to the accumulator
+		float Weight = EvaluateEasing(Context.GetTime());
+		Container.GetAccumulator().BlendToken(FMovieSceneEvaluationOperand(), ActuatorTypeID, FMovieSceneEvaluationScope(), Context, TBlendableToken<FLinearColor>(AnimationData, BlendType, Weight));
+	}
+}
\ No newline at end of file
diff --git a/Engine/Source/Runtime/MovieSceneTracks/Private/Evaluation/MovieSceneColorTemplate.h b/Engine/Source/Runtime/MovieSceneTracks/Private/Evaluation/MovieSceneColorTemplate.h
--- a/Engine/Source/Runtime/MovieSceneTracks/Private/Evaluation/MovieSceneColorTemplate.h
+++ b/Engine/Source/Runtime/MovieSceneTracks/Private/Evaluation/MovieSceneColorTemplate.h
@@ -32,4 +32,6 @@ private:
 
 	virtual UScriptStruct& GetScriptStructImpl() const override { return *StaticStruct(); }
 	virtual void Evaluate(const FMovieSceneEvaluationOperand& Operand, const FMovieSceneContext& Context, const FPersistentEvaluationData& PersistentData, FMovieSceneExecutionTokens& ExecutionTokens) const override;
+	virtual void Interrogate(const FMovieSceneContext& Context, FMovieSceneInterrogationData& Container, UObject* BindingOverride) const override;
+
 };
diff --git a/Engine/Source/Runtime/MovieSceneTracks/Private/Sections/MovieSceneColorSection.cpp b/Engine/Source/Runtime/MovieSceneTracks/Private/Sections/MovieSceneColorSection.cpp
--- a/Engine/Source/Runtime/MovieSceneTracks/Private/Sections/MovieSceneColorSection.cpp
+++ b/Engine/Source/Runtime/MovieSceneTracks/Private/Sections/MovieSceneColorSection.cpp
@@ -5,6 +5,7 @@
 #include "UObject/SequencerObjectVersion.h"
 #include "Channels/MovieSceneChannelProxy.h"
 #include "Styling/SlateColor.h"
+#include "Evaluation/MovieScenePropertyTemplate.h"
 
 #if WITH_EDITOR
 struct FColorSectionEditorData
@@ -34,6 +35,11 @@ struct FColorSectionEditorData
 		ExternalValues[1].OnGetExternalValue = ExtractChannelG;
 		ExternalValues[2].OnGetExternalValue = ExtractChannelB;
 		ExternalValues[3].OnGetExternalValue = ExtractChannelA;
+		ExternalValues[0].OnGetCurrentValueAndWeight = GetChannelRValueAndWeight;
+		ExternalValues[1].OnGetCurrentValueAndWeight = GetChannelGValueAndWeight;
+		ExternalValues[2].OnGetCurrentValueAndWeight = GetChannelBValueAndWeight;
+		ExternalValues[3].OnGetCurrentValueAndWeight = GetChannelAValueAndWeight;
+
 	}
 
 	static FLinearColor GetPropertyValue(UObject& InObject, FTrackInstancePropertyBindings& Bindings)
@@ -78,6 +84,70 @@ struct FColorSectionEditorData
 		return Bindings ? GetPropertyValue(InObject, *Bindings).A : TOptional<float>();
 	}
 
+	static void GetChannelRValueAndWeight(UObject* Object, UMovieSceneSection*  SectionToKey, FFrameNumber KeyTime, FFrameRate TickResolution, FMovieSceneRootEvaluationTemplateInstance& RootTemplate,
+		float& OutValue, float OutWeight)
+	{
+		GetChannelValueAndWeight(0, Object, SectionToKey, KeyTime, TickResolution, RootTemplate, OutValue, OutWeight);
+	}
+	static void GetChannelGValueAndWeight(UObject* Object, UMovieSceneSection*  SectionToKey, FFrameNumber KeyTime, FFrameRate TickResolution, FMovieSceneRootEvaluationTemplateInstance& RootTemplate,
+		float& OutValue, float OutWeight)
+	{
+		GetChannelValueAndWeight(1, Object, SectionToKey, KeyTime, TickResolution, RootTemplate, OutValue, OutWeight);
+	}
+	static void GetChannelBValueAndWeight(UObject* Object, UMovieSceneSection*  SectionToKey, FFrameNumber KeyTime, FFrameRate TickResolution, FMovieSceneRootEvaluationTemplateInstance& RootTemplate,
+		float& OutValue, float OutWeight)
+	{
+		GetChannelValueAndWeight(2, Object, SectionToKey, KeyTime, TickResolution, RootTemplate, OutValue, OutWeight);
+	}
+	static void GetChannelAValueAndWeight(UObject* Object, UMovieSceneSection*  SectionToKey, FFrameNumber KeyTime, FFrameRate TickResolution, FMovieSceneRootEvaluationTemplateInstance& RootTemplate,
+		float& OutValue, float OutWeight)
+	{
+		GetChannelValueAndWeight(3, Object, SectionToKey, KeyTime, TickResolution, RootTemplate, OutValue, OutWeight);
+	}
+
+	static void GetChannelValueAndWeight(int32 Index, UObject* Object, UMovieSceneSection*  SectionToKey, FFrameNumber KeyTime, FFrameRate TickResolution, FMovieSceneRootEvaluationTemplateInstance& RootTemplate,
+		float& OutValue, float& OutWeight)
+	{
+		OutValue = 0.0f;
+		OutWeight = 1.0f;
+
+		UMovieSceneTrack* Track = SectionToKey->GetTypedOuter<UMovieSceneTrack>();
+
+		if (Track)
+		{
+			FMovieSceneEvaluationTrack EvalTrack = Track->GenerateTrackTemplate();
+			FMovieSceneInterrogationData InterrogationData;
+			RootTemplate.CopyActuators(InterrogationData.GetAccumulator());
+
+			FMovieSceneContext Context(FMovieSceneEvaluationRange(KeyTime, TickResolution));
+			EvalTrack.Interrogate(Context, InterrogationData, Object);
+
+			FLinearColor Val(0.0f, 0.0f, 0.0f, 0.0f);
+			for (const FLinearColor& InColor : InterrogationData.Iterate<FLinearColor>(FMovieScenePropertySectionTemplate::GetColorInterrogationKey()))
+			{
+				Val = InColor;
+				break;
+			}
+			switch (Index)
+			{
+			case 0:
+				OutValue = Val.R;
+				break;
+			case 1:
+				OutValue = Val.G;
+				break;
+			case 2:
+				OutValue = Val.B;
+				break;
+			case 3:
+				OutValue = Val.A;
+				break;
+			}
+			
+		}
+		OutWeight = MovieSceneHelpers::CalculateWeightForBlending(SectionToKey, KeyTime);
+	}
+
 	FMovieSceneChannelMetaData      MetaData[4];
 	TMovieSceneExternalValue<float> ExternalValues[4];
 };
diff --git a/Engine/Source/Runtime/UMG/Public/Animation/MovieScene2DTransformSection.h b/Engine/Source/Runtime/UMG/Public/Animation/MovieScene2DTransformSection.h
--- a/Engine/Source/Runtime/UMG/Public/Animation/MovieScene2DTransformSection.h
+++ b/Engine/Source/Runtime/UMG/Public/Animation/MovieScene2DTransformSection.h
@@ -122,7 +122,7 @@ protected:
 
 	void UpdateChannelProxy();
 public:
-	const static FMovieSceneInterrogationKey GetWidgetTransformInterrogationKey();
+	UMG_API const static FMovieSceneInterrogationKey GetWidgetTransformInterrogationKey();
 
 public:
 
diff --git a/Engine/Source/Runtime/UMG/Public/Animation/MovieSceneMarginSection.h b/Engine/Source/Runtime/UMG/Public/Animation/MovieSceneMarginSection.h
--- a/Engine/Source/Runtime/UMG/Public/Animation/MovieSceneMarginSection.h
+++ b/Engine/Source/Runtime/UMG/Public/Animation/MovieSceneMarginSection.h
@@ -21,7 +21,7 @@ class UMovieSceneMarginSection
 {
 	GENERATED_UCLASS_BODY()
 public:
-	const static FMovieSceneInterrogationKey GetMarginInterrogationKey();
+	UMG_API const static FMovieSceneInterrogationKey GetMarginInterrogationKey();
 
 public:
 
