Sequencer: Repeater Event Tracks now fire properly when playing in reverse.

diff --git a/Engine/Source/Runtime/MovieSceneTracks/Private/Evaluation/MovieSceneEventTemplate.cpp b/Engine/Source/Runtime/MovieSceneTracks/Private/Evaluation/MovieSceneEventTemplate.cpp
--- a/Engine/Source/Runtime/MovieSceneTracks/Private/Evaluation/MovieSceneEventTemplate.cpp
+++ b/Engine/Source/Runtime/MovieSceneTracks/Private/Evaluation/MovieSceneEventTemplate.cpp
@@ -252,7 +252,7 @@ struct FEventTriggerExecutionToken
 #if WITH_EDITOR
 			else if (!bIsGameWorld && !Function->HasMetaData(NAME_CallInEditor))
 			{
-				UE_LOG(LogMovieScene, Log, TEXT("Refusing to trigger event '%s' in editor world when 'Call in Editor' is false."), *EventName.ToString());
+				UE_LOG(LogMovieScene, Verbose, TEXT("Refusing to trigger event '%s' in editor world when 'Call in Editor' is false."), *EventName.ToString());
 				continue;
 			}
 #endif // WITH_EDITOR
@@ -521,15 +521,18 @@ FMovieSceneEventRepeaterTemplate::FMovieSceneEventRepeaterTemplate(const UMovieS
 
 void FMovieSceneEventRepeaterTemplate::EvaluateSwept(const FMovieSceneEvaluationOperand& Operand, const FMovieSceneContext& Context, const TRange<FFrameNumber>& SweptRange, const FPersistentEvaluationData& PersistentData, FMovieSceneExecutionTokens& ExecutionTokens) const
 {
+	const bool bBackwards = Context.GetDirection() == EPlayDirection::Backwards;
+	FFrameNumber CurrentFrame = bBackwards ? Context.GetTime().CeilToFrame() : Context.GetTime().FloorToFrame();
+
 	// Don't allow events to fire when playback is in a stopped state. This can occur when stopping 
 	// playback and returning the current position to the start of playback. It's not desireable to have 
 	// all the events from the last playback position to the start of playback be fired.
-	if (!SweptRange.Contains(Context.GetTime().FrameNumber) || Context.GetStatus() == EMovieScenePlayerStatus::Stopped || Context.IsSilent())
+	if (!SweptRange.Contains(CurrentFrame) || Context.GetStatus() == EMovieScenePlayerStatus::Stopped || Context.IsSilent())
 	{
 		return;
 	}
 
-	const bool bBackwards = Context.GetDirection() == EPlayDirection::Backwards;
+	
 	if ((!bBackwards && bFireEventsWhenForwards) || (bBackwards && bFireEventsWhenBackwards))
 	{
 		TArray<FName> Events = { EventToTrigger };
