Sequencer: Implement OnGetCurrentValueAndWeight for MovieSceneFloatSection

diff --git a/Engine/Source/Runtime/MovieSceneTracks/Private/Sections/MovieSceneFloatSection.cpp b/Engine/Source/Runtime/MovieSceneTracks/Private/Sections/MovieSceneFloatSection.cpp
--- a/Engine/Source/Runtime/MovieSceneTracks/Private/Sections/MovieSceneFloatSection.cpp
+++ b/Engine/Source/Runtime/MovieSceneTracks/Private/Sections/MovieSceneFloatSection.cpp
@@ -3,6 +3,9 @@
 #include "Sections/MovieSceneFloatSection.h"
 #include "UObject/SequencerObjectVersion.h"
 #include "Channels/MovieSceneChannelProxy.h"
+#include "MovieSceneCommonHelpers.h"
+#include "Evaluation/MovieScenePropertyTemplate.h"
+
 
 UMovieSceneFloatSection::UMovieSceneFloatSection( const FObjectInitializer& ObjectInitializer )
 	: Super( ObjectInitializer )
@@ -17,7 +20,39 @@ UMovieSceneFloatSection::UMovieSceneFloatSection( const FObjectInitializer& Obje
 	bSupportsInfiniteRange = true;
 #if WITH_EDITOR
 
-	ChannelProxy = MakeShared<FMovieSceneChannelProxy>(FloatCurve, FMovieSceneChannelMetaData(), TMovieSceneExternalValue<float>::Make());
+	struct FGetCurrentValueAndWeight
+	{
+		static void GetFloatValueAndWeight(UObject* Object, UMovieSceneSection*  SectionToKey, FFrameNumber KeyTime, FFrameRate TickResolution, FMovieSceneRootEvaluationTemplateInstance& RootTemplate,
+			float& OutValue, float& OutWeight)
+		{
+			OutValue = 0.0f;
+			OutWeight = 1.0f;
+
+			UMovieSceneTrack* Track = SectionToKey->GetTypedOuter<UMovieSceneTrack>();
+
+			if (Track)
+			{
+				FMovieSceneEvaluationTrack EvalTrack = Track->GenerateTrackTemplate();
+				FMovieSceneInterrogationData InterrogationData;
+				RootTemplate.CopyActuators(InterrogationData.GetAccumulator());
+
+				FMovieSceneContext Context(FMovieSceneEvaluationRange(KeyTime, TickResolution));
+				EvalTrack.Interrogate(Context, InterrogationData, Object);
+
+				for (const float &Value : InterrogationData.Iterate<float>(FMovieScenePropertySectionTemplate::GetFloatInterrogationKey()))
+				{
+					OutValue = Value;
+					break;
+				}
+			}
+			OutWeight = MovieSceneHelpers::CalculateWeightForBlending(SectionToKey, KeyTime);
+		}
+	};
+
+	TMovieSceneExternalValue<float> ExternalValue;
+	ExternalValue.OnGetCurrentValueAndWeight = FGetCurrentValueAndWeight::GetFloatValueAndWeight;
+
+	ChannelProxy = MakeShared<FMovieSceneChannelProxy>(FloatCurve, FMovieSceneChannelMetaData(), ExternalValue);
 
 #else
 
