Sequencer: Stop only the audio components corresponding to the audio data row index on tear down. Previously, this was stopping all sounds, causing some hitches in restarting active sounds that are still being evaluated next frame.

diff --git a/Engine/Source/Runtime/MovieSceneTracks/Private/Evaluation/MovieSceneAudioTemplate.cpp b/Engine/Source/Runtime/MovieSceneTracks/Private/Evaluation/MovieSceneAudioTemplate.cpp
--- a/Engine/Source/Runtime/MovieSceneTracks/Private/Evaluation/MovieSceneAudioTemplate.cpp
+++ b/Engine/Source/Runtime/MovieSceneTracks/Private/Evaluation/MovieSceneAudioTemplate.cpp
@@ -173,6 +173,20 @@ struct FCachedAudioTrackData : IPersistentEvaluationData
 			}
 		}
 	}
+
+	void StopSoundsOnRow(int32 RowIndex)
+	{
+		if (RowIndex >= 0 && RowIndex < AudioComponentsByRow.Num())
+		{
+			for (auto& Pair : AudioComponentsByRow[RowIndex])
+			{
+				if (UAudioComponent* AudioComponent = Pair.Value.Get())
+				{
+					AudioComponent->Stop();
+				}
+			}
+		}
+	}
 };
 
 
@@ -394,6 +408,6 @@ void FMovieSceneAudioSectionTemplate::TearDown(FPersistentEvaluationData& Persis
 	{
 		FCachedAudioTrackData& TrackData = PersistentData.GetOrAddTrackData<FCachedAudioTrackData>();
 
-		TrackData.StopAllSounds();
+		TrackData.StopSoundsOnRow(AudioData.RowIndex);
 	}
 }
\ No newline at end of file
