UI for vector keys show wrong values
Need to set the values on the key struct.

diff --git a/Engine/Source/Runtime/MovieSceneTracks/Private/Sections/MovieSceneVectorSection.cpp b/Engine/Source/Runtime/MovieSceneTracks/Private/Sections/MovieSceneVectorSection.cpp
--- a/Engine/Source/Runtime/MovieSceneTracks/Private/Sections/MovieSceneVectorSection.cpp
+++ b/Engine/Source/Runtime/MovieSceneTracks/Private/Sections/MovieSceneVectorSection.cpp
@@ -279,6 +279,7 @@ TSharedPtr<FStructOnScope> UMovieSceneVectorSection::GetKeyStruct(TArrayView<con
 			Struct->KeyStructInterop.Add(FMovieSceneChannelValueHelper(ChannelProxy->MakeHandle<FMovieSceneFloatChannel>(Index), Struct->GetPropertyChannelByIndex(Index), KeyHandles));
 		}
 
+		Struct->KeyStructInterop.SetStartingValues();
 		Struct->Time = Struct->KeyStructInterop.GetUnifiedKeyTime().Get(0);
 	}
 
