Sequencer: Fix audio play clamping at incorrect sound duration

diff --git a/Engine/Source/Runtime/MovieSceneTracks/Private/Evaluation/MovieSceneAudioTemplate.cpp b/Engine/Source/Runtime/MovieSceneTracks/Private/Evaluation/MovieSceneAudioTemplate.cpp
--- a/Engine/Source/Runtime/MovieSceneTracks/Private/Evaluation/MovieSceneAudioTemplate.cpp
+++ b/Engine/Source/Runtime/MovieSceneTracks/Private/Evaluation/MovieSceneAudioTemplate.cpp
@@ -346,7 +346,7 @@ void FMovieSceneAudioSectionTemplateData::EnsureAudioIsPlaying(UAudioComponent&
 		}
 
 		const float AudioTime = (Context.GetTime() / Context.GetFrameRate()) - SectionStartTimeSeconds + (float)Context.GetFrameRate().AsSeconds(AudioStartOffset);
-		if (AudioTime > 0.f && AudioTime < AudioComponent.Sound->Duration)
+		if (AudioTime >= 0.f && AudioTime < AudioComponent.Sound->GetDuration())
 		{
 			AudioComponent.Play(AudioTime);
 		}
@@ -410,4 +410,6 @@ void FMovieSceneAudioSectionTemplate::TearDown(FPersistentEvaluationData& Persis
 
 		TrackData.StopSoundsOnRow(AudioData.RowIndex);
 	}
-}
\ No newline at end of file
+}
+
+PRAGMA_ENABLE_OPTIMIZATION
\ No newline at end of file
