Remove accidental PRAGMA_OPTIMIZATION

diff --git a/Engine/Source/Runtime/MovieSceneTracks/Private/Evaluation/MovieSceneAudioTemplate.cpp b/Engine/Source/Runtime/MovieSceneTracks/Private/Evaluation/MovieSceneAudioTemplate.cpp
--- a/Engine/Source/Runtime/MovieSceneTracks/Private/Evaluation/MovieSceneAudioTemplate.cpp
+++ b/Engine/Source/Runtime/MovieSceneTracks/Private/Evaluation/MovieSceneAudioTemplate.cpp
@@ -411,5 +411,3 @@ void FMovieSceneAudioSectionTemplate::TearDown(FPersistentEvaluationData& Persis
 		TrackData.StopSoundsOnRow(AudioData.RowIndex);
 	}
 }
-
-PRAGMA_ENABLE_OPTIMIZATION
\ No newline at end of file
