Fix animation playback by making sure last animation frame is called by rounding up the frame and then adding 1 to it(like what happens with other sections w ith the last key at section end). Also needed to tweak calculation of anim time since Mod(EndValue,EndValue) = 0 not EndValue.

diff --git a/Engine/Source/Runtime/MovieSceneTracks/Private/Evaluation/MovieSceneSkeletalAnimationTemplate.cpp b/Engine/Source/Runtime/MovieSceneTracks/Private/Evaluation/MovieSceneSkeletalAnimationTemplate.cpp
--- a/Engine/Source/Runtime/MovieSceneTracks/Private/Evaluation/MovieSceneSkeletalAnimationTemplate.cpp
+++ b/Engine/Source/Runtime/MovieSceneTracks/Private/Evaluation/MovieSceneSkeletalAnimationTemplate.cpp
@@ -439,7 +439,13 @@ void FMovieSceneSkeletalAnimationSectionTemplate::Evaluate(const FMovieSceneEval
 
 float FMovieSceneSkeletalAnimationSectionTemplateParameters::MapTimeToAnimation(FFrameTime InPosition, FFrameRate InFrameRate) const
 {
-	InPosition = FMath::Clamp(InPosition, FFrameTime(SectionStartTime), FFrameTime(SectionEndTime-1));
+
+	FFrameTime AnimationLength = GetSequenceLength() * InFrameRate;
+	int32 LengthInFrames = AnimationLength.FrameNumber.Value + (int)(AnimationLength.GetSubFrame() + 0.5f) + 1;
+	//we only play end if we are not looping, and assuming we are looping if Length is greater than default length;
+	bool bLooping = (SectionEndTime.Value - SectionStartTime.Value) > LengthInFrames;
+
+	InPosition = FMath::Clamp(InPosition, FFrameTime(SectionStartTime), FFrameTime(SectionEndTime -1));
 	
 	const float SectionPlayRate = PlayRate * Animation->RateScale;
 	const float AnimPlayRate = FMath::IsNearlyZero(SectionPlayRate) ? 1.0f : SectionPlayRate;
@@ -447,7 +453,7 @@ float FMovieSceneSkeletalAnimationSectionTemplateParameters::MapTimeToAnimation(
 	const float SeqLength = GetSequenceLength() - InFrameRate.AsSeconds(StartFrameOffset + EndFrameOffset);
 
 	float AnimPosition = FFrameTime::FromDecimal((InPosition - SectionStartTime).AsDecimal() * AnimPlayRate) / InFrameRate;
-	if (SeqLength > 0.f)
+	if (SeqLength > 0.f && (bLooping || !FMath::IsNearlyEqual(AnimPosition, SeqLength,1e-4f)))
 	{
 		AnimPosition = FMath::Fmod(AnimPosition, SeqLength);
 	}
diff --git a/Engine/Source/Runtime/MovieSceneTracks/Private/Sections/MovieSceneSkeletalAnimationSection.cpp b/Engine/Source/Runtime/MovieSceneTracks/Private/Sections/MovieSceneSkeletalAnimationSection.cpp
--- a/Engine/Source/Runtime/MovieSceneTracks/Private/Sections/MovieSceneSkeletalAnimationSection.cpp
+++ b/Engine/Source/Runtime/MovieSceneTracks/Private/Sections/MovieSceneSkeletalAnimationSection.cpp
@@ -185,8 +185,9 @@ TOptional<TRange<FFrameNumber> > UMovieSceneSkeletalAnimationSection::GetAutoSiz
 	FFrameRate FrameRate = GetTypedOuter<UMovieScene>()->GetTickResolution();
 
 	FFrameTime AnimationLength = Params.GetSequenceLength() * FrameRate;
+	int32 IFrameNumber = AnimationLength.FrameNumber.Value + (int)(AnimationLength.GetSubFrame() + 0.5f);
 
-	return TRange<FFrameNumber>(GetInclusiveStartFrame(), GetInclusiveStartFrame() + AnimationLength.FrameNumber);
+	return TRange<FFrameNumber>(GetInclusiveStartFrame(), GetInclusiveStartFrame() + IFrameNumber + 1);
 }
 
 
diff --git a/Engine/Source/Runtime/MovieSceneTracks/Private/Tracks/MovieSceneSkeletalAnimationTrack.cpp b/Engine/Source/Runtime/MovieSceneTracks/Private/Tracks/MovieSceneSkeletalAnimationTrack.cpp
--- a/Engine/Source/Runtime/MovieSceneTracks/Private/Tracks/MovieSceneSkeletalAnimationTrack.cpp
+++ b/Engine/Source/Runtime/MovieSceneTracks/Private/Tracks/MovieSceneSkeletalAnimationTrack.cpp
@@ -38,7 +38,8 @@ UMovieSceneSection* UMovieSceneSkeletalAnimationTrack::AddNewAnimationOnRow(FFra
 	UMovieSceneSkeletalAnimationSection* NewSection = Cast<UMovieSceneSkeletalAnimationSection>(CreateNewSection());
 	{
 		FFrameTime AnimationLength = AnimSequence->SequenceLength * GetTypedOuter<UMovieScene>()->GetTickResolution();
-		NewSection->InitialPlacementOnRow(AnimationSections, KeyTime, AnimationLength.FrameNumber.Value, RowIndex);
+		int32 IFrameNumber = AnimationLength.FrameNumber.Value + (int)(AnimationLength.GetSubFrame() + 0.5f) + 1;
+		NewSection->InitialPlacementOnRow(AnimationSections, KeyTime, IFrameNumber, RowIndex);
 		NewSection->Params.Animation = AnimSequence;
 	}
 
