Sequencer: Fix issue when setting a value on the last frame of the section

diff --git a/Engine/Source/Runtime/MovieSceneTracks/Private/Tracks/MovieScenePropertyTrack.cpp b/Engine/Source/Runtime/MovieSceneTracks/Private/Tracks/MovieScenePropertyTrack.cpp
--- a/Engine/Source/Runtime/MovieSceneTracks/Private/Tracks/MovieScenePropertyTrack.cpp
+++ b/Engine/Source/Runtime/MovieSceneTracks/Private/Tracks/MovieScenePropertyTrack.cpp
@@ -158,9 +158,12 @@ UMovieSceneSection* UMovieScenePropertyTrack::FindOrExtendSection(FFrameNumber T
 		bool bCalculateWeight = false;
 		if (SectionToKey && !OverlappingSections.Contains(SectionToKey))
 		{
-			if (SectionToKey->HasEndFrame() && SectionToKey->GetExclusiveEndFrame() < Time)
+			if (SectionToKey->HasEndFrame() && SectionToKey->GetExclusiveEndFrame() <= Time)
 			{
-				SectionToKey->SetEndFrame(Time);
+				if (SectionToKey->GetExclusiveEndFrame() != Time)
+				{
+					SectionToKey->SetEndFrame(Time);
+				}
 			}
 			else
 			{
@@ -223,10 +226,13 @@ UMovieSceneSection* UMovieScenePropertyTrack::FindOrExtendSection(FFrameNumber T
 			{
 				// SectionIndex == 0 
 				UMovieSceneSection* PreviousSection = Sections[0];
-				if(PreviousSection->HasEndFrame() && PreviousSection->GetExclusiveEndFrame() < Time)
+				if(PreviousSection->HasEndFrame() && PreviousSection->GetExclusiveEndFrame() <= Time)
 				{
 					// Append and grow the section
-					PreviousSection->SetEndFrame(Time);
+					if (PreviousSection->GetExclusiveEndFrame() != Time)
+					{
+						PreviousSection->SetEndFrame(Time);
+					}
 				}
 				else
 				{
