Sequencer: Changed "Refusing to trigger event" log to Verbose

diff --git a/Engine/Source/Runtime/MovieSceneTracks/Private/Evaluation/MovieSceneEventTemplate.cpp b/Engine/Source/Runtime/MovieSceneTracks/Private/Evaluation/MovieSceneEventTemplate.cpp
--- a/Engine/Source/Runtime/MovieSceneTracks/Private/Evaluation/MovieSceneEventTemplate.cpp
+++ b/Engine/Source/Runtime/MovieSceneTracks/Private/Evaluation/MovieSceneEventTemplate.cpp
@@ -252,7 +252,7 @@ struct FEventTriggerExecutionToken
 #if WITH_EDITOR
 			else if (!bIsGameWorld && !Function->HasMetaData(NAME_CallInEditor))
 			{
-				UE_LOG(LogMovieScene, Log, TEXT("Refusing to trigger event '%s' in editor world when 'Call in Editor' is false."), *EventName.ToString());
+				UE_LOG(LogMovieScene, Verbose, TEXT("Refusing to trigger event '%s' in editor world when 'Call in Editor' is false."), *EventName.ToString());
 				continue;
 			}
 #endif // WITH_EDITOR
