Sequencer: Allow for event receivers in addition to the object binding

diff --git a/Engine/Source/Runtime/MovieSceneTracks/Private/Evaluation/MovieSceneEventTemplate.cpp b/Engine/Source/Runtime/MovieSceneTracks/Private/Evaluation/MovieSceneEventTemplate.cpp
--- a/Engine/Source/Runtime/MovieSceneTracks/Private/Evaluation/MovieSceneEventTemplate.cpp
+++ b/Engine/Source/Runtime/MovieSceneTracks/Private/Evaluation/MovieSceneEventTemplate.cpp
@@ -200,8 +200,9 @@ struct FEventTriggerExecutionToken
 				}
 			}
 		}
+
 		// If we have specified event receivers
-		else if (EventReceivers.Num())
+		if (EventReceivers.Num())
 		{
 			EventContexts.Reserve(EventReceivers.Num());
 			for (FMovieSceneObjectBindingID ID : EventReceivers)
@@ -219,8 +220,9 @@ struct FEventTriggerExecutionToken
 				}
 			}
 		}
-		// If we haven't specified event receivers, use the default set defined on the player
-		else
+
+		// If we haven't specified event receivers, use the default set defined on the player		
+		if (EventContexts.Num() == 0)
 		{
 			EventContexts = Player.GetEventContexts();
 		}
