Make some blueprint categories consistent. Changed "Cinematic" to "Cinematics" "Cinematic" and "Sequencer" to "Sequence" to follow existing categories.

diff --git a/Engine/Source/Runtime/MovieScene/Public/Generators/MovieSceneEasingFunction.h b/Engine/Source/Runtime/MovieScene/Public/Generators/MovieSceneEasingFunction.h
--- a/Engine/Source/Runtime/MovieScene/Public/Generators/MovieSceneEasingFunction.h
+++ b/Engine/Source/Runtime/MovieScene/Public/Generators/MovieSceneEasingFunction.h
@@ -27,7 +27,7 @@ public:
 protected:
 
 	/** Evaluate the easing with an interpolation value between 0 and 1 */
-	UFUNCTION(BlueprintCallable, BlueprintImplementableEvent, Category = "Sequencer", meta=(CallInEditor="true"))
+	UFUNCTION(BlueprintCallable, BlueprintImplementableEvent, Category = "Sequence", meta=(CallInEditor="true"))
 	float OnEvaluate(float Interp) const;
 
 	/** Evaluate the easing with an interpolation value between 0 and 1 */
diff --git a/Engine/Source/Runtime/MovieSceneCapture/Public/MovieSceneCaptureEnvironment.h b/Engine/Source/Runtime/MovieSceneCapture/Public/MovieSceneCaptureEnvironment.h
--- a/Engine/Source/Runtime/MovieSceneCapture/Public/MovieSceneCaptureEnvironment.h
+++ b/Engine/Source/Runtime/MovieSceneCapture/Public/MovieSceneCaptureEnvironment.h
@@ -17,11 +17,11 @@ public:
 	GENERATED_BODY()
 
 	/** Get the frame number of the current capture */
-	UFUNCTION(BlueprintPure, Category="Cinematic|Capture")
+	UFUNCTION(BlueprintPure, Category="Cinematics|Capture")
 	static int32 GetCaptureFrameNumber();
 
 	/** Get the total elapsed time of the current capture in seconds */
-	UFUNCTION(BlueprintPure, Category="Cinematic|Capture")
+	UFUNCTION(BlueprintPure, Category="Cinematics|Capture")
 	static float GetCaptureElapsedTime();
 
 	/**
diff --git a/Engine/Source/Runtime/MovieSceneTracks/Public/Tracks/IMovieSceneTransformOrigin.h b/Engine/Source/Runtime/MovieSceneTracks/Public/Tracks/IMovieSceneTransformOrigin.h
--- a/Engine/Source/Runtime/MovieSceneTracks/Public/Tracks/IMovieSceneTransformOrigin.h
+++ b/Engine/Source/Runtime/MovieSceneTracks/Public/Tracks/IMovieSceneTransformOrigin.h
@@ -30,7 +30,7 @@ public:
 protected:
 
 	/** Get the transform from which all absolute component transform sections should be relative. Scale is ignored. */
-	UFUNCTION(BlueprintCallable, BlueprintImplementableEvent, Category = "Sequencer", DisplayName = "GetTransformOrigin", meta=(CallInEditor="true"))
+	UFUNCTION(BlueprintCallable, BlueprintImplementableEvent, Category = "Sequence", DisplayName = "GetTransformOrigin", meta=(CallInEditor="true"))
 	FTransform BP_GetTransformOrigin() const;
 	virtual FTransform NativeGetTransformOrigin() const { return BP_GetTransformOrigin(); }
 };
