Sequencer: Fixed actor reference tracks that operated on soft actor properties

diff --git a/Engine/Source/Editor/MovieSceneTools/Private/Channels/BuiltInChannelEditors.cpp b/Engine/Source/Editor/MovieSceneTools/Private/Channels/BuiltInChannelEditors.cpp
--- a/Engine/Source/Editor/MovieSceneTools/Private/Channels/BuiltInChannelEditors.cpp
+++ b/Engine/Source/Editor/MovieSceneTools/Private/Channels/BuiltInChannelEditors.cpp
@@ -72,7 +72,8 @@ FKeyHandle AddOrUpdateKey(FMovieSceneActorReferenceData* Channel, UMovieSceneSec
 		{
 			if (UObject* Object = WeakObject.Get())
 			{
-				CurrentActor = PropertyBindings->GetCurrentValue<AActor*>(*Object);
+				// Care is taken here to ensure that we call GetCurrentValue with the correct instantiation of UObject* rather than AActor*
+				CurrentActor = Cast<AActor>(PropertyBindings->GetCurrentValue<UObject*>(*Object));
 				break;
 			}
 		}
diff --git a/Engine/Source/Editor/MovieSceneTools/Private/TrackEditors/PropertyTrackEditors/ActorReferencePropertyTrackEditor.cpp b/Engine/Source/Editor/MovieSceneTools/Private/TrackEditors/PropertyTrackEditors/ActorReferencePropertyTrackEditor.cpp
--- a/Engine/Source/Editor/MovieSceneTools/Private/TrackEditors/PropertyTrackEditors/ActorReferencePropertyTrackEditor.cpp
+++ b/Engine/Source/Editor/MovieSceneTools/Private/TrackEditors/PropertyTrackEditors/ActorReferencePropertyTrackEditor.cpp
@@ -11,7 +11,8 @@ TSharedRef<ISequencerTrackEditor> FActorReferencePropertyTrackEditor::CreateTrac
 
 void FActorReferencePropertyTrackEditor::GenerateKeysFromPropertyChanged( const FPropertyChangedParams& PropertyChangedParams, FGeneratedTrackKeys& OutGeneratedKeys )
 {
-	AActor* NewReferencedActor = PropertyChangedParams.GetPropertyValue<AActor*>();
+	// Care is taken here to ensure that GetPropertyValue is templated on UObject* which causes it to use the correct instantiation of GetPropertyValueImpl
+	AActor* NewReferencedActor = Cast<AActor>(PropertyChangedParams.GetPropertyValue<UObject*>());
 	if ( NewReferencedActor != nullptr )
 	{
 		FGuid ActorGuid = GetSequencer()->GetHandleToObject( NewReferencedActor );
diff --git a/Engine/Source/Runtime/MovieSceneTracks/Private/Evaluation/MovieSceneActorReferenceTemplate.cpp b/Engine/Source/Runtime/MovieSceneTracks/Private/Evaluation/MovieSceneActorReferenceTemplate.cpp
--- a/Engine/Source/Runtime/MovieSceneTracks/Private/Evaluation/MovieSceneActorReferenceTemplate.cpp
+++ b/Engine/Source/Runtime/MovieSceneTracks/Private/Evaluation/MovieSceneActorReferenceTemplate.cpp
@@ -11,7 +11,7 @@
 namespace PropertyTemplate
 {
 	template<>
-	AActor* ConvertFromIntermediateType<AActor*, FMovieSceneObjectBindingID>(const FMovieSceneObjectBindingID& InObjectBinding, const FMovieSceneEvaluationOperand& Operand, FPersistentEvaluationData& PersistentData, IMovieScenePlayer& Player)
+	UObject* ConvertFromIntermediateType<UObject*, FMovieSceneObjectBindingID>(const FMovieSceneObjectBindingID& InObjectBinding, const FMovieSceneEvaluationOperand& Operand, FPersistentEvaluationData& PersistentData, IMovieScenePlayer& Player)
 	{
 		FMovieSceneObjectBindingID ResolvedID = InObjectBinding.ResolveLocalToRoot(Operand.SequenceID, Player.GetEvaluationTemplate().GetHierarchy());
 
@@ -26,25 +26,25 @@ namespace PropertyTemplate
 	}
 
 	template<>
-	AActor* ConvertFromIntermediateType<AActor*, TWeakObjectPtr<AActor>>(const TWeakObjectPtr<AActor>& InWeakPtr, IMovieScenePlayer& Player)
+	UObject* ConvertFromIntermediateType<UObject*, TWeakObjectPtr<>>(const TWeakObjectPtr<>& InWeakPtr, IMovieScenePlayer& Player)
 	{
 		return InWeakPtr.Get();
 	}
 
 	template<>
-	AActor* ConvertFromIntermediateType<AActor*, TWeakObjectPtr<AActor>>(const TWeakObjectPtr<AActor>& InWeakPtr, const FMovieSceneEvaluationOperand& Operand, FPersistentEvaluationData& PersistentData, IMovieScenePlayer& Player)
+	UObject* ConvertFromIntermediateType<UObject*, TWeakObjectPtr<>>(const TWeakObjectPtr<>& InWeakPtr, const FMovieSceneEvaluationOperand& Operand, FPersistentEvaluationData& PersistentData, IMovieScenePlayer& Player)
 	{
 		return InWeakPtr.Get();
 	}
 
-	static bool IsValueValid(AActor* InValue)
+	static bool IsValueValid(UObject* InValue)
 	{
 		return InValue != nullptr;
 	}
 
-	template<> IMovieScenePreAnimatedTokenPtr CacheExistingState<AActor*, FMovieSceneObjectBindingID>(UObject& Object, FTrackInstancePropertyBindings& PropertyBindings)
+	template<> IMovieScenePreAnimatedTokenPtr CacheExistingState<UObject*, FMovieSceneObjectBindingID>(UObject& Object, FTrackInstancePropertyBindings& PropertyBindings)
 	{
-		return TCachedState<AActor*, TWeakObjectPtr<AActor>>(PropertyBindings.GetCurrentValue<AActor*>(Object), PropertyBindings);
+		return TCachedState<UObject*, TWeakObjectPtr<>>(PropertyBindings.GetCurrentValue<UObject*>(Object), PropertyBindings);
 	}
 }
 
@@ -59,5 +59,5 @@ void FMovieSceneActorReferenceSectionTemplate::Evaluate(const FMovieSceneEvaluat
 	using namespace PropertyTemplate;
 
 	FMovieSceneActorReferenceKey ObjectBinding = ActorReferenceData.Evaluate(Context.GetTime());
-	ExecutionTokens.Add(TPropertyTrackExecutionToken<AActor*, FMovieSceneObjectBindingID>(ObjectBinding.Object));
+	ExecutionTokens.Add(TPropertyTrackExecutionToken<UObject*, FMovieSceneObjectBindingID>(ObjectBinding.Object));
 }
