Add support for blending multiple audio tracks on the same sequencer row.

diff --git a/Engine/Source/Runtime/MovieSceneTracks/Private/Evaluation/MovieSceneAudioTemplate.cpp b/Engine/Source/Runtime/MovieSceneTracks/Private/Evaluation/MovieSceneAudioTemplate.cpp
--- a/Engine/Source/Runtime/MovieSceneTracks/Private/Evaluation/MovieSceneAudioTemplate.cpp
+++ b/Engine/Source/Runtime/MovieSceneTracks/Private/Evaluation/MovieSceneAudioTemplate.cpp
@@ -77,29 +77,53 @@ struct FDestroyAudioPreAnimatedToken : IMovieScenePreAnimatedToken
 
 struct FCachedAudioTrackData : IPersistentEvaluationData
 {
-	TArray<TMap<FObjectKey, TWeakObjectPtr<UAudioComponent>>> AudioComponentsByRow;
+	TMap<FObjectKey, TMap<FObjectKey, TWeakObjectPtr<UAudioComponent>>> AudioComponentsByActorKey;
+	
+	FCachedAudioTrackData()
+	{
+		// Create the container for master tracks, which do not have an actor to attach to
+		AudioComponentsByActorKey.Add(FObjectKey(), TMap<FObjectKey, TWeakObjectPtr<UAudioComponent>>());
+	}
 
-	UAudioComponent* GetAudioComponentForRow(int32 RowIndex, FObjectKey Key)
+	UAudioComponent* GetAudioComponent(FObjectKey ActorKey, FObjectKey SectionKey)
 	{
-		if (AudioComponentsByRow.Num() > RowIndex)
+		if (TMap<FObjectKey, TWeakObjectPtr<UAudioComponent>>* Map = AudioComponentsByActorKey.Find(ActorKey))
 		{
-			return AudioComponentsByRow[RowIndex].FindRef(Key).Get();
+			// First, check for an exact match for this section
+			TWeakObjectPtr<UAudioComponent> ExistingComponentPtr = Map->FindRef(SectionKey);
+			if (ExistingComponentPtr.IsValid())
+			{
+				return ExistingComponentPtr.Get();
+			}
+
+			// If no exact match, check for any AudioComponent that isn't busy
+			for (TPair<FObjectKey, TWeakObjectPtr<UAudioComponent >> Pair : *Map)
+			{
+				UAudioComponent* ExistingComponent = Map->FindRef(Pair.Key).Get();
+				if (ExistingComponent && !ExistingComponent->IsPlaying())
+				{
+					// Replace this entry with the new SectionKey to claim it
+					Map->Remove(Pair.Key);
+					Map->Add(SectionKey, ExistingComponent);
+					return ExistingComponent;
+				}
+			}
 		}
 
 		return nullptr;
 	}
 
 	/** Only to be called on the game thread */
-	UAudioComponent* AddAudioComponentForRow(int32 RowIndex, AActor& PrincipalActor, IMovieScenePlayer& Player)
+	UAudioComponent* AddAudioComponentForRow(int32 RowIndex, FObjectKey SectionKey, AActor& PrincipalActor, IMovieScenePlayer& Player)
 	{
-		const int32 NumExtra = RowIndex + 1 - AudioComponentsByRow.Num();
-		if (NumExtra > 0)
+		FObjectKey ActorKey(&PrincipalActor);
+		
+		if (!AudioComponentsByActorKey.Contains(ActorKey))
 		{
-			AudioComponentsByRow.AddDefaulted(NumExtra);
+			AudioComponentsByActorKey.Add(ActorKey, TMap<FObjectKey, TWeakObjectPtr<UAudioComponent>>());
 		}
 
-		FObjectKey ActorKey(&PrincipalActor);
-		UAudioComponent* ExistingComponent = AudioComponentsByRow[RowIndex].FindRef(ActorKey).Get();
+		UAudioComponent* ExistingComponent = GetAudioComponent(ActorKey, SectionKey);
 		if (!ExistingComponent)
 		{
 			USoundCue* TempPlaybackAudioCue = NewObject<USoundCue>();
@@ -121,7 +145,7 @@ struct FCachedAudioTrackData : IPersistentEvaluationData
 
 			Player.SavePreAnimatedState(*ExistingComponent, FMovieSceneAnimTypeID::Unique(), FDestroyAudioPreAnimatedToken::FProducer());
 
-			AudioComponentsByRow[RowIndex].Add(ActorKey, ExistingComponent);
+			AudioComponentsByActorKey[ActorKey].Add(SectionKey, ExistingComponent);
 
 			ExistingComponent->SetFlags(RF_Transient);
 			ExistingComponent->AttachToComponent(PrincipalActor.GetRootComponent(), FAttachmentTransformRules::KeepRelativeTransform);
@@ -131,15 +155,9 @@ struct FCachedAudioTrackData : IPersistentEvaluationData
 	}
 
 	/** Only to be called on the game thread */
-	UAudioComponent* AddMasterAudioComponentForRow(int32 RowIndex, UWorld* World, IMovieScenePlayer& Player)
+	UAudioComponent* AddMasterAudioComponentForRow(int32 RowIndex, FObjectKey SectionKey, UWorld* World, IMovieScenePlayer& Player)
 	{
-		const int32 NumExtra = RowIndex + 1 - AudioComponentsByRow.Num();
-		if (NumExtra > 0)
-		{
-			AudioComponentsByRow.AddDefaulted(NumExtra);
-		}
-
-		UAudioComponent* ExistingComponent = AudioComponentsByRow[RowIndex].FindRef(FObjectKey()).Get();
+		UAudioComponent* ExistingComponent = GetAudioComponent(FObjectKey(), SectionKey);
 		if (!ExistingComponent)
 		{
 			USoundCue* TempPlaybackAudioCue = NewObject<USoundCue>();
@@ -154,7 +172,7 @@ struct FCachedAudioTrackData : IPersistentEvaluationData
 			Player.SavePreAnimatedState(*ExistingComponent, FMovieSceneAnimTypeID::Unique(), FDestroyAudioPreAnimatedToken::FProducer());
 			
 			ExistingComponent->SetFlags(RF_Transient);
-			AudioComponentsByRow[RowIndex].Add(FObjectKey(), ExistingComponent);
+			AudioComponentsByActorKey[FObjectKey()].Add(SectionKey, ExistingComponent);
 		}
 
 		return ExistingComponent;
@@ -162,9 +180,9 @@ struct FCachedAudioTrackData : IPersistentEvaluationData
 
 	void StopAllSounds()
 	{
-		for (TMap<FObjectKey, TWeakObjectPtr<UAudioComponent>>& Map : AudioComponentsByRow)
+		for (TPair<FObjectKey, TMap<FObjectKey, TWeakObjectPtr<UAudioComponent>>>& Map : AudioComponentsByActorKey)
 		{
-			for (auto& Pair : Map)
+			for (TPair<FObjectKey, TWeakObjectPtr<UAudioComponent>>& Pair : Map.Value)
 			{
 				if (UAudioComponent* AudioComponent = Pair.Value.Get())
 				{
@@ -174,16 +192,13 @@ struct FCachedAudioTrackData : IPersistentEvaluationData
 		}
 	}
 
-	void StopSoundsOnRow(int32 RowIndex)
+	void StopSoundsOnSection(FObjectKey ObjectKey)
 	{
-		if (RowIndex >= 0 && RowIndex < AudioComponentsByRow.Num())
+		for (TPair<FObjectKey, TMap<FObjectKey, TWeakObjectPtr<UAudioComponent>>>& Pair : AudioComponentsByActorKey)
 		{
-			for (auto& Pair : AudioComponentsByRow[RowIndex])
+			if (UAudioComponent* AudioComponent = Pair.Value.FindRef(ObjectKey).Get())
 			{
-				if (UAudioComponent* AudioComponent = Pair.Value.Get())
-				{
-					AudioComponent->Stop();
-				}
+				AudioComponent->Stop();
 			}
 		}
 	}
@@ -192,8 +207,8 @@ struct FCachedAudioTrackData : IPersistentEvaluationData
 
 struct FAudioSectionExecutionToken : IMovieSceneExecutionToken
 {
-	FAudioSectionExecutionToken(const FMovieSceneAudioSectionTemplateData& InAudioData)
-		: AudioData(InAudioData)
+	FAudioSectionExecutionToken(const UMovieSceneAudioSection* InAudioSection)
+		: AudioSection(InAudioSection), SectionKey(InAudioSection)
 	{}
 
 	virtual void Execute(const FMovieSceneContext& Context, const FMovieSceneEvaluationOperand& Operand, FPersistentEvaluationData& PersistentData, IMovieScenePlayer& Player)
@@ -211,32 +226,32 @@ struct FAudioSectionExecutionToken : IMovieSceneExecutionToken
 		{
 			UObject* PlaybackContext = Player.GetPlaybackContext();
 
-			UAudioComponent* AudioComponent = TrackData.GetAudioComponentForRow(AudioData.RowIndex, FObjectKey());
+			UAudioComponent* AudioComponent = TrackData.GetAudioComponent(FObjectKey(), SectionKey);
 			if (!AudioComponent)
 			{
 				// Initialize the sound
-				AudioComponent = TrackData.AddMasterAudioComponentForRow(AudioData.RowIndex, PlaybackContext ? PlaybackContext->GetWorld() : nullptr, Player);
+				AudioComponent = TrackData.AddMasterAudioComponentForRow(AudioSection->GetRowIndex(), SectionKey, PlaybackContext ? PlaybackContext->GetWorld() : nullptr, Player);
 
 				if (AudioComponent)
 				{
-					if (AudioData.OnQueueSubtitles.IsBound())
+					if (AudioSection->GetOnQueueSubtitles().IsBound())
 					{
-						AudioComponent->OnQueueSubtitles = AudioData.OnQueueSubtitles;
+						AudioComponent->OnQueueSubtitles = AudioSection->GetOnQueueSubtitles();
 					}
-					if (AudioData.OnAudioFinished.IsBound())
+					if (AudioSection->GetOnAudioFinished().IsBound())
 					{
-						AudioComponent->OnAudioFinished = AudioData.OnAudioFinished;
+						AudioComponent->OnAudioFinished = AudioSection->GetOnAudioFinished();
 					}
-					if(AudioData.OnAudioPlaybackPercent.IsBound())
+					if (AudioSection->GetOnAudioPlaybackPercent().IsBound())
 					{
-						AudioComponent->OnAudioPlaybackPercent = AudioData.OnAudioPlaybackPercent;
+						AudioComponent->OnAudioPlaybackPercent = AudioSection->GetOnAudioPlaybackPercent();
 					}
 				}
 			}
 
 			if (AudioComponent)
 			{
-				AudioData.EnsureAudioIsPlaying(*AudioComponent, PersistentData, Context, false, Player);
+				EnsureAudioIsPlaying(*AudioComponent, PersistentData, Context, false, Player);
 			}
 		}
 
@@ -251,141 +266,130 @@ struct FAudioSectionExecutionToken : IMovieSceneExecutionToken
 					continue;
 				}
 
-				UAudioComponent* AudioComponent = TrackData.GetAudioComponentForRow(AudioData.RowIndex, Actor);
+				UAudioComponent* AudioComponent = TrackData.GetAudioComponent(Actor, SectionKey);
 				if (!AudioComponent)
 				{
 					// Initialize the sound
-					AudioComponent = TrackData.AddAudioComponentForRow(AudioData.RowIndex, *Actor, Player);
+					AudioComponent = TrackData.AddAudioComponentForRow(AudioSection->GetRowIndex(), SectionKey, *Actor, Player);
 
 					if (AudioComponent)
 					{
-						if (AudioData.OnQueueSubtitles.IsBound())
+						if (AudioSection->GetOnQueueSubtitles().IsBound())
 						{
-							AudioComponent->OnQueueSubtitles = AudioData.OnQueueSubtitles;
+							AudioComponent->OnQueueSubtitles = AudioSection->GetOnQueueSubtitles();
 						}
-						if (AudioData.OnAudioFinished.IsBound())
+						if (AudioSection->GetOnAudioFinished().IsBound())
 						{
-							AudioComponent->OnAudioFinished = AudioData.OnAudioFinished;
+							AudioComponent->OnAudioFinished = AudioSection->GetOnAudioFinished();
 						}
-						if (AudioData.OnAudioPlaybackPercent.IsBound())
+						if (AudioSection->GetOnAudioPlaybackPercent().IsBound())
 						{
-							AudioComponent->OnAudioPlaybackPercent = AudioData.OnAudioPlaybackPercent;
+							AudioComponent->OnAudioPlaybackPercent = AudioSection->GetOnAudioPlaybackPercent();
 						}
 					}
 				}
-				
+
 				if (AudioComponent)
 				{
-					AudioData.EnsureAudioIsPlaying(*AudioComponent, PersistentData, Context, true, Player);
+					EnsureAudioIsPlaying(*AudioComponent, PersistentData, Context, true, Player);
 				}
 			}
 		}
 	}
 
-	FMovieSceneAudioSectionTemplateData AudioData;
-};
-
-FMovieSceneAudioSectionTemplateData::FMovieSceneAudioSectionTemplateData(const UMovieSceneAudioSection& Section)
-	: Sound(Section.GetSound())
-	, AudioStartOffset(Section.GetStartOffset())
-	, AudioPitchMultiplierCurve(Section.GetPitchMultiplierChannel())
-	, AudioVolumeCurve(Section.GetSoundVolumeChannel())
-	, RowIndex(Section.GetRowIndex())
-	, bOverrideAttenuation(Section.GetOverrideAttenuation())
-	, AttenuationSettings(Section.GetAttenuationSettings())
-	, OnQueueSubtitles(Section.GetOnQueueSubtitles())
-	, OnAudioFinished(Section.GetOnAudioFinished())
-	, OnAudioPlaybackPercent(Section.GetOnAudioPlaybackPercent())
-{
-	SectionStartTimeSeconds = (Section.HasStartFrame() ? Section.GetInclusiveStartFrame() : 0) / Section.GetTypedOuter<UMovieScene>()->GetTickResolution();
-}
-
-void FMovieSceneAudioSectionTemplateData::EnsureAudioIsPlaying(UAudioComponent& AudioComponent, FPersistentEvaluationData& PersistentData, const FMovieSceneContext& Context, bool bAllowSpatialization, IMovieScenePlayer& Player) const
-{
-	Player.SavePreAnimatedState(AudioComponent, FStopAudioPreAnimatedToken::GetAnimTypeID(), FStopAudioPreAnimatedToken::FProducer());
-
-	bool bPlaySound = !AudioComponent.IsPlaying() || AudioComponent.Sound != Sound;
-
-	float AudioVolume = 1.f;
-	AudioVolumeCurve.Evaluate(Context.GetTime(), AudioVolume);
-	if (AudioComponent.VolumeMultiplier != AudioVolume)
+	void EnsureAudioIsPlaying(UAudioComponent& AudioComponent, FPersistentEvaluationData& PersistentData, const FMovieSceneContext& Context, bool bAllowSpatialization, IMovieScenePlayer& Player) const
 	{
-		AudioComponent.SetVolumeMultiplier(AudioVolume);
-	}
+		Player.SavePreAnimatedState(AudioComponent, FStopAudioPreAnimatedToken::GetAnimTypeID(), FStopAudioPreAnimatedToken::FProducer());
 
-	float PitchMultiplier = 1.f;
-	AudioPitchMultiplierCurve.Evaluate(Context.GetTime(), PitchMultiplier);
-	if (AudioComponent.PitchMultiplier != PitchMultiplier)
-	{
-		AudioComponent.SetPitchMultiplier(PitchMultiplier);
-	}
+		bool bPlaySound = !AudioComponent.IsPlaying() || AudioComponent.Sound != AudioSection->GetSound();
 
-	if (bPlaySound)
-	{
-		AudioComponent.bAllowSpatialization = bAllowSpatialization;
-	
-		if (bOverrideAttenuation)
+		float AudioVolume = 1.f;
+		AudioSection->GetSoundVolumeChannel().Evaluate(Context.GetTime(), AudioVolume);
+		AudioVolume = AudioVolume * AudioSection->EvaluateEasing(Context.GetTime());
+		if (AudioComponent.VolumeMultiplier != AudioVolume)
 		{
-			AudioComponent.AttenuationSettings = AttenuationSettings;
-		}
-		
-		AudioComponent.Stop();
-		AudioComponent.SetSound(Sound);
-#if WITH_EDITOR
-		UObject* PlaybackContext = Player.GetPlaybackContext();
-		UWorld* World = PlaybackContext ? PlaybackContext->GetWorld() : nullptr;
-		if (GIsEditor && World != nullptr && !World->IsPlayInEditor())
-		{
-			AudioComponent.bIsUISound = true;
-			AudioComponent.bIsPreviewSound = true;
-		}
-		else
-#endif // WITH_EDITOR
-		{
-			AudioComponent.bIsUISound = false;
+			AudioComponent.SetVolumeMultiplier(AudioVolume);
 		}
 
-		const float AudioTime = (Context.GetTime() / Context.GetFrameRate()) - SectionStartTimeSeconds + (float)Context.GetFrameRate().AsSeconds(AudioStartOffset);
-		if (AudioTime >= 0.f && AudioComponent.Sound && AudioTime < AudioComponent.Sound->GetDuration())
+		float PitchMultiplier = 1.f;
+		AudioSection->GetPitchMultiplierChannel().Evaluate(Context.GetTime(), PitchMultiplier);
+		if (AudioComponent.PitchMultiplier != PitchMultiplier)
 		{
-			AudioComponent.Play(AudioTime);
+			AudioComponent.SetPitchMultiplier(PitchMultiplier);
 		}
 
-		if (Context.GetStatus() == EMovieScenePlayerStatus::Scrubbing)
+		if (bPlaySound)
 		{
-			// Fade out the sound at the same volume in order to simply
-			// set a short duration on the sound, far from ideal soln
-			AudioComponent.FadeOut(AudioTrackConstants::ScrubDuration, 1.f);
+			AudioComponent.bAllowSpatialization = bAllowSpatialization;
+
+			if (AudioSection->GetOverrideAttenuation())
+			{
+				AudioComponent.AttenuationSettings = AudioSection->GetAttenuationSettings();
+			}
+
+			AudioComponent.Stop();
+			AudioComponent.SetSound(AudioSection->GetSound());
+#if WITH_EDITOR
+			UObject* PlaybackContext = Player.GetPlaybackContext();
+			UWorld* World = PlaybackContext ? PlaybackContext->GetWorld() : nullptr;
+			if (GIsEditor && World != nullptr && !World->IsPlayInEditor())
+			{
+				AudioComponent.bIsUISound = true;
+				AudioComponent.bIsPreviewSound = true;
+			}
+			else
+#endif // WITH_EDITOR
+			{
+				AudioComponent.bIsUISound = false;
+			}
+
+			float SectionStartTimeSeconds = (AudioSection->HasStartFrame() ? AudioSection->GetInclusiveStartFrame() : 0) / AudioSection->GetTypedOuter<UMovieScene>()->GetTickResolution();
+
+			const float AudioTime = (Context.GetTime() / Context.GetFrameRate()) - SectionStartTimeSeconds + (float)Context.GetFrameRate().AsSeconds(AudioSection->GetStartOffset());
+			if (AudioTime >= 0.f && AudioComponent.Sound && AudioTime < AudioComponent.Sound->GetDuration())
+			{
+				AudioComponent.Play(AudioTime);
+			}
+
+			if (Context.GetStatus() == EMovieScenePlayerStatus::Scrubbing)
+			{
+				// Fade out the sound at the same volume in order to simply
+				// set a short duration on the sound, far from ideal soln
+				AudioComponent.FadeOut(AudioTrackConstants::ScrubDuration, 1.f);
+			}
 		}
-	}
 
-	if (bAllowSpatialization)
-	{
-		if (FAudioDevice* AudioDevice = AudioComponent.GetAudioDevice())
+		if (bAllowSpatialization)
 		{
-			DECLARE_CYCLE_STAT(TEXT("FAudioThreadTask.MovieSceneUpdateAudioTransform"), STAT_MovieSceneUpdateAudioTransform, STATGROUP_TaskGraphTasks);
-
-			const FTransform ActorTransform = AudioComponent.GetOwner()->GetTransform();
-			const uint64 ActorComponentID = AudioComponent.GetAudioComponentID();
-			FAudioThread::RunCommandOnAudioThread([AudioDevice, ActorComponentID, ActorTransform]()
+			if (FAudioDevice* AudioDevice = AudioComponent.GetAudioDevice())
 			{
-				if (FActiveSound* ActiveSound = AudioDevice->FindActiveSound(ActorComponentID))
+				DECLARE_CYCLE_STAT(TEXT("FAudioThreadTask.MovieSceneUpdateAudioTransform"), STAT_MovieSceneUpdateAudioTransform, STATGROUP_TaskGraphTasks);
+
+				const FTransform ActorTransform = AudioComponent.GetOwner()->GetTransform();
+				const uint64 ActorComponentID = AudioComponent.GetAudioComponentID();
+				FAudioThread::RunCommandOnAudioThread([AudioDevice, ActorComponentID, ActorTransform]()
 				{
-					ActiveSound->bLocationDefined = true;
-					ActiveSound->Transform = ActorTransform;
-				}
-			}, GET_STATID(STAT_MovieSceneUpdateAudioTransform));
+					if (FActiveSound* ActiveSound = AudioDevice->FindActiveSound(ActorComponentID))
+					{
+						ActiveSound->bLocationDefined = true;
+						ActiveSound->Transform = ActorTransform;
+					}
+				}, GET_STATID(STAT_MovieSceneUpdateAudioTransform));
+			}
 		}
 	}
-}
+
+	const UMovieSceneAudioSection* AudioSection;
+	FObjectKey SectionKey;
+};
 
 FMovieSceneAudioSectionTemplate::FMovieSceneAudioSectionTemplate()
-	: AudioData()
+	: AudioSection()
 {
 }
 
 FMovieSceneAudioSectionTemplate::FMovieSceneAudioSectionTemplate(const UMovieSceneAudioSection& Section)
-	: AudioData(Section)
+	: AudioSection(&Section)
 {
 }
 
@@ -396,7 +400,7 @@ void FMovieSceneAudioSectionTemplate::Evaluate(const FMovieSceneEvaluationOperan
 
 	if (GEngine && GEngine->UseSound() && Context.GetStatus() != EMovieScenePlayerStatus::Jumping)
 	{
-		ExecutionTokens.Add(FAudioSectionExecutionToken(AudioData));
+		ExecutionTokens.Add(FAudioSectionExecutionToken(AudioSection));
 	}
 }
 
@@ -408,6 +412,6 @@ void FMovieSceneAudioSectionTemplate::TearDown(FPersistentEvaluationData& Persis
 	{
 		FCachedAudioTrackData& TrackData = PersistentData.GetOrAddTrackData<FCachedAudioTrackData>();
 
-		TrackData.StopSoundsOnRow(AudioData.RowIndex);
+		TrackData.StopSoundsOnSection(AudioSection);
 	}
 }
diff --git a/Engine/Source/Runtime/MovieSceneTracks/Private/Evaluation/MovieSceneAudioTemplate.h b/Engine/Source/Runtime/MovieSceneTracks/Private/Evaluation/MovieSceneAudioTemplate.h
--- a/Engine/Source/Runtime/MovieSceneTracks/Private/Evaluation/MovieSceneAudioTemplate.h
+++ b/Engine/Source/Runtime/MovieSceneTracks/Private/Evaluation/MovieSceneAudioTemplate.h
@@ -14,61 +14,6 @@ class UAudioComponent;
 class UMovieSceneAudioSection;
 class USoundBase;
 
-USTRUCT()
-struct FMovieSceneAudioSectionTemplateData
-{
-	GENERATED_BODY()
-
-	FMovieSceneAudioSectionTemplateData() : Sound(nullptr), AudioStartOffset(0), SectionStartTimeSeconds(0.0f), RowIndex(0), bOverrideAttenuation(false), AttenuationSettings(nullptr) {}
-	FMovieSceneAudioSectionTemplateData(const UMovieSceneAudioSection& Section);
-
-	/** Ensure that the sound is playing for the specified audio component and data */
-	void EnsureAudioIsPlaying(UAudioComponent& AudioComponent, FPersistentEvaluationData& PersistentData, const FMovieSceneContext& Context, bool bAllowSpatialization, IMovieScenePlayer& Player) const;
-
-	/** The sound cue or wave that this template plays. Not to be dereferenced on a background thread */
-	UPROPERTY()
-	USoundBase* Sound;
-
-	/** The offset into the beginning of the audio clip */
-	UPROPERTY()
-	FFrameNumber AudioStartOffset;
-
-	/** The frame number at which the audio starts playing */
-	UPROPERTY()
-	double SectionStartTimeSeconds;
-
-	/** The amount which this audio is time dilated by */
-	UPROPERTY()
-	FMovieSceneFloatChannel AudioPitchMultiplierCurve;
-
-	/** The volume the sound will be played with. */
-	UPROPERTY()
-	FMovieSceneFloatChannel AudioVolumeCurve;
-
-	/** The row index of the section */
-	UPROPERTY()
-	int32 RowIndex;
-
-	/** Should the attenuation settings on this section be used. */
-	UPROPERTY()
-	bool bOverrideAttenuation;
-
-	/** The attenuation settings */
-	UPROPERTY()
-	class USoundAttenuation* AttenuationSettings;
-
-	/** Called when subtitles are sent to the SubtitleManager.  Set this delegate if you want to hijack the subtitles for other purposes */
-	UPROPERTY()
-	FOnQueueSubtitles OnQueueSubtitles;
-
-	/** called when we finish playing audio, either because it played to completion or because a Stop() call turned it off early */
-	UPROPERTY()
-	FOnAudioFinished OnAudioFinished;
-
-	UPROPERTY()
-	FOnAudioPlaybackPercent OnAudioPlaybackPercent;
-};
-
 USTRUCT()
 struct FMovieSceneAudioSectionTemplate : public FMovieSceneEvalTemplate
 {
@@ -78,7 +23,7 @@ struct FMovieSceneAudioSectionTemplate : public FMovieSceneEvalTemplate
 	FMovieSceneAudioSectionTemplate(const UMovieSceneAudioSection& Section);
 
 	UPROPERTY()
-	FMovieSceneAudioSectionTemplateData AudioData;
+	const UMovieSceneAudioSection* AudioSection;
 
 private:
 
diff --git a/Engine/Source/Runtime/MovieSceneTracks/Private/Sections/MovieSceneAudioSection.cpp b/Engine/Source/Runtime/MovieSceneTracks/Private/Sections/MovieSceneAudioSection.cpp
--- a/Engine/Source/Runtime/MovieSceneTracks/Private/Sections/MovieSceneAudioSection.cpp
+++ b/Engine/Source/Runtime/MovieSceneTracks/Private/Sections/MovieSceneAudioSection.cpp
@@ -38,6 +38,7 @@ UMovieSceneAudioSection::UMovieSceneAudioSection( const FObjectInitializer& Obje
 	AudioVolume_DEPRECATED = AudioDeprecatedMagicNumber;
 	bSuppressSubtitles = false;
 	bOverrideAttenuation = false;
+	BlendType = EMovieSceneBlendType::Absolute;
 
 	EvalOptions.EnableAndSetCompletionMode
 		(GetLinkerCustomVersion(FSequencerObjectVersion::GUID) < FSequencerObjectVersion::WhenFinishedDefaultsToProjectDefault ? 
diff --git a/Engine/Source/Runtime/MovieSceneTracks/Private/Tracks/MovieSceneAudioTrack.cpp b/Engine/Source/Runtime/MovieSceneTracks/Private/Tracks/MovieSceneAudioTrack.cpp
--- a/Engine/Source/Runtime/MovieSceneTracks/Private/Tracks/MovieSceneAudioTrack.cpp
+++ b/Engine/Source/Runtime/MovieSceneTracks/Private/Tracks/MovieSceneAudioTrack.cpp
@@ -18,6 +18,7 @@
 UMovieSceneAudioTrack::UMovieSceneAudioTrack( const FObjectInitializer& ObjectInitializer )
 	: Super( ObjectInitializer )
 {
+	SupportedBlendTypes.Add(EMovieSceneBlendType::Absolute);
 #if WITH_EDITORONLY_DATA
 	TrackTint = FColor(93, 95, 136);
 	RowHeight = 50;
@@ -119,9 +120,6 @@ FMovieSceneTrackRowSegmentBlenderPtr UMovieSceneAudioTrack::GetRowSegmentBlender
 		{
 			// Run the default high pass filter for overlap priority
 			MovieSceneSegmentCompiler::FilterOutUnderlappingSections(BlendData);
-
-			// Weed out based on array index (legacy behaviour)
-			MovieSceneSegmentCompiler::BlendSegmentLegacySectionOrder(BlendData);
 		}
 	};
 	return FBlender();
