Sequencer: Made generic object animation use its own pre-animated state token rather than relying on the same templates as the actor reference track

  - The actor reference track does not deem nullptr values to be valid, so it does not set the property correctly.

diff --git a/Engine/Source/Runtime/MovieSceneTracks/Private/Evaluation/MovieSceneObjectPropertyTemplate.cpp b/Engine/Source/Runtime/MovieSceneTracks/Private/Evaluation/MovieSceneObjectPropertyTemplate.cpp
--- a/Engine/Source/Runtime/MovieSceneTracks/Private/Evaluation/MovieSceneObjectPropertyTemplate.cpp
+++ b/Engine/Source/Runtime/MovieSceneTracks/Private/Evaluation/MovieSceneObjectPropertyTemplate.cpp
@@ -3,6 +3,35 @@
 #include "Evaluation/MovieSceneObjectPropertyTemplate.h"
 #include "Tracks/MovieSceneObjectPropertyTrack.h"
 #include "Sections/MovieSceneObjectPropertySection.h"
+#include "UObject/StrongObjectPtr.h"
+
+
+struct FMovieSceneObjectPropertyValue : IMovieScenePreAnimatedToken
+{
+	TStrongObjectPtr<UObject> Value;
+	FTrackInstancePropertyBindings PropertyBindings;
+
+	FMovieSceneObjectPropertyValue(UObject* AnimatedObject, FTrackInstancePropertyBindings* InPropertyBindings)
+		: Value(InPropertyBindings->GetCurrentValue<UObject*>(*AnimatedObject))
+		, PropertyBindings(*InPropertyBindings)
+	{}
+
+	virtual void RestoreState(UObject& Object, IMovieScenePlayer& Player)
+	{
+		PropertyBindings.CallFunction<UObject*>(Object, Value.Get());
+	}
+};
+
+struct FMovieSceneObjectPropertyValueProducer : IMovieScenePreAnimatedTokenProducer
+{
+	FTrackInstancePropertyBindings* PropertyBindings;
+	FMovieSceneObjectPropertyValueProducer(FTrackInstancePropertyBindings* InPropertyBindings) : PropertyBindings(InPropertyBindings) {}
+
+	virtual IMovieScenePreAnimatedTokenPtr CacheExistingState(UObject& Object) const
+	{
+		return FMovieSceneObjectPropertyValue(&Object, PropertyBindings);
+	}
+};
 
 
 struct FObjectPropertyExecToken : IMovieSceneExecutionToken
@@ -33,7 +62,7 @@ struct FObjectPropertyExecToken : IMovieSceneExecutionToken
 				continue;
 			}
 
-			Player.SavePreAnimatedState(*ObjectPtr, PropertyTrackData.PropertyID, FTokenProducer<UObject*>(*PropertyBindings));
+			Player.SavePreAnimatedState(*ObjectPtr, PropertyTrackData.PropertyID, FMovieSceneObjectPropertyValueProducer(PropertyBindings));
 
 			UObject* ExistingValue = PropertyBindings->GetCurrentValue<UObject*>(*ObjectPtr);
 			if (ExistingValue != NewObjectValue)
