Zero out Defaults when changing section to Absolute/Additive to allow for initial section be Additive/Relative.

diff --git a/Engine/Source/Runtime/MovieSceneTracks/Private/Sections/MovieScene3DTransformSection.cpp b/Engine/Source/Runtime/MovieSceneTracks/Private/Sections/MovieScene3DTransformSection.cpp
--- a/Engine/Source/Runtime/MovieSceneTracks/Private/Sections/MovieScene3DTransformSection.cpp
+++ b/Engine/Source/Runtime/MovieSceneTracks/Private/Sections/MovieScene3DTransformSection.cpp
@@ -699,22 +699,33 @@ float UMovieScene3DTransformSection::GetTotalWeightValue(FFrameTime InTime) cons
 	return Weight;
 }
 
-void UMovieScene3DTransformSection::SetBlendType(EMovieSceneBlendType InBlendType) 
+void UMovieScene3DTransformSection::SetBlendType(EMovieSceneBlendType InBlendType)
 {
 	if (GetSupportedBlendTypes().Contains(InBlendType))
 	{
 		BlendType = InBlendType;
-		//Set the Scale Default based upon the type that was Set
-		float DefaultVal = InBlendType == EMovieSceneBlendType::Absolute ? 1.0f : 0.0f;
+		// Set Defaults based upon Type
 		TArrayView<FMovieSceneFloatChannel*> FloatChannels = ChannelProxy->GetChannels<FMovieSceneFloatChannel>();
-		// Set default scale and weight
-		for (int32 Index = 6; Index < 9; ++Index)
+
+		if (InBlendType == EMovieSceneBlendType::Absolute)
 		{
-			FloatChannels[Index]->SetDefault(DefaultVal);
+			//If Absolute just make sure scale default is 1.0f
+			for (int32 Index = 6; Index < 9; ++Index)
+			{
+				FloatChannels[Index]->SetDefault(1.0f);
+			}
+		}
+		else
+		{
+			//If Additive or Relative, zero out Pos/Rot/ and Scale.
+			for (int32 Index = 0; Index < 9; ++Index)
+			{
+				FloatChannels[Index]->SetDefault(0.0f);
+			}
 		}
-
 	}
 }
+
 FMovieSceneInterrogationKey UMovieScene3DTransformSection::GetInterrogationKey()
 {
 	static FMovieSceneAnimTypeID TypeID = FMovieSceneAnimTypeID::Unique();
