Sequencer: Camera cut frames now inform the renderer of simulated previous frame transforms as best they can

Several parts to this but ultimately what we need to achieve is that anything that is transformed by sequencer **and anything attached to such objects** get simulated transforms as close to accurate as possible.

Note: This is an opt-in feature that is enabled for the time-being through a cvar: "r.MotionVectorSimulation 1"

diff --git a/Engine/Source/Editor/MovieSceneTools/Private/AutomatedLevelSequenceCapture.cpp b/Engine/Source/Editor/MovieSceneTools/Private/AutomatedLevelSequenceCapture.cpp
--- a/Engine/Source/Editor/MovieSceneTools/Private/AutomatedLevelSequenceCapture.cpp
+++ b/Engine/Source/Editor/MovieSceneTools/Private/AutomatedLevelSequenceCapture.cpp
@@ -20,6 +20,8 @@
 #include "MovieSceneTimeHelpers.h"
 #include "MovieSceneToolHelpers.h"
 #include "Protocols/AudioCaptureProtocol.h"
+#include "Evaluation/IMovieSceneMotionVectorSimulation.h"
+#include "Rendering/MotionVectorSimulation.h"
 
 const FName UAutomatedLevelSequenceCapture::AutomatedLevelSequenceCaptureUIName = FName(TEXT("AutomatedLevelSequenceCaptureUIInstance"));
 
@@ -548,6 +550,11 @@ void UAutomatedLevelSequenceCapture::SetupFrameRange()
 				 	PlaybackStartFrame -= RemainingWarmUpFrames;
 				}
 
+				if (Actor->SequencePlayer->MotionVectorSimulation.IsValid())
+				{
+					Actor->SequencePlayer->MotionVectorSimulation->PreserveSimulatedMotion(true);
+				}
+
 				// Override the movie scene's playback range
 				Actor->SequencePlayer->SetFrameRate(Settings.GetFrameRate());
 				Actor->SequencePlayer->SetFrameRange(PlaybackStartFrame.Value, (PlaybackEndFrame - PlaybackStartFrame).Value);
@@ -779,11 +786,22 @@ void UAutomatedLevelSequenceCapture::SequenceUpdated(const UMovieSceneSequencePl
 				
 				CaptureState = ELevelSequenceCaptureState::Paused;
 
+				if (Actor->SequencePlayer->MotionVectorSimulation.IsValid())
+				{
+					Actor->SequencePlayer->MotionVectorSimulation->PreserveSimulatedMotion(true);
+				}
+
 				Actor->GetWorld()->GetTimerManager().SetTimer(DelayTimer, FTimerDelegate::CreateUObject(this, &UAutomatedLevelSequenceCapture::PauseFinished), DelayBeforeShotWarmUp + DelayEveryFrame, false);
 				Actor->SequencePlayer->Pause();
 			}
 			else if (CaptureState == ELevelSequenceCaptureState::FinishedWarmUp)
 			{
+				// If we were preserving simulated motion, now's the time to stop that since we've captured the frame that was being simulated
+				if (Actor->SequencePlayer->MotionVectorSimulation.IsValid())
+				{
+					Actor->SequencePlayer->MotionVectorSimulation->PreserveSimulatedMotion(false);
+				}
+
 				// These are called each frame to allow the state machine inside the protocol to transition back to capturing
 				// after paused if needed. This is needed for things like the avi writer who spin up an avi writer per shot (if needed)
 				// so that we can capture the movies into individual avi files per shot due to the format text.
diff --git a/Engine/Source/Runtime/MovieScene/Private/Evaluation/MovieSceneEvaluationTemplateInstance.cpp b/Engine/Source/Runtime/MovieScene/Private/Evaluation/MovieSceneEvaluationTemplateInstance.cpp
--- a/Engine/Source/Runtime/MovieScene/Private/Evaluation/MovieSceneEvaluationTemplateInstance.cpp
+++ b/Engine/Source/Runtime/MovieScene/Private/Evaluation/MovieSceneEvaluationTemplateInstance.cpp
@@ -8,6 +8,8 @@
 #include "Compilation/MovieSceneCompiler.h"
 #include "Sections/MovieSceneSubSection.h"
 #include "Compilation/MovieSceneEvaluationTemplateGenerator.h"
+#include "Evaluation/MovieSceneMotionVectorSimulation.h"
+#include "Rendering/MotionVectorSimulation.h"
 
 #include "IMovieSceneModule.h"
 #include "Algo/Sort.h"
@@ -193,6 +195,11 @@ void FMovieSceneRootEvaluationTemplateInstance::Initialize(UMovieSceneSequence&
 	RootTemplate = &TemplateStore->AccessTemplate(InRootSequence);
 
 	RootID = MovieSceneSequenceID::Root;
+
+	if (FMotionVectorSimulation::Get().IsEnabled() && !Player.MotionVectorSimulation.IsValid())
+	{
+		Player.MotionVectorSimulation = FMovieSceneMotionVectorSimulation();
+	}
 }
 
 void FMovieSceneRootEvaluationTemplateInstance::Finish(IMovieScenePlayer& Player)
@@ -203,6 +210,7 @@ void FMovieSceneRootEvaluationTemplateInstance::Finish(IMovieScenePlayer& Player
 	CallSetupTearDown(Player);
 
 	ResetDirectorInstances();
+	Player.MotionVectorSimulation.Reset();
 }
 
 void FMovieSceneRootEvaluationTemplateInstance::Evaluate(FMovieSceneContext Context, IMovieScenePlayer& Player, FMovieSceneSequenceID InOverrideRootID)
@@ -257,6 +265,11 @@ void FMovieSceneRootEvaluationTemplateInstance::Evaluate(FMovieSceneContext Cont
 
 	// Process execution tokens
 	ExecutionTokens.Apply(Context, Player);
+
+	if (Player.MotionVectorSimulation.IsValid())
+	{
+		Player.MotionVectorSimulation->Apply(Player);
+	}
 }
 
 FMovieSceneEvaluationPtrCache FMovieSceneRootEvaluationTemplateInstance::ConstructEvaluationPtrCacheForFrame(UMovieSceneSequence* OverrideRootSequence)
diff --git a/Engine/Source/Runtime/MovieScene/Private/Evaluation/MovieSceneMotionVectorSimulation.cpp b/Engine/Source/Runtime/MovieScene/Private/Evaluation/MovieSceneMotionVectorSimulation.cpp
new file mode 100644
--- /dev/null
+++ b/Engine/Source/Runtime/MovieScene/Private/Evaluation/MovieSceneMotionVectorSimulation.cpp
@@ -0,0 +1,186 @@
+// Copyright 1998-2017 Epic Games, Inc. All Rights Reserved.
+
+#include "Evaluation/MovieSceneMotionVectorSimulation.h"
+#include "IMovieScenePlayer.h"
+#include "MovieSceneCommonHelpers.h"
+#include "Components/SceneComponent.h"
+#include "Rendering/MotionVectorSimulation.h"
+#include "Evaluation/MovieScenePlayback.h"
+
+struct FSharedMotionSimulationData : IPersistentEvaluationData
+{
+	FSharedMotionSimulationData()
+		: bIsEnabled(false)
+	{}
+
+	static FSharedPersistentDataKey GetSharedDataId()
+	{
+		static FMovieSceneSharedDataId ID = FMovieSceneSharedDataId::Allocate();
+		return FSharedPersistentDataKey(ID, FMovieSceneEvaluationOperand());
+	}
+
+	static void SetEnabled(FPersistentEvaluationData& PersistentData, bool bInIsEnabled)
+	{
+		FSharedPersistentDataKey DataKey = GetSharedDataId();
+		if (bInIsEnabled)
+		{
+			PersistentData.GetOrAdd<FSharedMotionSimulationData>(DataKey).bIsEnabled = bInIsEnabled;
+		}
+		else
+		{
+			PersistentData.Reset(DataKey);
+		}
+	}
+
+	static bool IsEnabled(const FPersistentEvaluationData& PersistentData)
+	{
+		const FSharedMotionSimulationData* Data = PersistentData.Find<FSharedMotionSimulationData>(GetSharedDataId());
+		return Data && Data->bIsEnabled;
+	}
+
+	bool bIsEnabled;
+};
+
+void IMovieSceneMotionVectorSimulation::EnableThisFrame(FPersistentEvaluationData& PersistentData)
+{
+	FSharedMotionSimulationData::SetEnabled(PersistentData, true);
+}
+
+bool IMovieSceneMotionVectorSimulation::IsEnabled(const FPersistentEvaluationData& PersistentData, const FMovieSceneContext& Context)
+{
+	return !Context.IsSilent() && FSharedMotionSimulationData::IsEnabled(PersistentData) && FSharedMotionSimulationData::IsEnabled(PersistentData);
+}
+
+FFrameTime IMovieSceneMotionVectorSimulation::GetSimulationTime(const FMovieSceneContext& Context)
+{
+	FFrameTime DeltaTime     = FMath::Max(Context.GetDelta(), 0.00833333334 * Context.GetFrameRate());
+	FFrameTime SimulatedTime = Context.GetOffsetTime(DeltaTime);
+
+	return SimulatedTime;
+}
+
+void FMovieSceneMotionVectorSimulation::PreserveSimulatedMotion(bool bShouldPreserveTransforms)
+{
+	bPreserveTransforms = bShouldPreserveTransforms;
+}
+
+void FMovieSceneMotionVectorSimulation::Add(USceneComponent* Component, const FTransform& SimulatedTransform, FName SocketName)
+{
+	TransformData.Add(Component, FSimulatedTransform(SimulatedTransform, SocketName));
+}
+
+void FMovieSceneMotionVectorSimulation::Apply(IMovieScenePlayer& Player)
+{
+	TSet<USceneComponent*> RootComponents;
+
+	for (auto It = TransformData.CreateIterator(); It; ++It)
+	{
+		// If this is a socket transform, we want to add the component as a whole, to ensure that anything attached to this socket gets simulated correctly
+
+		USceneComponent* Component = Cast<USceneComponent>(It.Key().ResolveObjectPtr());
+		if (!Component || HavePreviousTransformForParent(Component))
+		{
+			continue;
+		}
+
+		RootComponents.Add(Component);
+	}
+
+	for (USceneComponent* Component : RootComponents)
+	{
+		FTransform ParentToWorld = FTransform::Identity;
+
+		USceneComponent* ParentComp = Component->GetAttachParent();
+		FName AttachSocket = Component->GetAttachSocketName();
+		if (ParentComp)
+		{
+			FTransform ParentTransform = ParentComp->GetSocketTransform(AttachSocket, RTS_World);
+			if (!Component->bAbsoluteLocation)
+			{
+				ParentToWorld.SetTranslation(ParentTransform.GetTranslation());
+			}
+			if (!Component->bAbsoluteRotation)
+			{
+				ParentToWorld.SetRotation(ParentTransform.GetRotation());
+			}
+			if (!Component->bAbsoluteScale)
+			{
+				ParentToWorld.SetScale3D(ParentTransform.GetScale3D());
+			}
+		}
+
+		ApplySimulatedTransforms(Component, GetTransform(Component) * ParentToWorld);
+	}
+
+	if (!bPreserveTransforms)
+	{
+		TransformData.Reset();
+	}
+
+	FPersistentEvaluationData PersistentDataProxy(Player);
+	FSharedMotionSimulationData::SetEnabled(PersistentDataProxy, false);
+}
+
+FTransform FMovieSceneMotionVectorSimulation::GetTransform(USceneComponent* Component)
+{
+	FObjectKey Key(Component);
+	for (auto It = TransformData.CreateConstKeyIterator(Key); It; ++It)
+	{
+		if (It.Value().SocketName == NAME_None)
+		{
+			return It.Value().Transform;
+		}
+	}
+
+	return Component->GetRelativeTransform();
+}
+
+FTransform FMovieSceneMotionVectorSimulation::GetSocketTransform(USceneComponent* Component, FName SocketName)
+{
+	FObjectKey Key(Component);
+	for (auto It = TransformData.CreateConstKeyIterator(Key); It; ++It)
+	{
+		if (It.Value().SocketName == SocketName)
+		{
+			return It.Value().Transform;
+		}
+	}
+
+	return Component->GetSocketTransform(SocketName, RTS_Component);
+}
+
+bool FMovieSceneMotionVectorSimulation::HavePreviousTransformForParent(USceneComponent* InComponent) const
+{
+	USceneComponent* Parent = InComponent->GetAttachParent();
+	return Parent && (TransformData.Contains(Parent) || HavePreviousTransformForParent(Parent));
+}
+
+void FMovieSceneMotionVectorSimulation::ApplySimulatedTransforms(USceneComponent* InComponent, const FTransform& InPreviousTransform)
+{
+	check(InComponent);
+	FMotionVectorSimulation::Get().SetPreviousTransform(InComponent, InPreviousTransform);
+
+	for (USceneComponent* Child : InComponent->GetAttachChildren())
+	{
+		FName AttachSocketName = Child->GetAttachSocketName();
+
+		FTransform SocketTransform = AttachSocketName == NAME_None ? FTransform::Identity : GetSocketTransform(InComponent, AttachSocketName);
+		FTransform ParentToWorld = SocketTransform * InPreviousTransform;
+
+		if (Child->bAbsoluteLocation)
+		{
+			ParentToWorld.SetTranslation(FVector::ZeroVector);
+		}
+		if (Child->bAbsoluteRotation)
+		{
+			ParentToWorld.SetRotation(FQuat::Identity);
+		}
+		if (Child->bAbsoluteScale)
+		{
+			ParentToWorld.SetScale3D(FVector(1.f, 1.f, 1.f));
+		}
+
+		FTransform ChildTransform = GetTransform(Child);
+		ApplySimulatedTransforms(Child, ChildTransform * ParentToWorld);
+	}
+}
\ No newline at end of file
diff --git a/Engine/Source/Runtime/MovieScene/Private/Evaluation/MovieSceneMotionVectorSimulation.h b/Engine/Source/Runtime/MovieScene/Private/Evaluation/MovieSceneMotionVectorSimulation.h
new file mode 100644
--- /dev/null
+++ b/Engine/Source/Runtime/MovieScene/Private/Evaluation/MovieSceneMotionVectorSimulation.h
@@ -0,0 +1,51 @@
+// Copyright 1998-2017 Epic Games, Inc. All Rights Reserved.
+
+#include "UObject/ObjectKey.h"
+#include "Evaluation/IMovieSceneMotionVectorSimulation.h"
+#include "Math/Transform.h"
+#include "Containers/Map.h"
+#include "UObject/NameTypes.h"
+
+class USceneComponent;
+
+
+struct FMovieSceneMotionVectorSimulation : IMovieSceneMotionVectorSimulation
+{
+	virtual void PreserveSimulatedMotion(bool bShouldPreserveTransforms) override;
+
+	virtual void Add(USceneComponent* Component, const FTransform& SimulatedTransform, FName SocketName) override;
+
+	virtual void Apply(IMovieScenePlayer& Player) override;
+
+private:
+
+	FTransform GetTransform(USceneComponent* Component);
+
+	FTransform GetSocketTransform(USceneComponent* Component, FName SocketName);
+
+	bool HavePreviousTransformForParent(USceneComponent* InComponent) const;
+
+	void ApplySimulatedTransforms(USceneComponent* InComponent, const FTransform& InPreviousTransform);
+
+private:
+
+	/** Transform data relating to a specific object or socket */
+	struct FSimulatedTransform
+	{
+		FSimulatedTransform(const FTransform& InTransform, FName InSocketName = NAME_None)
+			: Transform(InTransform), SocketName(InSocketName)
+		{}
+
+		/** The transform for the object */
+		FTransform Transform;
+
+		/** The socket name to which this relates */
+		FName SocketName;
+	};
+
+	/** Map of object -> previous transform data */
+	TMultiMap<FObjectKey, FSimulatedTransform> TransformData;
+
+	/** Whether to reset TransformData at the end of the frame or not */
+	bool bPreserveTransforms = false;
+};
\ No newline at end of file
diff --git a/Engine/Source/Runtime/MovieScene/Public/Evaluation/IMovieSceneMotionVectorSimulation.h b/Engine/Source/Runtime/MovieScene/Public/Evaluation/IMovieSceneMotionVectorSimulation.h
new file mode 100644
--- /dev/null
+++ b/Engine/Source/Runtime/MovieScene/Public/Evaluation/IMovieSceneMotionVectorSimulation.h
@@ -0,0 +1,68 @@
+// Copyright 1998-2017 Epic Games, Inc. All Rights Reserved.
+
+#pragma once
+
+#include "CoreTypes.h"
+#include "Math/Transform.h"
+#include "UObject/NameTypes.h"
+
+struct FPersistentEvaluationData;
+struct FMovieSceneContext;
+struct FFrameTime;
+
+class USceneComponent;
+class IMovieScenePlayer;
+
+
+/**
+ * Crude helper class for applying simulated transforms for all animated objects of a sequence on camera cut frames from sequencer
+ * Simulated transforms are passed onto the renderer at the end of a sequence's evaluation and will be used for motion vector computation
+ */
+struct MOVIESCENE_API IMovieSceneMotionVectorSimulation
+{
+	virtual ~IMovieSceneMotionVectorSimulation() {}
+
+	/**
+	 * Add a new simulated transform for the specified component (and optionally a socket on that component). Persists only while PreserveSimulatedMotion is enabled.
+	 *
+	 * @param Component                   The component that should have the simulated transform associated with it
+	 * @param SimulatedTransform          The simulated previous position, rotation and scale of the component last frame
+	 * @param SocketName                  (Optional) The name of a socket on 'Component' which affects the simulated transform of all components attached to the socket.
+	 */
+	virtual void Add(USceneComponent* Component, const FTransform& SimulatedTransform, FName SocketName) = 0;
+
+
+	/**
+	 * Apply all the simulated transforms from this frame using the specified player, passing them onto the renderer for this frame
+	 *
+	 * @param Player                      The movie scene player currently playing back the sequence
+	 */
+	virtual void Apply(IMovieScenePlayer& Player) = 0;
+
+
+	/**
+	 * Indicate that all the currently stored simulated transforms should be preserved or reset once evaluated this frame
+	 *
+	 * @param bShouldPreserveTransforms   When true, simulated transforms will be preserved indefinitely (useful if a sequence is paused); when false, transforms will all be reset once evaluated this frame.
+	 */
+	virtual void PreserveSimulatedMotion(bool bShouldPreserveTransforms) = 0;
+
+public:
+
+	/**
+	 * Check whether we should simulate motion vectors for the current evaluation.
+	 * @return true if we should simulate motion vectors, false otherwise.
+	 */
+	static bool IsEnabled(const FPersistentEvaluationData& PersistentData, const FMovieSceneContext& Context);
+
+	/**
+	 * Enable simulated motion vectors for the current frame
+	 */
+	static void EnableThisFrame(FPersistentEvaluationData& PersistentData);
+
+	/**
+	 * Computes a time at which to simulate motion vectors for the current frame.
+	 * @return A time one-frame in the future. Information for the previous frame should be extrapolated backwards from here to ensure correct simulations where no previous data exists (such as animations that start on frame 0).
+	 */
+	static FFrameTime GetSimulationTime(const FMovieSceneContext& Context);
+};
\ No newline at end of file
diff --git a/Engine/Source/Runtime/MovieScene/Public/Evaluation/MovieScenePlayback.h b/Engine/Source/Runtime/MovieScene/Public/Evaluation/MovieScenePlayback.h
--- a/Engine/Source/Runtime/MovieScene/Public/Evaluation/MovieScenePlayback.h
+++ b/Engine/Source/Runtime/MovieScene/Public/Evaluation/MovieScenePlayback.h
@@ -101,7 +101,16 @@ struct MOVIESCENE_API FMovieSceneEvaluationRange
 	{
 		return Direction == EPlayDirection::Forwards ? EvaluationRange.GetLowerBoundValue() : EvaluationRange.GetUpperBoundValue();
 	}
-	
+
+	/**
+	 * Get the current time offset by the specified amount in the direction of play
+	 */
+	FORCEINLINE FFrameTime GetOffsetTime(FFrameTime InOffset) const
+	{
+		FFrameTime Now = GetTime();
+		return Direction == EPlayDirection::Forwards ? Now + InOffset : Now - InOffset;
+	}
+
 	/**
 	 * Override the time that we're actually evaluating at
 	 */
diff --git a/Engine/Source/Runtime/MovieScene/Public/IMovieScenePlayer.h b/Engine/Source/Runtime/MovieScene/Public/IMovieScenePlayer.h
--- a/Engine/Source/Runtime/MovieScene/Public/IMovieScenePlayer.h
+++ b/Engine/Source/Runtime/MovieScene/Public/IMovieScenePlayer.h
@@ -11,6 +11,8 @@
 #include "MovieSceneSpawnRegister.h"
 #include "Containers/ArrayView.h"
 #include "Evaluation/MovieSceneEvaluationState.h"
+#include "Misc/InlineValue.h"
+#include "Evaluation/IMovieSceneMotionVectorSimulation.h"
 #include "Evaluation/MovieSceneEvaluationOperand.h"
 
 
@@ -331,7 +333,10 @@ public:
 
 	/** Container that stores any per-animated state tokens  */
 	FMovieScenePreAnimatedState PreAnimatedState;
-	
+
+	/** Motion Vector Simulation */
+	TInlineValue<IMovieSceneMotionVectorSimulation> MotionVectorSimulation;
+
 private:
 
 	/** Null register that asserts on use */
diff --git a/Engine/Source/Runtime/MovieSceneTracks/Private/Evaluation/MovieScene3DTransformTemplate.cpp b/Engine/Source/Runtime/MovieSceneTracks/Private/Evaluation/MovieScene3DTransformTemplate.cpp
--- a/Engine/Source/Runtime/MovieSceneTracks/Private/Evaluation/MovieScene3DTransformTemplate.cpp
+++ b/Engine/Source/Runtime/MovieSceneTracks/Private/Evaluation/MovieScene3DTransformTemplate.cpp
@@ -12,6 +12,8 @@
 #include "Evaluation/Blending/MovieSceneBlendingActuatorID.h"
 #include "IMovieScenePlaybackClient.h"
 #include "Tracks/IMovieSceneTransformOrigin.h"
+#include "Evaluation/IMovieSceneMotionVectorSimulation.h"
+
 
 DECLARE_CYCLE_STAT(TEXT("Transform Track Evaluate"), MovieSceneEval_TransformTrack_Evaluate, STATGROUP_MovieSceneEval);
 DECLARE_CYCLE_STAT(TEXT("Transform Track Token Execute"), MovieSceneEval_TransformTrack_TokenExecute, STATGROUP_MovieSceneEval);
@@ -100,6 +102,48 @@ struct FComponentTransformActuator : TMovieSceneBlendingActuator<F3DTransformTra
 	}
 };
 
+
+
+/** Actuator that knows how to apply transform track tokens to an object */
+struct FSimulatedComponentTransformActuator : TMovieSceneBlendingActuator<F3DTransformTrackToken>
+{
+	FSimulatedComponentTransformActuator()
+		: TMovieSceneBlendingActuator<F3DTransformTrackToken>(GetActuatorTypeID())
+	{}
+
+	/** Access a unique identifier for this actuator type */
+	static FMovieSceneBlendingActuatorID GetActuatorTypeID()
+	{
+		static FMovieSceneAnimTypeID TypeID = TMovieSceneAnimTypeID<FSimulatedComponentTransformActuator>();
+		return FMovieSceneBlendingActuatorID(TypeID);
+	}
+
+	/** Get an object's current transform */
+	virtual F3DTransformTrackToken RetrieveCurrentValue(UObject* InObject, IMovieScenePlayer* Player) const
+	{
+		if (USceneComponent* SceneComponent = MovieSceneHelpers::SceneComponentFromRuntimeObject(InObject))
+		{
+			return F3DTransformTrackToken(SceneComponent->RelativeLocation, SceneComponent->RelativeRotation, SceneComponent->RelativeScale3D);
+		}
+
+		return F3DTransformTrackToken();
+	}
+
+	/** Set an object's transform */
+	virtual void Actuate(UObject* InObject, const F3DTransformTrackToken& InFinalValue, const TBlendableTokenStack<F3DTransformTrackToken>& OriginalStack, const FMovieSceneContext& Context, FPersistentEvaluationData& PersistentData, IMovieScenePlayer& Player) override
+	{
+		ensureMsgf(InObject, TEXT("Attempting to evaluate a Transform track with a null object."));
+
+		USceneComponent* SceneComponent = MovieSceneHelpers::SceneComponentFromRuntimeObject(InObject);
+		if (SceneComponent && Player.MotionVectorSimulation.IsValid())
+		{
+			Player.MotionVectorSimulation->Add(SceneComponent, FTransform(InFinalValue.Rotation, InFinalValue.Translation, InFinalValue.Scale), NAME_None);
+		}
+	}
+};
+
+
+
 FMovieSceneComponentTransformSectionTemplate::FMovieSceneComponentTransformSectionTemplate(const UMovieScene3DTransformSection& Section)
 	: TemplateData(Section)
 {
@@ -129,19 +173,17 @@ void FMovieSceneComponentTransformSectionTemplate::Initialize(const FMovieSceneE
 	}
 }
 
-void FMovieSceneComponentTransformSectionTemplate::Evaluate(const FMovieSceneEvaluationOperand& Operand, const FMovieSceneContext& Context, const FPersistentEvaluationData& PersistentData, FMovieSceneExecutionTokens& ExecutionTokens) const
+MovieScene::TMultiChannelValue<float, 9> FMovieSceneComponentTransformSectionTemplate::EvaluateTransform(FFrameTime Time, const FComponentTransformPersistentData* SectionData) const
 {
-	using namespace MovieScene;
-
-	TMultiChannelValue<float, 9> TransformValue = TemplateData.Evaluate(Context.GetTime());
+	MovieScene::TMultiChannelValue<float, 9> TransformValue = TemplateData.Evaluate(Time);
 	if (TransformValue.IsEmpty())
 	{
-		return;
+		return TransformValue;
 	}
 
 	// Apply origin transformation if necessary 
-	const FComponentTransformPersistentData* Data = PersistentData.FindSectionData<FComponentTransformPersistentData>();
-	if (TemplateData.BlendType == EMovieSceneBlendType::Absolute && Data)
+	
+	if (TemplateData.BlendType == EMovieSceneBlendType::Absolute && SectionData)
 	{
 		float Components[6] = {
 			TransformValue.Get(0, 0.f), TransformValue.Get(1, 0.f), TransformValue.Get(2, 0.f),
@@ -149,7 +191,7 @@ void FMovieSceneComponentTransformSectionTemplate::Evaluate(const FMovieSceneEva
 		};
 
 		FTransform AnimatedTransform(FRotator(Components[4], Components[5], Components[3]), FVector(Components[0], Components[1], Components[2]));
-		AnimatedTransform = AnimatedTransform * Data->Origin;
+		AnimatedTransform = AnimatedTransform * SectionData->Origin;
 
 		FVector Location = AnimatedTransform.GetTranslation();
 		Components[0] = Location.X;
@@ -170,23 +212,72 @@ void FMovieSceneComponentTransformSectionTemplate::Evaluate(const FMovieSceneEva
 		}
 	}
 
-	// Ensure the accumulator knows how to actually apply component transforms
-	FMovieSceneBlendingActuatorID ActuatorTypeID = FComponentTransformActuator::GetActuatorTypeID();
-	if (!ExecutionTokens.GetBlendingAccumulator().FindActuator<F3DTransformTrackToken>(ActuatorTypeID))
+	return TransformValue;
+}
+
+void FMovieSceneComponentTransformSectionTemplate::Evaluate(const FMovieSceneEvaluationOperand& Operand, const FMovieSceneContext& Context, const FPersistentEvaluationData& PersistentData, FMovieSceneExecutionTokens& ExecutionTokens) const
+{
+	using namespace MovieScene;
+
+	const FComponentTransformPersistentData* SectionData = PersistentData.FindSectionData<FComponentTransformPersistentData>();
+
+	TMultiChannelValue<float, 9> TransformValue = EvaluateTransform(Context.GetTime(), SectionData);
+	if (!TransformValue.IsEmpty())
 	{
-		ExecutionTokens.GetBlendingAccumulator().DefineActuator(ActuatorTypeID, MakeShared<FComponentTransformActuator>());
+		// Ensure the accumulator knows how to actually apply component transforms
+		FMovieSceneBlendingActuatorID ActuatorTypeID = FComponentTransformActuator::GetActuatorTypeID();
+		if (!ExecutionTokens.GetBlendingAccumulator().FindActuator<F3DTransformTrackToken>(ActuatorTypeID))
+		{
+			ExecutionTokens.GetBlendingAccumulator().DefineActuator(ActuatorTypeID, MakeShared<FComponentTransformActuator>());
+		}
+
+		// Add the blendable to the accumulator
+		float Weight = EvaluateEasing(Context.GetTime());
+		if (EnumHasAllFlags(TemplateData.Mask.GetChannels(), EMovieSceneTransformChannel::Weight))
+		{
+			float ManualWeight = 1.f;
+			TemplateData.ManualWeight.Evaluate(Context.GetTime(), ManualWeight);
+			Weight *= ManualWeight;
+		}
+
+		ExecutionTokens.BlendToken(ActuatorTypeID, TBlendableToken<F3DTransformTrackToken>(TransformValue, TemplateData.BlendType, Weight));
 	}
 
-	// Add the blendable to the accumulator
-	float Weight = EvaluateEasing(Context.GetTime());
-	if (EnumHasAllFlags(TemplateData.Mask.GetChannels(), EMovieSceneTransformChannel::Weight))
+	if (IMovieSceneMotionVectorSimulation::IsEnabled(PersistentData, Context))
 	{
-		float ManualWeight = 1.f;
-		TemplateData.ManualWeight.Evaluate(Context.GetTime(), ManualWeight);
-		Weight *= ManualWeight;
-	}
+		FFrameTime SimulationTime = IMovieSceneMotionVectorSimulation::GetSimulationTime(Context);
+		TMultiChannelValue<float, 9> SimulatedValue = EvaluateTransform(SimulationTime, SectionData);
 
-	ExecutionTokens.BlendToken(ActuatorTypeID, TBlendableToken<F3DTransformTrackToken>(TransformValue, TemplateData.BlendType, Weight));
+		if (!SimulatedValue.IsEmpty())
+		{
+			// Ensure the accumulator knows how to actually apply component transforms
+			FMovieSceneBlendingActuatorID SimulationActuatorTypeID = FSimulatedComponentTransformActuator::GetActuatorTypeID();
+			if (!ExecutionTokens.GetBlendingAccumulator().FindActuator<F3DTransformTrackToken>(SimulationActuatorTypeID))
+			{
+				ExecutionTokens.GetBlendingAccumulator().DefineActuator(SimulationActuatorTypeID, MakeShared<FSimulatedComponentTransformActuator>());
+			}
+
+			// Add the blendable to the accumulator
+			float Weight = EvaluateEasing(SimulationTime);
+			if (EnumHasAllFlags(TemplateData.Mask.GetChannels(), EMovieSceneTransformChannel::Weight))
+			{
+				float ManualWeight = 1.f;
+				TemplateData.ManualWeight.Evaluate(Context.GetTime(), ManualWeight);
+				Weight *= ManualWeight;
+			}
+
+			// Reproject backwards
+			for (int32 Index = 0; Index < 9; ++Index)
+			{
+				if (SimulatedValue.IsSet(Index) && TransformValue.IsSet(Index))
+				{
+					SimulatedValue.Set(Index, 2.f*TransformValue[Index] - SimulatedValue[Index]);
+				}
+			}
+
+			ExecutionTokens.BlendToken(SimulationActuatorTypeID, TBlendableToken<F3DTransformTrackToken>(SimulatedValue, TemplateData.BlendType, Weight));
+		}
+	}
 }
 
 void FMovieSceneComponentTransformSectionTemplate::Interrogate(const FMovieSceneContext& Context, FMovieSceneInterrogationData& Container, UObject* BindingOverride) const
diff --git a/Engine/Source/Runtime/MovieSceneTracks/Private/Evaluation/MovieSceneCameraCutTemplate.cpp b/Engine/Source/Runtime/MovieSceneTracks/Private/Evaluation/MovieSceneCameraCutTemplate.cpp
--- a/Engine/Source/Runtime/MovieSceneTracks/Private/Evaluation/MovieSceneCameraCutTemplate.cpp
+++ b/Engine/Source/Runtime/MovieSceneTracks/Private/Evaluation/MovieSceneCameraCutTemplate.cpp
@@ -6,9 +6,11 @@
 #include "IMovieScenePlayer.h"
 #include "GameFramework/Actor.h"
 #include "ContentStreaming.h"
+#include "Evaluation/IMovieSceneMotionVectorSimulation.h"
 
 DECLARE_CYCLE_STAT(TEXT("Camera Cut Track Token Execute"), MovieSceneEval_CameraCutTrack_TokenExecute, STATGROUP_MovieSceneEval);
 
+
 /** A movie scene execution token that sets up the streaming system with the camera cut location */
 struct FCameraCutPreRollExecutionToken : IMovieSceneExecutionToken
 {
@@ -139,6 +141,7 @@ struct FCameraCutExecutionToken : IMovieSceneExecutionToken
 		{
 			Player.UpdateCameraCut(CameraObject, CameraCutCache.LastLockedCamera.Get(), Context.HasJumped());
 			CameraCutCache.LastLockedCamera = CameraObject;
+			IMovieSceneMotionVectorSimulation::EnableThisFrame(PersistentData);
 		}
 		else if (CameraObject)
 		{
diff --git a/Engine/Source/Runtime/MovieSceneTracks/Private/Evaluation/MovieSceneSkeletalAnimationTemplate.cpp b/Engine/Source/Runtime/MovieSceneTracks/Private/Evaluation/MovieSceneSkeletalAnimationTemplate.cpp
--- a/Engine/Source/Runtime/MovieSceneTracks/Private/Evaluation/MovieSceneSkeletalAnimationTemplate.cpp
+++ b/Engine/Source/Runtime/MovieSceneTracks/Private/Evaluation/MovieSceneSkeletalAnimationTemplate.cpp
@@ -1,7 +1,7 @@
 // Copyright 1998-2019 Epic Games, Inc. All Rights Reserved.
 
 #include "Evaluation/MovieSceneSkeletalAnimationTemplate.h"
-
+#include "Evaluation/MovieSceneCameraCutTemplate.h"
 #include "Compilation/MovieSceneCompilerRules.h"
 #include "Animation/AnimMontage.h"
 #include "Animation/AnimSingleNodeInstance.h"
@@ -11,6 +11,8 @@
 #include "Evaluation/MovieSceneEvaluation.h"
 #include "IMovieScenePlayer.h"
 #include "UObject/ObjectKey.h"
+#include "Rendering/MotionVectorSimulation.h"
+#include "Evaluation/IMovieSceneMotionVectorSimulation.h"
 #include "UObject/StrongObjectPtr.h"
 
 
@@ -162,6 +164,10 @@ struct FMinimalAnimParameters
 	bool bSkipAnimNotifiers;
 	bool bForceCustomMode;
 };
+struct FSimulatedAnimParameters
+{
+	FMinimalAnimParameters AnimParams;
+};
 
 /** Montage player per section data */
 struct FMontagePlayerPerSectionData 
@@ -174,6 +180,7 @@ namespace MovieScene
 {
 	struct FBlendedAnimation
 	{
+		TArray<FMinimalAnimParameters> SimulatedAnimations;
 		TArray<FMinimalAnimParameters> AllAnimations;
 
 		FBlendedAnimation& Resolve(TMovieSceneInitialValueStore<FBlendedAnimation>& InitialValueStore)
@@ -186,6 +193,10 @@ namespace MovieScene
 	{
 		OutBlend.AllAnimations.Add(InValue);
 	}
+	void BlendValue(FBlendedAnimation& OutBlend, const FSimulatedAnimParameters& InValue, float Weight, EMovieSceneBlendType BlendType, TMovieSceneInitialValueStore<FBlendedAnimation>& InitialValueStore)
+	{
+		OutBlend.SimulatedAnimations.Add(InValue.AnimParams);
+	}
 
 	struct FComponentAnimationActuator : TMovieSceneBlendingActuator<FBlendedAnimation>
 	{
@@ -239,8 +250,7 @@ namespace MovieScene
 										PlayerStatus == EMovieScenePlayerStatus::Jumping || 
 										PlayerStatus == EMovieScenePlayerStatus::Scrubbing || 
 										(DeltaTime == 0.0f && PlayerStatus != EMovieScenePlayerStatus::Stopped); 
-		
-			static const bool bLooping = false;
+
 			//Need to zero all weights first since we may be blending animation that are keeping state but are no longer active.
 			
 			if(SequencerInstance)
@@ -260,27 +270,22 @@ namespace MovieScene
 					}
 				}
 			}
-			for (const FMinimalAnimParameters& AnimParams : InFinalValue.AllAnimations)
+
+			const bool bHasTickedThisFrame = SkeletalMeshComponent->PoseTickedThisFrame();
+			if (InFinalValue.SimulatedAnimations.Num() != 0 && Player.MotionVectorSimulation.IsValid())
 			{
-				Player.PreAnimatedState.SetCaptureEntity(AnimParams.EvaluationScope.Key, AnimParams.EvaluationScope.CompletionMode);
+				ApplyAnimations(PersistentData, Player, SkeletalMeshComponent, InFinalValue.SimulatedAnimations, DeltaTime, bPreviewPlayback, bFireNotifies, bResetDynamics);
 
-				if (bPreviewPlayback)
-				{
-					PreviewSetAnimPosition(PersistentData, Player, SkeletalMeshComponent,
-						AnimParams.SlotName, AnimParams.Section, AnimParams.Animation, AnimParams.EvalTime, AnimParams.BlendWeight,
-						bLooping, bFireNotifies && !AnimParams.bSkipAnimNotifiers, DeltaTime, Player.GetPlaybackStatus() == EMovieScenePlayerStatus::Playing, 
-						bResetDynamics, AnimParams.bForceCustomMode);
-				}
-				else
-				{
-					SetAnimPosition(PersistentData, Player, SkeletalMeshComponent,
-						AnimParams.SlotName, AnimParams.Section, AnimParams.Animation, AnimParams.EvalTime, AnimParams.BlendWeight,
-						bLooping, Player.GetPlaybackStatus() == EMovieScenePlayerStatus::Playing, bFireNotifies && !AnimParams.bSkipAnimNotifiers,
-						AnimParams.bForceCustomMode
-					);
-				}
+				SkeletalMeshComponent->TickAnimation(0.f, false);
+				SkeletalMeshComponent->RefreshBoneTransforms();
+				SkeletalMeshComponent->FinalizeBoneTransform();
+				SkeletalMeshComponent->ForceMotionVector();
+
+				SimulateMotionVectors(PersistentData, SkeletalMeshComponent, Player);
 			}
 
+			ApplyAnimations(PersistentData, Player, SkeletalMeshComponent, InFinalValue.AllAnimations, DeltaTime, bPreviewPlayback, bFireNotifies, bResetDynamics);
+
 			Player.PreAnimatedState.SetCaptureEntity(FMovieSceneEvaluationKey(), EMovieSceneCompletionMode::KeepState);
 		}
 
@@ -301,8 +306,49 @@ namespace MovieScene
 			return nullptr;
 		}
 
-		void SetAnimPosition(FPersistentEvaluationData& PersistentData, IMovieScenePlayer& Player, USkeletalMeshComponent* SkeletalMeshComponent, FName SlotName, FObjectKey Section, UAnimSequenceBase* InAnimSequence, float InPosition, float Weight, bool bLooping, bool bPlaying, bool bFireNotifies, bool bForceCustomMode)
+		void SimulateMotionVectors(FPersistentEvaluationData& PersistentData, USkeletalMeshComponent* SkeletalMeshComponent, IMovieScenePlayer& Player)
+		{
+			for (USceneComponent* Child : SkeletalMeshComponent->GetAttachChildren())
+			{
+				FName SocketName = Child->GetAttachSocketName();
+				if (SocketName != NAME_None)
+				{
+					FTransform SocketTransform = SkeletalMeshComponent->GetSocketTransform(SocketName, RTS_Component);
+					Player.MotionVectorSimulation->Add(SkeletalMeshComponent, SocketTransform, SocketName);
+				}
+			}
+		}
+
+		void ApplyAnimations(FPersistentEvaluationData& PersistentData, IMovieScenePlayer& Player, USkeletalMeshComponent* SkeletalMeshComponent, TArrayView<const FMinimalAnimParameters> Parameters, float DeltaTime, bool bPreviewPlayback, bool bFireNotifies, bool bResetDynamics)
 		{
+			const EMovieScenePlayerStatus::Type PlayerStatus = Player.GetPlaybackStatus();
+
+			for (const FMinimalAnimParameters& AnimParams : Parameters)
+			{
+				Player.PreAnimatedState.SetCaptureEntity(AnimParams.EvaluationScope.Key, AnimParams.EvaluationScope.CompletionMode);
+
+				if (bPreviewPlayback)
+				{
+					PreviewSetAnimPosition(PersistentData, Player, SkeletalMeshComponent,
+						AnimParams.SlotName, AnimParams.Section, AnimParams.Animation, AnimParams.EvalTime, AnimParams.BlendWeight,
+						bFireNotifies && !AnimParams.bSkipAnimNotifiers, DeltaTime, PlayerStatus == EMovieScenePlayerStatus::Playing, 
+						bResetDynamics, AnimParams.bForceCustomMode);
+				}
+				else
+				{
+					SetAnimPosition(PersistentData, Player, SkeletalMeshComponent,
+						AnimParams.SlotName, AnimParams.Section, AnimParams.Animation, AnimParams.EvalTime, AnimParams.BlendWeight,
+						PlayerStatus == EMovieScenePlayerStatus::Playing, bFireNotifies && !AnimParams.bSkipAnimNotifiers,
+						AnimParams.bForceCustomMode
+					);
+				}
+			}
+		}
+
+		void SetAnimPosition(FPersistentEvaluationData& PersistentData, IMovieScenePlayer& Player, USkeletalMeshComponent* SkeletalMeshComponent, FName SlotName, FObjectKey Section, UAnimSequenceBase* InAnimSequence, float InPosition, float Weight, bool bPlaying, bool bFireNotifies, bool bForceCustomMode)
+		{
+			static const bool bLooping = false;
+
 			if (!CanPlayAnimation(SkeletalMeshComponent, InAnimSequence))
 			{
 				return;
@@ -344,8 +390,10 @@ namespace MovieScene
 			}
 		}
 
-		void PreviewSetAnimPosition(FPersistentEvaluationData& PersistentData, IMovieScenePlayer& Player, USkeletalMeshComponent* SkeletalMeshComponent, FName SlotName, FObjectKey Section, UAnimSequenceBase* InAnimSequence, float InPosition, float Weight, bool bLooping, bool bFireNotifies, float DeltaTime, bool bPlaying, bool bResetDynamics, bool bForceCustomMode)
+		void PreviewSetAnimPosition(FPersistentEvaluationData& PersistentData, IMovieScenePlayer& Player, USkeletalMeshComponent* SkeletalMeshComponent, FName SlotName, FObjectKey Section, UAnimSequenceBase* InAnimSequence, float InPosition, float Weight, bool bFireNotifies, float DeltaTime, bool bPlaying, bool bResetDynamics, bool bForceCustomMode)
 		{
+			static const bool bLooping = false;
+
 			if (!CanPlayAnimation(SkeletalMeshComponent, InAnimSequence))
 			{
 				return;
@@ -413,13 +461,6 @@ void FMovieSceneSkeletalAnimationSectionTemplate::Evaluate(const FMovieSceneEval
 {
 	if (Params.Animation)
 	{
-		// calculate the time at which to evaluate the animation
-		float EvalTime = Params.MapTimeToAnimation(Context.GetTime(), Context.GetFrameRate());
-
-		float ManualWeight = 1.f;
-		Params.Weight.Evaluate(Context.GetTime(), ManualWeight);
-
-		const float Weight = ManualWeight * EvaluateEasing(Context.GetTime());
 
 		FOptionalMovieSceneBlendType BlendType = GetSourceSection()->GetBlendType();
 		check(BlendType.IsValid());
@@ -432,12 +473,38 @@ void FMovieSceneSkeletalAnimationSectionTemplate::Evaluate(const FMovieSceneEval
 			Accumulator.DefineActuator(ActuatorTypeID, MakeShared<MovieScene::FComponentAnimationActuator>());
 		}
 
+		// calculate the time at which to evaluate the animation
+		float EvalTime = Params.MapTimeToAnimation(Context.GetTime(), Context.GetFrameRate());
+
+		float ManualWeight = 1.f;
+		Params.Weight.Evaluate(Context.GetTime(), ManualWeight);
+
+		const float Weight = ManualWeight * EvaluateEasing(Context.GetTime());
+
 		// Add the blendable to the accumulator
 		FMinimalAnimParameters AnimParams(
 			Params.Animation, EvalTime, Weight, ExecutionTokens.GetCurrentScope(), Params.SlotName, GetSourceSection(), Params.bSkipAnimNotifiers, 
 			Params.bForceCustomMode
 		);
 		ExecutionTokens.BlendToken(ActuatorTypeID, TBlendableToken<MovieScene::FBlendedAnimation>(AnimParams, BlendType.Get(), 1.f));
+
+		if (IMovieSceneMotionVectorSimulation::IsEnabled(PersistentData, Context))
+		{
+			FFrameTime SimulatedTime = IMovieSceneMotionVectorSimulation::GetSimulationTime(Context);
+
+			// calculate the time at which to evaluate the animation
+			float SimulatedEvalTime = Params.MapTimeToAnimation(SimulatedTime, Context.GetFrameRate());
+
+			float SimulatedManualWeight = 1.f;
+			Params.Weight.Evaluate(SimulatedTime, SimulatedManualWeight);
+
+			const float SimulatedWeight = SimulatedManualWeight * EvaluateEasing(SimulatedTime);
+
+			FSimulatedAnimParameters SimulatedAnimParams{ AnimParams };
+			SimulatedAnimParams.AnimParams.EvalTime = SimulatedEvalTime;
+			SimulatedAnimParams.AnimParams.BlendWeight = SimulatedWeight;
+			ExecutionTokens.BlendToken(ActuatorTypeID, TBlendableToken<MovieScene::FBlendedAnimation>(SimulatedAnimParams, BlendType.Get(), 1.f));
+		}
 	}
 }
 
diff --git a/Engine/Source/Runtime/MovieSceneTracks/Private/Tracks/MovieSceneCameraCutTrack.cpp b/Engine/Source/Runtime/MovieSceneTracks/Private/Tracks/MovieSceneCameraCutTrack.cpp
--- a/Engine/Source/Runtime/MovieSceneTracks/Private/Tracks/MovieSceneCameraCutTrack.cpp
+++ b/Engine/Source/Runtime/MovieSceneTracks/Private/Tracks/MovieSceneCameraCutTrack.cpp
@@ -1,11 +1,14 @@
 // Copyright 1998-2019 Epic Games, Inc. All Rights Reserved.
 
 #include "Tracks/MovieSceneCameraCutTrack.h"
+#include "Tracks/MovieSceneSpawnTrack.h"
 #include "MovieScene.h"
 #include "MovieSceneCommonHelpers.h"
 #include "MovieSceneTimeHelpers.h"
 #include "Sections/MovieSceneCameraCutSection.h"
 #include "Evaluation/MovieSceneCameraCutTemplate.h"
+#include "Evaluation/MovieSceneEvaluationTrack.h"
+#include "IMovieSceneTracksModule.h"
 
 #define LOCTEXT_NAMESPACE "MovieSceneCameraCutTrack"
 
@@ -22,6 +25,12 @@ UMovieSceneCameraCutTrack::UMovieSceneCameraCutTrack( const FObjectInitializer&
 	EvalOptions.bEvaluateInPreroll = EvalOptions.bEvaluateInPostroll = false;
 }
 
+void UMovieSceneCameraCutTrack::PostCompile(FMovieSceneEvaluationTrack& OutTrack, const FMovieSceneTrackCompilerArgs& Args) const
+{
+	// Evaluate camera cuts in the same flush group as spawnables, this ensures cameras are set before anything else happens (so tracks can check whether a camera cut has happened)
+	OutTrack.SetEvaluationGroup(IMovieSceneTracksModule::GetEvaluationGroupName(EBuiltInEvaluationGroup::SpawnObjects));
+}
+
 
 UMovieSceneCameraCutSection* UMovieSceneCameraCutTrack::AddNewCameraCut(const FMovieSceneObjectBindingID& CameraBindingID, FFrameNumber StartTime)
 {
diff --git a/Engine/Source/Runtime/MovieSceneTracks/Public/Evaluation/MovieScene3DTransformTemplate.h b/Engine/Source/Runtime/MovieSceneTracks/Public/Evaluation/MovieScene3DTransformTemplate.h
--- a/Engine/Source/Runtime/MovieSceneTracks/Public/Evaluation/MovieScene3DTransformTemplate.h
+++ b/Engine/Source/Runtime/MovieSceneTracks/Public/Evaluation/MovieScene3DTransformTemplate.h
@@ -12,6 +12,7 @@
 #include "Evaluation/Blending/MovieSceneMultiChannelBlending.h"
 #include "MovieScene3DTransformTemplate.generated.h"
 
+struct FComponentTransformPersistentData;
 class UMovieScene3DTransformSection;
 
 USTRUCT()
@@ -67,4 +68,7 @@ protected:
 	virtual void Initialize(const FMovieSceneEvaluationOperand& Operand, const FMovieSceneContext& Context, FPersistentEvaluationData& PersistentData, IMovieScenePlayer& Player) const override;
 	virtual void Evaluate(const FMovieSceneEvaluationOperand& Operand, const FMovieSceneContext& Context, const FPersistentEvaluationData& PersistentData, FMovieSceneExecutionTokens& ExecutionTokens) const override;
 	virtual void Interrogate(const FMovieSceneContext& Context, FMovieSceneInterrogationData& Container, UObject* BindingOverride) const override;
+
+private:
+	MovieScene::TMultiChannelValue<float, 9> EvaluateTransform(FFrameTime Time, const FComponentTransformPersistentData* SectionData) const;
 };
diff --git a/Engine/Source/Runtime/MovieSceneTracks/Public/Tracks/MovieSceneCameraCutTrack.h b/Engine/Source/Runtime/MovieSceneTracks/Public/Tracks/MovieSceneCameraCutTrack.h
--- a/Engine/Source/Runtime/MovieSceneTracks/Public/Tracks/MovieSceneCameraCutTrack.h
+++ b/Engine/Source/Runtime/MovieSceneTracks/Public/Tracks/MovieSceneCameraCutTrack.h
@@ -35,7 +35,7 @@ public:
 public:
 
 	// UMovieSceneTrack interface
-
+	virtual void PostCompile(FMovieSceneEvaluationTrack& OutTrack, const FMovieSceneTrackCompilerArgs& Args) const override;
 	virtual void AddSection(UMovieSceneSection& Section) override;
 	virtual bool SupportsType(TSubclassOf<UMovieSceneSection> SectionClass) const override;
 	virtual UMovieSceneSection* CreateNewSection() override;
