Reorganised error text to make it easier to localize

diff --git a/Engine/Source/Runtime/MovieSceneTracks/Private/Evaluation/MovieSceneEventTemplate.cpp b/Engine/Source/Runtime/MovieSceneTracks/Private/Evaluation/MovieSceneEventTemplate.cpp
--- a/Engine/Source/Runtime/MovieSceneTracks/Private/Evaluation/MovieSceneEventTemplate.cpp
+++ b/Engine/Source/Runtime/MovieSceneTracks/Private/Evaluation/MovieSceneEventTemplate.cpp
@@ -246,9 +246,8 @@ struct FEventTriggerExecutionToken
 			if (!Function)
 			{
 				FMessageLog("PIE").Warning()
-					->AddToken(FTextToken::Create(FText::Format(LOCTEXT("LevelBP_MissingEvent_Error1", "Failed to trigger event '{0}' for"), FText::FromName(EventName))))
 					->AddToken(FUObjectToken::Create(Player.GetEvaluationTemplate().GetSequence(Operand.SequenceID)))
-					->AddToken(FTextToken::Create(LOCTEXT("LevelBP_MissingEvent_Error2", "because the function does not exist on the director instance.")));
+					->AddToken(FTextToken::Create(FText::Format(LOCTEXT("LevelBP_MissingEvent_Error", "Failed to trigger event '{0}' because the function does not exist on the director instance."), FText::FromName(EventName))));
 				continue;
 			}
 #if WITH_EDITOR
@@ -277,11 +276,9 @@ struct FEventTriggerExecutionToken
 				if (NumLevelScripts > 0 && NumLevelScripts == EventContexts.Num() && !Function->PropertyLink->IsA<UInterfaceProperty>())
 				{
 					FMessageLog("PIE").Warning()
-						->AddToken(FTextToken::Create(LOCTEXT("LevelBP_ObjectPin_Error1", "Failed to trigger event")))
 						->AddToken(FUObjectToken::Create(Function))
-						->AddToken(FTextToken::Create(LOCTEXT("LevelBP_ObjectPin_Error2", "within")))
 						->AddToken(FUObjectToken::Create(Player.GetEvaluationTemplate().GetSequence(Operand.SequenceID)))
-						->AddToken(FTextToken::Create(LOCTEXT("LevelBP_ObjectPin_Error3", "because only Interface pins are supported for master tracks within Level Sequences. Please remove the pin, or change it to an interface that is implemented on the desired level blueprint.")));
+						->AddToken(FTextToken::Create(LOCTEXT("LevelBP_ObjectPin_Error", "Failed to trigger event: only Interface pins are supported for master tracks within Level Sequences. Please remove the pin, or change it to an interface that is implemented on the desired level blueprint.")));
 					continue;
 				}
 
@@ -293,11 +290,9 @@ struct FEventTriggerExecutionToken
 			else
 			{
 				FMessageLog("PIE").Warning()
-					->AddToken(FTextToken::Create(LOCTEXT("LevelBP_InvalidEvent_Error1", "Failed to trigger event")))
 					->AddToken(FUObjectToken::Create(Function))
-					->AddToken(FTextToken::Create(LOCTEXT("LevelBP_InvalidEvent_Error2", "within")))
 					->AddToken(FUObjectToken::Create(Player.GetEvaluationTemplate().GetSequence(Operand.SequenceID)))
-					->AddToken(FTextToken::Create(LOCTEXT("LevelBP_InvalidEvent_Error3", "because its signature is not compatible. Function signatures must have either 0 or 1 (non-ref) parameters.")));
+					->AddToken(FTextToken::Create(LOCTEXT("LevelBP_InvalidEvent_Error", "Failed to trigger event because its signature is not compatible. Function signatures must have either 0 or 1 (non-ref) parameters.")));
 			}
 		}
 	}
@@ -343,11 +338,9 @@ struct FEventTriggerExecutionToken
 		}
 
 		FMessageLog("PIE").Warning()
-			->AddToken(FTextToken::Create(LOCTEXT("LevelBP_InvalidObjectEvent_Error1", "Failed to trigger event")))
 			->AddToken(FUObjectToken::Create(Function))
-			->AddToken(FTextToken::Create(LOCTEXT("LevelBP_InvalidObjectEvent_Error2", "within")))
 			->AddToken(FUObjectToken::Create(Player.GetEvaluationTemplate().GetSequence(SequenceID)))
-			->AddToken(FTextToken::Create(FText::Format(LOCTEXT("LevelBP_InvalidObjectEvent_Error3", "because its signature is not compatible. Function expects a '%s' parameter, but only object and interface parameters are supported."), FText::FromName(Function->PropertyLink->GetClass()->GetFName()))));
+			->AddToken(FTextToken::Create(FText::Format(LOCTEXT("LevelBP_InvalidObjectEvent_Error", "Failed to trigger event because its signature is not compatible. Function expects a '{0}' parameter, but only object and interface parameters are supported."), FText::FromName(Function->PropertyLink->GetClass()->GetFName()))));
 	}
 
 #if !NO_LOGGING
