Sequencer Skel Anim Looping works again with frame offsets.

diff --git a/Engine/Source/Runtime/MovieSceneTracks/Private/Evaluation/MovieSceneSkeletalAnimationTemplate.cpp b/Engine/Source/Runtime/MovieSceneTracks/Private/Evaluation/MovieSceneSkeletalAnimationTemplate.cpp
--- a/Engine/Source/Runtime/MovieSceneTracks/Private/Evaluation/MovieSceneSkeletalAnimationTemplate.cpp
+++ b/Engine/Source/Runtime/MovieSceneTracks/Private/Evaluation/MovieSceneSkeletalAnimationTemplate.cpp
@@ -534,7 +534,7 @@ float FMovieSceneSkeletalAnimationSectionTemplateParameters::MapTimeToAnimation(
 	FFrameTime AnimationLength = GetSequenceLength() * InFrameRate;
 	int32 LengthInFrames = AnimationLength.FrameNumber.Value + (int)(AnimationLength.GetSubFrame() + 0.5f) + 1;
 	//we only play end if we are not looping, and assuming we are looping if Length is greater than default length;
-	bool bLooping = (SectionEndTime.Value - SectionStartTime.Value) > LengthInFrames;
+	bool bLooping = (SectionEndTime.Value - SectionStartTime.Value + StartFrameOffset + EndFrameOffset) > LengthInFrames;
 
 	InPosition = FMath::Clamp(InPosition, FFrameTime(SectionStartTime), FFrameTime(SectionEndTime -1));
 	
