Sequencer: Fix crash copying animation sections. Check for valid movie scene on PostLoad.

diff --git a/Engine/Source/Runtime/MovieSceneTracks/Private/Sections/MovieSceneSkeletalAnimationSection.cpp b/Engine/Source/Runtime/MovieSceneTracks/Private/Sections/MovieSceneSkeletalAnimationSection.cpp
--- a/Engine/Source/Runtime/MovieSceneTracks/Private/Sections/MovieSceneSkeletalAnimationSection.cpp
+++ b/Engine/Source/Runtime/MovieSceneTracks/Private/Sections/MovieSceneSkeletalAnimationSection.cpp
@@ -110,21 +110,26 @@ void UMovieSceneSkeletalAnimationSection::PostLoad()
 		Params.SlotName = SlotName_DEPRECATED;
 	}
 
-	FFrameRate DisplayRate = GetTypedOuter<UMovieScene>()->GetDisplayRate();
-	FFrameRate TickResolution = GetTypedOuter<UMovieScene>()->GetTickResolution();
+	UMovieScene* MovieScene = GetTypedOuter<UMovieScene>();
 
-	if (Params.StartOffset_DEPRECATED != SkeletalDeprecatedMagicNumber)
+	if (MovieScene)
 	{
-		Params.StartFrameOffset = ConvertFrameTime(FFrameTime::FromDecimal(DisplayRate.AsDecimal() * Params.StartOffset_DEPRECATED), DisplayRate, TickResolution).FrameNumber;
+		FFrameRate DisplayRate = MovieScene->GetDisplayRate();
+		FFrameRate TickResolution = MovieScene->GetTickResolution();
 
-		Params.StartOffset_DEPRECATED = SkeletalDeprecatedMagicNumber;
-	}
+		if (Params.StartOffset_DEPRECATED != SkeletalDeprecatedMagicNumber)
+		{
+			Params.StartFrameOffset = ConvertFrameTime(FFrameTime::FromDecimal(DisplayRate.AsDecimal() * Params.StartOffset_DEPRECATED), DisplayRate, TickResolution).FrameNumber;
 
-	if (Params.EndOffset_DEPRECATED != SkeletalDeprecatedMagicNumber)
-	{
-		Params.EndFrameOffset = ConvertFrameTime(FFrameTime::FromDecimal(DisplayRate.AsDecimal() * Params.EndOffset_DEPRECATED), DisplayRate, TickResolution).FrameNumber;
+			Params.StartOffset_DEPRECATED = SkeletalDeprecatedMagicNumber;
+		}
 
-		Params.EndOffset_DEPRECATED = SkeletalDeprecatedMagicNumber;
+		if (Params.EndOffset_DEPRECATED != SkeletalDeprecatedMagicNumber)
+		{
+			Params.EndFrameOffset = ConvertFrameTime(FFrameTime::FromDecimal(DisplayRate.AsDecimal() * Params.EndOffset_DEPRECATED), DisplayRate, TickResolution).FrameNumber;
+
+			Params.EndOffset_DEPRECATED = SkeletalDeprecatedMagicNumber;
+		}
 	}
 
 	// if version is less than this
