Sequencer: Changed Transform origins to be implemented as shared global data rather to remove the need to Initialize transform tracks

Transform tracks are heavily used, whereas transform origins are not. As such we want to remove the ::Initialize call that is unnecessary for the majority of transform tracks. We now compute the sequence transform origin once at the start of the frame, and use that for the rest of the evaluation.

diff --git a/Engine/Source/Runtime/LevelSequence/Private/LevelSequenceActor.cpp b/Engine/Source/Runtime/LevelSequence/Private/LevelSequenceActor.cpp
--- a/Engine/Source/Runtime/LevelSequence/Private/LevelSequenceActor.cpp
+++ b/Engine/Source/Runtime/LevelSequence/Private/LevelSequenceActor.cpp
@@ -6,6 +6,7 @@
 #include "Components/BillboardComponent.h"
 #include "LevelSequenceBurnIn.h"
 #include "DefaultLevelSequenceInstanceData.h"
+#include "Evaluation/MovieScene3DTransformTemplate.h"
 #include "Engine/ActorChannel.h"
 #include "Logging/MessageLog.h"
 #include "Misc/UObjectToken.h"
@@ -161,6 +162,25 @@ void ALevelSequenceActor::Tick(float DeltaSeconds)
 
 	if (SequencePlayer)
 	{
+		// If the global instance data implements a transform origin interface, use its transform as an origin for this frame
+		{
+			UObject*                          InstanceData = GetInstanceData();
+			const IMovieSceneTransformOrigin* RawInterface = Cast<IMovieSceneTransformOrigin>(InstanceData);
+
+			const bool bHasInterface = RawInterface || (InstanceData && InstanceData->GetClass()->ImplementsInterface(UMovieSceneTransformOrigin::StaticClass()));
+			if (bHasInterface)
+			{
+				static FSharedPersistentDataKey GlobalTransformDataKey = FGlobalTransformPersistentData::GetDataKey();
+
+				// Retrieve the current origin
+				FTransform TransformOrigin = RawInterface ? RawInterface->GetTransformOrigin() : IMovieSceneTransformOrigin::Execute_BP_GetTransformOrigin(InstanceData);
+
+				// Assign the transform origin to the peristent data so it can be queried in Evaluate
+				FPersistentEvaluationData PersistentData(*SequencePlayer);
+				PersistentData.GetOrAdd<FGlobalTransformPersistentData>(GlobalTransformDataKey).Origin = TransformOrigin;
+			}
+		}
+
 		SequencePlayer->Update(DeltaSeconds);
 	}
 }
diff --git a/Engine/Source/Runtime/MovieSceneTracks/Private/Evaluation/MovieScene3DTransformTemplate.cpp b/Engine/Source/Runtime/MovieSceneTracks/Private/Evaluation/MovieScene3DTransformTemplate.cpp
--- a/Engine/Source/Runtime/MovieSceneTracks/Private/Evaluation/MovieScene3DTransformTemplate.cpp
+++ b/Engine/Source/Runtime/MovieSceneTracks/Private/Evaluation/MovieScene3DTransformTemplate.cpp
@@ -18,6 +18,12 @@
 DECLARE_CYCLE_STAT(TEXT("Transform Track Evaluate"), MovieSceneEval_TransformTrack_Evaluate, STATGROUP_MovieSceneEval);
 DECLARE_CYCLE_STAT(TEXT("Transform Track Token Execute"), MovieSceneEval_TransformTrack_TokenExecute, STATGROUP_MovieSceneEval);
 
+FSharedPersistentDataKey FGlobalTransformPersistentData::GetDataKey()
+{
+	static FSharedPersistentDataKey Key(FMovieSceneSharedDataId::Allocate(), FMovieSceneEvaluationOperand());
+	return Key;
+}
+
 namespace MovieScene
 {
 	/** Convert a transform track token to a 9 channel float */
@@ -149,31 +155,7 @@ FMovieSceneComponentTransformSectionTemplate::FMovieSceneComponentTransformSecti
 {
 }
 
-struct FComponentTransformPersistentData : IPersistentEvaluationData
-{
-	FTransform Origin;
-};
-
-void FMovieSceneComponentTransformSectionTemplate::Initialize(const FMovieSceneEvaluationOperand& Operand, const FMovieSceneContext& Context, FPersistentEvaluationData& PersistentData, IMovieScenePlayer& Player) const
-{
-	// If the global instance data implements a transform origin interface, use its transform as an origin for this transform
-	IMovieScenePlaybackClient*        PlaybackClient = Player.GetPlaybackClient();
-	UObject*                          InstanceData   = PlaybackClient ? PlaybackClient->GetInstanceData() : nullptr;
-	const IMovieSceneTransformOrigin* RawInterface   = Cast<IMovieSceneTransformOrigin>(InstanceData);
-
-	const bool bHasInterface = RawInterface || (InstanceData && InstanceData->GetClass()->ImplementsInterface(UMovieSceneTransformOrigin::StaticClass()));
-
-	if (bHasInterface && TemplateData.BlendType == EMovieSceneBlendType::Absolute)
-	{
-		// Retrieve the current origin
-		FTransform TransformOrigin = RawInterface ? RawInterface->GetTransformOrigin() : IMovieSceneTransformOrigin::Execute_BP_GetTransformOrigin(InstanceData);
-
-		// Assign the transform origin to the peristent data so it can be queried in Evaluate
-		PersistentData.GetOrAddSectionData<FComponentTransformPersistentData>().Origin = TransformOrigin;
-	}
-}
-
-MovieScene::TMultiChannelValue<float, 9> FMovieSceneComponentTransformSectionTemplate::EvaluateTransform(FFrameTime Time, const FComponentTransformPersistentData* SectionData) const
+MovieScene::TMultiChannelValue<float, 9> FMovieSceneComponentTransformSectionTemplate::EvaluateTransform(FFrameTime Time, const FGlobalTransformPersistentData* GlobalTransformData) const
 {
 	MovieScene::TMultiChannelValue<float, 9> TransformValue = TemplateData.Evaluate(Time);
 	if (TransformValue.IsEmpty())
@@ -182,8 +164,7 @@ MovieScene::TMultiChannelValue<float, 9> FMovieSceneComponentTransformSectionTem
 	}
 
 	// Apply origin transformation if necessary 
-	
-	if (TemplateData.BlendType == EMovieSceneBlendType::Absolute && SectionData)
+	if (TemplateData.BlendType == EMovieSceneBlendType::Absolute && GlobalTransformData)
 	{
 		float Components[6] = {
 			TransformValue.Get(0, 0.f), TransformValue.Get(1, 0.f), TransformValue.Get(2, 0.f),
@@ -191,7 +172,7 @@ MovieScene::TMultiChannelValue<float, 9> FMovieSceneComponentTransformSectionTem
 		};
 
 		FTransform AnimatedTransform(FRotator(Components[4], Components[5], Components[3]), FVector(Components[0], Components[1], Components[2]));
-		AnimatedTransform = AnimatedTransform * SectionData->Origin;
+		AnimatedTransform = AnimatedTransform * GlobalTransformData->Origin;
 
 		FVector Location = AnimatedTransform.GetTranslation();
 		Components[0] = Location.X;
@@ -219,9 +200,10 @@ void FMovieSceneComponentTransformSectionTemplate::Evaluate(const FMovieSceneEva
 {
 	using namespace MovieScene;
 
-	const FComponentTransformPersistentData* SectionData = PersistentData.FindSectionData<FComponentTransformPersistentData>();
+	static FSharedPersistentDataKey GlobalTransformDataKey = FGlobalTransformPersistentData::GetDataKey();
+	const FGlobalTransformPersistentData* GlobalTransformData = PersistentData.Find<FGlobalTransformPersistentData>(GlobalTransformDataKey);
 
-	TMultiChannelValue<float, 9> TransformValue = EvaluateTransform(Context.GetTime(), SectionData);
+	TMultiChannelValue<float, 9> TransformValue = EvaluateTransform(Context.GetTime(), GlobalTransformData);
 	if (!TransformValue.IsEmpty())
 	{
 		// Ensure the accumulator knows how to actually apply component transforms
@@ -246,7 +228,7 @@ void FMovieSceneComponentTransformSectionTemplate::Evaluate(const FMovieSceneEva
 	if (IMovieSceneMotionVectorSimulation::IsEnabled(PersistentData, Context))
 	{
 		FFrameTime SimulationTime = IMovieSceneMotionVectorSimulation::GetSimulationTime(Context);
-		TMultiChannelValue<float, 9> SimulatedValue = EvaluateTransform(SimulationTime, SectionData);
+		TMultiChannelValue<float, 9> SimulatedValue = EvaluateTransform(SimulationTime, GlobalTransformData);
 
 		if (!SimulatedValue.IsEmpty())
 		{
diff --git a/Engine/Source/Runtime/MovieSceneTracks/Public/Evaluation/MovieScene3DTransformTemplate.h b/Engine/Source/Runtime/MovieSceneTracks/Public/Evaluation/MovieScene3DTransformTemplate.h
--- a/Engine/Source/Runtime/MovieSceneTracks/Public/Evaluation/MovieScene3DTransformTemplate.h
+++ b/Engine/Source/Runtime/MovieSceneTracks/Public/Evaluation/MovieScene3DTransformTemplate.h
@@ -15,6 +15,13 @@
 struct FComponentTransformPersistentData;
 class UMovieScene3DTransformSection;
 
+struct MOVIESCENETRACKS_API FGlobalTransformPersistentData : IPersistentEvaluationData
+{
+	static FSharedPersistentDataKey GetDataKey();
+
+	FTransform Origin;
+};
+
 USTRUCT()
 struct FMovieScene3DTransformTemplateData
 {
@@ -64,11 +71,9 @@ struct FMovieSceneComponentTransformSectionTemplate : public FMovieSceneEvalTemp
 protected:
 
 	virtual UScriptStruct& GetScriptStructImpl() const override { return *StaticStruct(); }
-	virtual void SetupOverrides() override { EnableOverrides(RequiresInitializeFlag); }
-	virtual void Initialize(const FMovieSceneEvaluationOperand& Operand, const FMovieSceneContext& Context, FPersistentEvaluationData& PersistentData, IMovieScenePlayer& Player) const override;
 	virtual void Evaluate(const FMovieSceneEvaluationOperand& Operand, const FMovieSceneContext& Context, const FPersistentEvaluationData& PersistentData, FMovieSceneExecutionTokens& ExecutionTokens) const override;
 	virtual void Interrogate(const FMovieSceneContext& Context, FMovieSceneInterrogationData& Container, UObject* BindingOverride) const override;
 
 private:
-	MovieScene::TMultiChannelValue<float, 9> EvaluateTransform(FFrameTime Time, const FComponentTransformPersistentData* SectionData) const;
+	MovieScene::TMultiChannelValue<float, 9> EvaluateTransform(FFrameTime Time, const FGlobalTransformPersistentData* GlobalTransformData) const;
 };
