Copying //UE4/Dev-Anim to Dev-Main (//UE4/Dev-Main)

diff --git a/Engine/Source/Runtime/AnimGraphRuntime/Private/CommonAnimationLibrary.cpp b/Engine/Source/Runtime/AnimGraphRuntime/Private/CommonAnimationLibrary.cpp
--- a/Engine/Source/Runtime/AnimGraphRuntime/Private/CommonAnimationLibrary.cpp
+++ b/Engine/Source/Runtime/AnimGraphRuntime/Private/CommonAnimationLibrary.cpp
@@ -109,8 +109,8 @@ float CommonAnimationLibrary::ScalarEasing(float Value, const FRuntimeFloatCurve
 
 FVector CommonAnimationLibrary::RetargetSingleLocation(
 	FVector Location,
-	FTransform Source,
-	FTransform Target,
+	const FTransform& Source,
+	const FTransform& Target,
 	const FRuntimeFloatCurve& CustomCurve,
 	EEasingFuncType EasingType,
 	bool bFlipEasing,
@@ -156,9 +156,9 @@ FVector CommonAnimationLibrary::RetargetSingleLocation(
 }
 
 FQuat CommonAnimationLibrary::RetargetSingleRotation(
-	FQuat Rotation,
-	FTransform Source,
-	FTransform Target,
+	const FQuat& RotationIn,
+	const FTransform& Source,
+	const FTransform& Target,
 	const FRuntimeFloatCurve& CustomCurve,
 	EEasingFuncType EasingType,
 	bool bFlipEasing,
@@ -172,6 +172,8 @@ FQuat CommonAnimationLibrary::RetargetSingleRotation(
 	float TargetMaximum
 )
 {
+	FQuat Rotation = RotationIn;
+
 	float SourceDelta = SourceMaximum - SourceMinimum;
 	float TargetDelta = TargetMaximum - TargetMinimum;
 	if (FMath::IsNearlyEqual(SourceDelta, 0.f) || FMath::IsNearlyEqual(TargetDelta, 0.f))
diff --git a/Engine/Source/Runtime/AnimGraphRuntime/Public/AnimNodes/AnimNode_MultiWayBlend.h b/Engine/Source/Runtime/AnimGraphRuntime/Public/AnimNodes/AnimNode_MultiWayBlend.h
--- a/Engine/Source/Runtime/AnimGraphRuntime/Public/AnimNodes/AnimNode_MultiWayBlend.h
+++ b/Engine/Source/Runtime/AnimGraphRuntime/Public/AnimNodes/AnimNode_MultiWayBlend.h
@@ -12,8 +12,6 @@ struct ANIMGRAPHRUNTIME_API FAnimNode_MultiWayBlend : public FAnimNode_Base
 {
 	GENERATED_USTRUCT_BODY()
 
-	friend struct FAnimSequencerInstanceProxy;
-
 public:
 
 	UPROPERTY(EditAnywhere, BlueprintReadWrite, Category = Links)
@@ -54,6 +52,7 @@ public:
 	{
 		Poses.AddDefaulted();
 		DesiredAlphas.AddZeroed();
+		UpdateCachedAlphas();
 
 		return Poses.Num();
 	}
@@ -61,12 +60,14 @@ public:
 	void RemovePose(int32 PoseIndex)
 	{
 		Poses.RemoveAt(PoseIndex);
+		CachedAlphas.RemoveAt(PoseIndex);
 		DesiredAlphas.RemoveAt(PoseIndex);
 	}
 
 	void ResetPoses()
 	{
 		Poses.Reset();
+		CachedAlphas.Reset();
 		DesiredAlphas.Reset();
 	}
 
diff --git a/Engine/Source/Runtime/AnimGraphRuntime/Public/CommonAnimationLibrary.h b/Engine/Source/Runtime/AnimGraphRuntime/Public/CommonAnimationLibrary.h
--- a/Engine/Source/Runtime/AnimGraphRuntime/Public/CommonAnimationLibrary.h
+++ b/Engine/Source/Runtime/AnimGraphRuntime/Public/CommonAnimationLibrary.h
@@ -45,8 +45,8 @@ namespace CommonAnimationLibrary
 	 */
 	ANIMGRAPHRUNTIME_API FVector RetargetSingleLocation(
 		FVector Location,
-		FTransform Source,
-		FTransform Target,
+		const FTransform& Source,
+		const FTransform& Target,
 		const FRuntimeFloatCurve& CustomCurve,
 		EEasingFuncType EasingType = EEasingFuncType::Linear,
 		bool bFlipEasing = false,
@@ -77,9 +77,9 @@ namespace CommonAnimationLibrary
 	 *	@param EasingType The type of easing to apply
 	 */
 	ANIMGRAPHRUNTIME_API FQuat RetargetSingleRotation(
-		FQuat Rotation,
-		FTransform Source,
-		FTransform Target,
+		const FQuat& Rotation,
+		const FTransform& Source,
+		const FTransform& Target,
 		const FRuntimeFloatCurve& CustomCurve,
 		EEasingFuncType EasingType = EEasingFuncType::Linear,
 		bool bFlipEasing = false,
diff --git a/Engine/Source/Runtime/Engine/Classes/Animation/Skeleton.h b/Engine/Source/Runtime/Engine/Classes/Animation/Skeleton.h
--- a/Engine/Source/Runtime/Engine/Classes/Animation/Skeleton.h
+++ b/Engine/Source/Runtime/Engine/Classes/Animation/Skeleton.h
@@ -837,8 +837,8 @@ public:
 	ENGINE_API void RemoveBonesFromSkeleton(const TArray<FName>& BonesToRemove, bool bRemoveChildBones);
 
 	// Asset registry information for animation notifies
-	static const FName AnimNotifyTag;
-	static const FString AnimNotifyTagDelimiter;
+	ENGINE_API static const FName AnimNotifyTag;
+	ENGINE_API static const FString AnimNotifyTagDelimiter;
 
 	// Asset registry information for animation curves
 	ENGINE_API static const FName CurveNameTag;
diff --git a/Engine/Source/Runtime/Engine/Private/Animation/AnimNode_CustomProperty.cpp b/Engine/Source/Runtime/Engine/Private/Animation/AnimNode_CustomProperty.cpp
--- a/Engine/Source/Runtime/Engine/Private/Animation/AnimNode_CustomProperty.cpp
+++ b/Engine/Source/Runtime/Engine/Private/Animation/AnimNode_CustomProperty.cpp
@@ -6,7 +6,8 @@
 #include "Animation/AnimNode_LinkedInputPose.h"
 
 FAnimNode_CustomProperty::FAnimNode_CustomProperty()
-	: TargetInstance(nullptr)
+	: FAnimNode_Base()
+	, TargetInstance(nullptr)
 #if WITH_EDITOR
 	, bReinitializeProperties(false)
 #endif // WITH_EDITOR
diff --git a/Engine/Source/Runtime/Engine/Private/Animation/AnimNode_StateMachine.cpp b/Engine/Source/Runtime/Engine/Private/Animation/AnimNode_StateMachine.cpp
--- a/Engine/Source/Runtime/Engine/Private/Animation/AnimNode_StateMachine.cpp
+++ b/Engine/Source/Runtime/Engine/Private/Animation/AnimNode_StateMachine.cpp
@@ -451,10 +451,6 @@ void FAnimNode_StateMachine::Update_AnyThread(const FAnimationUpdateContext& Con
 					{
 						InertializationNode->Request(ReferenceTransition.CrossfadeDuration);
 					}
-					else
-					{
-						UE_LOG(LogAnimation, Error, TEXT("FAnimNode_StateMachine: No inertialization context found for transition"));
-					}
 				}
 
 #if WITH_EDITORONLY_DATA
diff --git a/Engine/Source/Runtime/Engine/Private/Animation/AnimSequence.cpp b/Engine/Source/Runtime/Engine/Private/Animation/AnimSequence.cpp
--- a/Engine/Source/Runtime/Engine/Private/Animation/AnimSequence.cpp
+++ b/Engine/Source/Runtime/Engine/Private/Animation/AnimSequence.cpp
@@ -1672,25 +1672,13 @@ void UAnimSequence::GetBonePose_Additive(FCompactPose& OutPose, FBlendedCurve& O
 
 void UAnimSequence::GetAdditiveBasePose(FCompactPose& OutPose, FBlendedCurve& OutCurve, const FAnimExtractContext& ExtractionContext) const
 {
-	const bool bGetAdditiveBasePoseValid = (RefPoseType == ABPT_RefPose || !RefPoseSeq->IsValidAdditive() || RefPoseSeq->RawAnimationData.Num() > 0);
-	if (!bGetAdditiveBasePoseValid) //If this fails there is not enough information to get the base pose
-	{
-		FString Name = GetName();
-		FString SkelName = GetSkeleton() ? GetSkeleton()->GetName() : TEXT("NoSkeleton");
-		FString RefName = RefPoseSeq->GetName();
-		FString RefSkelName = RefPoseSeq->GetSkeleton() ? RefPoseSeq->GetSkeleton()->GetName() : TEXT("NoRefSeqSkeleton");
-		const TCHAR* NeedsLoad = RefPoseSeq->HasAnyFlags(RF_NeedLoad) ? TEXT("Yes") : TEXT("No");
-
-		checkf(false, TEXT("Cannot get valid base pose for Anim: ['%s' (Skel:%s)] RefSeq: ['%s' (Skel:%s)] RawAnimDataNum: %i NeedsLoad: %s"), *Name, *SkelName, *RefName, *RefSkelName, RefPoseSeq->RawAnimationData.Num(), NeedsLoad);
-	}
-
 	switch (RefPoseType)
 	{
 		// use whole animation as a base pose. Need BasePoseSeq.
 		case ABPT_AnimScaled:
 		{
 			// normalize time to fit base seq
-			const float Fraction = FMath::Clamp<float>(ExtractionContext.CurrentTime / SequenceLength, 0.f, 1.f);
+			const float Fraction = (SequenceLength > 0.f)? FMath::Clamp<float>(ExtractionContext.CurrentTime / SequenceLength, 0.f, 1.f) : 0.f;
 			const float BasePoseTime = RefPoseSeq->SequenceLength * Fraction;
 
 			FAnimExtractContext BasePoseExtractionContext(ExtractionContext);
diff --git a/Engine/Source/Runtime/Engine/Private/Components/SkinnedMeshComponent.cpp b/Engine/Source/Runtime/Engine/Private/Components/SkinnedMeshComponent.cpp
--- a/Engine/Source/Runtime/Engine/Private/Components/SkinnedMeshComponent.cpp
+++ b/Engine/Source/Runtime/Engine/Private/Components/SkinnedMeshComponent.cpp
@@ -1457,6 +1457,11 @@ void USkinnedMeshComponent::SetSkeletalMesh(USkeletalMesh* InSkelMesh, bool bRei
 
 		SkeletalMesh = InSkelMesh;
 
+		for (TWeakObjectPtr<USkinnedMeshComponent>& SlavePoseComponent : SlavePoseComponents)
+		{
+			SlavePoseComponent->UpdateMasterBoneMap();
+		}
+
 		// Don't init anim state if not registered
 		if (IsRegistered())
 		{
diff --git a/Engine/Source/Runtime/MovieScene/Private/Channels/MovieSceneChannelEditorData.cpp b/Engine/Source/Runtime/MovieScene/Private/Channels/MovieSceneChannelEditorData.cpp
--- a/Engine/Source/Runtime/MovieScene/Private/Channels/MovieSceneChannelEditorData.cpp
+++ b/Engine/Source/Runtime/MovieScene/Private/Channels/MovieSceneChannelEditorData.cpp
@@ -26,8 +26,8 @@ FMovieSceneChannelMetaData::FMovieSceneChannelMetaData()
 	, Name(NAME_None)
 {}
 
-FMovieSceneChannelMetaData::FMovieSceneChannelMetaData(FName InName, FText InDisplayText, FText InGroup)
-	: bEnabled(true)
+FMovieSceneChannelMetaData::FMovieSceneChannelMetaData(FName InName, FText InDisplayText, FText InGroup, bool bInEnabled)
+	: bEnabled(bInEnabled)
 	, bCanCollapseToTrack(true)
 	, SortOrder(0)
 	, Name(InName)
diff --git a/Engine/Source/Runtime/MovieScene/Public/Channels/MovieSceneChannelEditorData.h b/Engine/Source/Runtime/MovieScene/Public/Channels/MovieSceneChannelEditorData.h
--- a/Engine/Source/Runtime/MovieScene/Public/Channels/MovieSceneChannelEditorData.h
+++ b/Engine/Source/Runtime/MovieScene/Public/Channels/MovieSceneChannelEditorData.h
@@ -28,8 +28,9 @@ struct FMovieSceneChannelMetaData
 	 * @param InName           The unique name of this channel within the section
 	 * @param InDisplayText    Text to display on the sequencer node tree
 	 * @param InGroup          (Optional) When not empty, specifies a name to group channels by
+	 * @param bInEnabled        (Optional) When true the channel is enabled, if false it is not.
 	 */
-	MOVIESCENE_API FMovieSceneChannelMetaData(FName InName, FText InDisplayText, FText InGroup = FText());
+	MOVIESCENE_API FMovieSceneChannelMetaData(FName InName, FText InDisplayText, FText InGroup = FText(), bool bInEnabled = true);
 
 	/*
 	 * Set the identifiers for this editor data
diff --git a/Engine/Source/Runtime/MovieScene/Public/MovieSceneTrack.h b/Engine/Source/Runtime/MovieScene/Public/MovieSceneTrack.h
--- a/Engine/Source/Runtime/MovieScene/Public/MovieSceneTrack.h
+++ b/Engine/Source/Runtime/MovieScene/Public/MovieSceneTrack.h
@@ -250,6 +250,16 @@ public:
 		return SupportedBlendTypes.Num() != 0;
 	}
 
+	/** Set This Section as the one to key. If track doesn't support layered blends then don't implement
+	* @param InSection		Section that we want to key
+	*/
+	 virtual void SetSectionToKey(UMovieSceneSection* InSection) {};
+
+	/** Get the section we want to key. If track doesn't support layered blends it will return nulltpr.
+	* @return The section to key if one exists
+	*/
+	 virtual UMovieSceneSection* GetSectionToKey() const { return nullptr; }
+
 	/** Gets the greatest row index of all the sections owned by this track. */
 	MOVIESCENE_API int32 GetMaxRowIndex() const;
 
diff --git a/Engine/Source/Runtime/MovieSceneTracks/Private/Evaluation/MovieSceneParameterTemplate.cpp b/Engine/Source/Runtime/MovieSceneTracks/Private/Evaluation/MovieSceneParameterTemplate.cpp
--- a/Engine/Source/Runtime/MovieSceneTracks/Private/Evaluation/MovieSceneParameterTemplate.cpp
+++ b/Engine/Source/Runtime/MovieSceneTracks/Private/Evaluation/MovieSceneParameterTemplate.cpp
@@ -10,6 +10,7 @@ FMovieSceneParameterSectionTemplate::FMovieSceneParameterSectionTemplate(const U
 	: Scalars(Section.GetScalarParameterNamesAndCurves())
 	, Vectors(Section.GetVectorParameterNamesAndCurves())
 	, Colors(Section.GetColorParameterNamesAndCurves())
+	, Transforms(Section.GetTransformParameterNamesAndCurves())
 {
 }
 
@@ -56,6 +57,32 @@ void FMovieSceneParameterSectionTemplate::EvaluateCurves(const FMovieSceneContex
 			Values.ColorValues.Emplace(Color.ParameterName, ColorValue);
 		}
 	}
+
+	for (const FTransformParameterNameAndCurves& Transform : Transforms)
+	{
+		FVector Translation, Scale(FVector::OneVector);
+		FRotator Rotator;
+		bool bAnyEvaluated = false;
+		bAnyEvaluated |= Transform.Translation[0].Evaluate(Time,Translation[0]);
+		bAnyEvaluated |= Transform.Translation[1].Evaluate(Time, Translation[1]);
+		bAnyEvaluated |= Transform.Translation[2].Evaluate(Time, Translation[2]);
+
+		bAnyEvaluated |= Transform.Rotation[0].Evaluate(Time, Rotator.Roll);
+		bAnyEvaluated |= Transform.Rotation[1].Evaluate(Time, Rotator.Pitch);
+		bAnyEvaluated |= Transform.Rotation[2].Evaluate(Time, Rotator.Yaw);
+
+		//mz todo quat interp...
+
+		bAnyEvaluated |= Transform.Scale[0].Evaluate(Time, Scale[0]);
+		bAnyEvaluated |= Transform.Scale[1].Evaluate(Time, Scale[1]);
+		bAnyEvaluated |= Transform.Scale[2].Evaluate(Time, Scale[2]);
+
+		if (bAnyEvaluated)
+		{
+			FTransformParameterNameAndValue NameAndValue(Transform.ParameterName, Translation, Rotator, Scale);
+			Values.TransformValues.Emplace(NameAndValue);
+		}
+	}
 }
 
 void FDefaultMaterialAccessor::Apply(UMaterialInstanceDynamic& Material, const FEvaluatedParameterSectionValues& Values)
diff --git a/Engine/Source/Runtime/MovieSceneTracks/Private/Sections/MovieSceneParameterSection.cpp b/Engine/Source/Runtime/MovieSceneTracks/Private/Sections/MovieSceneParameterSection.cpp
--- a/Engine/Source/Runtime/MovieSceneTracks/Private/Sections/MovieSceneParameterSection.cpp
+++ b/Engine/Source/Runtime/MovieSceneTracks/Private/Sections/MovieSceneParameterSection.cpp
@@ -19,6 +19,11 @@ FColorParameterNameAndCurves::FColorParameterNameAndCurves( FName InParameterNam
 	ParameterName = InParameterName;
 }
 
+FTransformParameterNameAndCurves::FTransformParameterNameAndCurves(FName InParameterName)
+{
+	ParameterName = InParameterName;
+}
+
 UMovieSceneParameterSection::UMovieSceneParameterSection( const FObjectInitializer& ObjectInitializer )
 	: Super( ObjectInitializer )
 {
@@ -55,14 +60,14 @@ void UMovieSceneParameterSection::ReconstructChannelProxy()
 
 #if WITH_EDITOR
 
-	for ( FScalarParameterNameAndCurve& Scalar : GetScalarParameterNamesAndCurves() )
+	for (FScalarParameterNameAndCurve& Scalar : GetScalarParameterNamesAndCurves())
 	{
 		FMovieSceneChannelMetaData MetaData(Scalar.ParameterName, FText::FromName(Scalar.ParameterName));
 		// Prevent single channels from collapsing to the track node
 		MetaData.bCanCollapseToTrack = false;
 		Channels.Add(Scalar.ParameterCurve, MetaData, TMovieSceneExternalValue<float>());
 	}
-	for ( FVectorParameterNameAndCurves& Vector : GetVectorParameterNamesAndCurves() )
+	for (FVectorParameterNameAndCurves& Vector : GetVectorParameterNamesAndCurves())
 	{
 		FString ParameterString = Vector.ParameterName.ToString();
 		FText Group = FText::FromString(ParameterString);
@@ -71,7 +76,7 @@ void UMovieSceneParameterSection::ReconstructChannelProxy()
 		Channels.Add(Vector.YCurve, FMovieSceneChannelMetaData(*(ParameterString + TEXT(".Y")), FCommonChannelData::ChannelY, Group), TMovieSceneExternalValue<float>());
 		Channels.Add(Vector.ZCurve, FMovieSceneChannelMetaData(*(ParameterString + TEXT(".Z")), FCommonChannelData::ChannelZ, Group), TMovieSceneExternalValue<float>());
 	}
-	for ( FColorParameterNameAndCurves& Color : GetColorParameterNamesAndCurves() )
+	for (FColorParameterNameAndCurves& Color : GetColorParameterNamesAndCurves())
 	{
 		FString ParameterString = Color.ParameterName.ToString();
 		FText Group = FText::FromString(ParameterString);
@@ -91,25 +96,43 @@ void UMovieSceneParameterSection::ReconstructChannelProxy()
 		FMovieSceneChannelMetaData MetaData_A(*(ParameterString + TEXT("A")), FCommonChannelData::ChannelA, Group);
 		MetaData_A.SortOrder = 3;
 
-		Channels.Add(Color.RedCurve,   MetaData_R, TMovieSceneExternalValue<float>());
+		Channels.Add(Color.RedCurve, MetaData_R, TMovieSceneExternalValue<float>());
 		Channels.Add(Color.GreenCurve, MetaData_G, TMovieSceneExternalValue<float>());
-		Channels.Add(Color.BlueCurve,  MetaData_B, TMovieSceneExternalValue<float>());
+		Channels.Add(Color.BlueCurve, MetaData_B, TMovieSceneExternalValue<float>());
 		Channels.Add(Color.AlphaCurve, MetaData_A, TMovieSceneExternalValue<float>());
 	}
 
+	for (FTransformParameterNameAndCurves& Transform : GetTransformParameterNamesAndCurves())
+	{
+		FString ParameterString = Transform.ParameterName.ToString();
+		FText Group = FText::FromString(ParameterString);
+
+		Channels.Add(Transform.Translation[0], FMovieSceneChannelMetaData(*(ParameterString + TEXT(".Translation.X")), FCommonChannelData::ChannelX, Group), TMovieSceneExternalValue<float>());
+		Channels.Add(Transform.Translation[1], FMovieSceneChannelMetaData(*(ParameterString + TEXT(".Translation.Y")), FCommonChannelData::ChannelY, Group), TMovieSceneExternalValue<float>());
+		Channels.Add(Transform.Translation[2], FMovieSceneChannelMetaData(*(ParameterString + TEXT(".Translation.Z")), FCommonChannelData::ChannelZ, Group), TMovieSceneExternalValue<float>());
+
+		Channels.Add(Transform.Rotation[0], FMovieSceneChannelMetaData(*(ParameterString + TEXT(".Rotation.X")), FCommonChannelData::ChannelX, Group), TMovieSceneExternalValue<float>());
+		Channels.Add(Transform.Rotation[1], FMovieSceneChannelMetaData(*(ParameterString + TEXT(".Rotation.Y")), FCommonChannelData::ChannelY, Group), TMovieSceneExternalValue<float>());
+		Channels.Add(Transform.Rotation[2], FMovieSceneChannelMetaData(*(ParameterString + TEXT(".Rotation.Z")), FCommonChannelData::ChannelZ, Group), TMovieSceneExternalValue<float>());
+
+		Channels.Add(Transform.Scale[0], FMovieSceneChannelMetaData(*(ParameterString + TEXT(".Scale.X")), FCommonChannelData::ChannelX, Group), TMovieSceneExternalValue<float>());
+		Channels.Add(Transform.Scale[1], FMovieSceneChannelMetaData(*(ParameterString + TEXT(".Scale.Y")), FCommonChannelData::ChannelY, Group), TMovieSceneExternalValue<float>());
+		Channels.Add(Transform.Scale[2], FMovieSceneChannelMetaData(*(ParameterString + TEXT(".Scale.Z")), FCommonChannelData::ChannelZ, Group), TMovieSceneExternalValue<float>());
+
+	}
 #else
 
-	for ( FScalarParameterNameAndCurve& Scalar : GetScalarParameterNamesAndCurves() )
+	for (FScalarParameterNameAndCurve& Scalar : GetScalarParameterNamesAndCurves())
 	{
 		Channels.Add(Scalar.ParameterCurve);
 	}
-	for ( FVectorParameterNameAndCurves& Vector : GetVectorParameterNamesAndCurves() )
+	for (FVectorParameterNameAndCurves& Vector : GetVectorParameterNamesAndCurves())
 	{
 		Channels.Add(Vector.XCurve);
 		Channels.Add(Vector.YCurve);
 		Channels.Add(Vector.ZCurve);
 	}
-	for ( FColorParameterNameAndCurves& Color : GetColorParameterNamesAndCurves() )
+	for (FColorParameterNameAndCurves& Color : GetColorParameterNamesAndCurves())
 	{
 		Channels.Add(Color.RedCurve);
 		Channels.Add(Color.GreenCurve);
@@ -117,6 +140,22 @@ void UMovieSceneParameterSection::ReconstructChannelProxy()
 		Channels.Add(Color.AlphaCurve);
 	}
 
+	for (FTransformParameterNameAndCurves& Transform : GetTransformParameterNamesAndCurves())
+	{
+		Channels.Add(Transform.Translation[0]);
+		Channels.Add(Transform.Translation[1]);
+		Channels.Add(Transform.Translation[2]);
+
+		Channels.Add(Transform.Rotation[0]);
+		Channels.Add(Transform.Rotation[1]);
+		Channels.Add(Transform.Rotation[2]);
+
+		Channels.Add(Transform.Scale[0]);
+		Channels.Add(Transform.Scale[1]);
+		Channels.Add(Transform.Scale[2]);
+
+	}
+
 #endif
 
 	ChannelProxy = MakeShared<FMovieSceneChannelProxy>(MoveTemp(Channels));
@@ -208,6 +247,45 @@ void UMovieSceneParameterSection::AddColorParameterKey( FName InParameterName, F
 	}
 }
 
+void UMovieSceneParameterSection::AddTransformParameterKey(FName InParameterName, FFrameNumber InTime, const FTransform& InValue)
+{
+	FTransformParameterNameAndCurves* ExistingCurves = nullptr;
+	for (FTransformParameterNameAndCurves& TransformParameterNamesAndCurve : TransformParameterNamesAndCurves)
+	{
+		if (TransformParameterNamesAndCurve.ParameterName == InParameterName)
+		{
+			ExistingCurves = &TransformParameterNamesAndCurve;
+			break;
+		}
+	}
+	if (ExistingCurves == nullptr)
+	{
+		int32 NewIndex = TransformParameterNamesAndCurves.Add(FTransformParameterNameAndCurves(InParameterName));
+		ExistingCurves = &TransformParameterNamesAndCurves[NewIndex];
+
+		ReconstructChannelProxy();
+	}
+	FVector Translation = InValue.GetTranslation();
+	FRotator Rotator = InValue.GetRotation().Rotator();
+	FVector Scale = InValue.GetScale3D();
+	ExistingCurves->Translation[0].AddCubicKey(InTime, Translation[0]);
+	ExistingCurves->Translation[1].AddCubicKey(InTime, Translation[1]);
+	ExistingCurves->Translation[2].AddCubicKey(InTime, Translation[2]);
+
+	ExistingCurves->Rotation[0].AddCubicKey(InTime, Rotator.Roll);
+	ExistingCurves->Rotation[1].AddCubicKey(InTime, Rotator.Pitch);
+	ExistingCurves->Rotation[2].AddCubicKey(InTime, Rotator.Yaw);
+
+	ExistingCurves->Scale[0].AddCubicKey(InTime, Scale[0]);
+	ExistingCurves->Scale[1].AddCubicKey(InTime, Scale[1]);
+	ExistingCurves->Scale[2].AddCubicKey(InTime, Scale[2]);
+
+	if (TryModify())
+	{
+		SetRange(TRange<FFrameNumber>::Hull(TRange<FFrameNumber>(InTime), GetRange()));
+	}
+}
+
 bool UMovieSceneParameterSection::RemoveScalarParameter( FName InParameterName )
 {
 	for ( int32 i = 0; i < ScalarParameterNamesAndCurves.Num(); i++ )
@@ -250,6 +328,20 @@ bool UMovieSceneParameterSection::RemoveColorParameter( FName InParameterName )
 	return false;
 }
 
+bool UMovieSceneParameterSection::RemoveTransformParameter(FName InParameterName)
+{
+	for (int32 i = 0; i < TransformParameterNamesAndCurves.Num(); i++)
+	{
+		if (TransformParameterNamesAndCurves[i].ParameterName == InParameterName)
+		{
+			TransformParameterNamesAndCurves.RemoveAt(i);
+			ReconstructChannelProxy();
+			return true;
+		}
+	}
+	return false;
+}
+
 TArray<FScalarParameterNameAndCurve>& UMovieSceneParameterSection::GetScalarParameterNamesAndCurves()
 {
 	return ScalarParameterNamesAndCurves;
@@ -280,6 +372,16 @@ const TArray<FColorParameterNameAndCurves>& UMovieSceneParameterSection::GetColo
 	return ColorParameterNamesAndCurves;
 }
 
+TArray<FTransformParameterNameAndCurves>& UMovieSceneParameterSection::GetTransformParameterNamesAndCurves() 
+{
+	return TransformParameterNamesAndCurves;
+}
+
+const TArray<FTransformParameterNameAndCurves>& UMovieSceneParameterSection::GetTransformParameterNamesAndCurves() const
+{
+	return TransformParameterNamesAndCurves;
+}
+
 void UMovieSceneParameterSection::GetParameterNames( TSet<FName>& ParameterNames ) const
 {
 	for ( const FScalarParameterNameAndCurve& ScalarParameterNameAndCurve : ScalarParameterNamesAndCurves )
@@ -294,4 +396,8 @@ void UMovieSceneParameterSection::GetParameterNames( TSet<FName>& ParameterNames
 	{
 		ParameterNames.Add( ColorParameterNameAndCurves.ParameterName );
 	}
+	for (const FTransformParameterNameAndCurves& TransformParameterNamesAndCurve : TransformParameterNamesAndCurves)
+	{
+		ParameterNames.Add(TransformParameterNamesAndCurve.ParameterName);
+	}
 }
\ No newline at end of file
diff --git a/Engine/Source/Runtime/MovieSceneTracks/Private/Tracks/MovieScenePropertyTrack.cpp b/Engine/Source/Runtime/MovieSceneTracks/Private/Tracks/MovieScenePropertyTrack.cpp
--- a/Engine/Source/Runtime/MovieSceneTracks/Private/Tracks/MovieScenePropertyTrack.cpp
+++ b/Engine/Source/Runtime/MovieSceneTracks/Private/Tracks/MovieScenePropertyTrack.cpp
@@ -275,7 +275,7 @@ void UMovieScenePropertyTrack::SetSectionToKey(UMovieSceneSection* InSection)
 	SectionToKey = InSection;
 }
 
-UMovieSceneSection* UMovieScenePropertyTrack::GetSectionToKey()
+UMovieSceneSection* UMovieScenePropertyTrack::GetSectionToKey() const
 {
 	return SectionToKey;
 }
diff --git a/Engine/Source/Runtime/MovieSceneTracks/Public/Evaluation/MovieSceneParameterTemplate.h b/Engine/Source/Runtime/MovieSceneTracks/Public/Evaluation/MovieSceneParameterTemplate.h
--- a/Engine/Source/Runtime/MovieSceneTracks/Public/Evaluation/MovieSceneParameterTemplate.h
+++ b/Engine/Source/Runtime/MovieSceneTracks/Public/Evaluation/MovieSceneParameterTemplate.h
@@ -36,6 +36,8 @@ struct FEvaluatedParameterSectionValues
 	TArray<FVectorParameterNameAndValue, TInlineAllocator<2>> VectorValues;
 	/** Array of evaluated color values */
 	TArray<FColorParameterNameAndValue, TInlineAllocator<2>> ColorValues;
+	/** Array of evaluated transform values */
+	TArray<FTransformParameterNameAndValue, TInlineAllocator<2>> TransformValues;
 };
 
 /** Template that performs evaluation of parameter sections */
@@ -54,7 +56,7 @@ protected:
 	/** Evaluate our curves, outputting evaluated values into the specified container */
 	MOVIESCENETRACKS_API void EvaluateCurves(const FMovieSceneContext& Context, FEvaluatedParameterSectionValues& OutValues) const;
 
-private:
+protected:
 
 	/** The scalar parameter names and their associated curves. */
 	UPROPERTY()
@@ -67,6 +69,9 @@ private:
 	/** The color parameter names and their associated curves. */
 	UPROPERTY()
 	TArray<FColorParameterNameAndCurves> Colors;
+
+	UPROPERTY()
+	TArray<FTransformParameterNameAndCurves> Transforms;
 };
 
 /** Default accessor type for use with TMaterialTrackExecutionToken */
diff --git a/Engine/Source/Runtime/MovieSceneTracks/Public/Sections/MovieSceneParameterSection.h b/Engine/Source/Runtime/MovieSceneTracks/Public/Sections/MovieSceneParameterSection.h
--- a/Engine/Source/Runtime/MovieSceneTracks/Public/Sections/MovieSceneParameterSection.h
+++ b/Engine/Source/Runtime/MovieSceneTracks/Public/Sections/MovieSceneParameterSection.h
@@ -8,6 +8,7 @@
 #include "Curves/RichCurve.h"
 #include "MovieSceneSection.h"
 #include "Channels/MovieSceneFloatChannel.h"
+#include "Sections/MovieScene3DTransformSection.h"
 #include "MovieSceneParameterSection.generated.h"
 
 struct FMovieSceneParameterPropertyInterface;
@@ -32,7 +33,6 @@ struct FScalarParameterNameAndValue
 	float Value;
 };
 
-
 /**
  * Structure representing the animated value of a vector parameter.
  */
@@ -72,12 +72,31 @@ struct FColorParameterNameAndValue
 	FLinearColor Value;
 };
 
+struct FTransformParameterNameAndValue
+{
+
+	/** The name of the transform  parameter. */
+	FName ParameterName;
+	/** Translation component */
+	FVector Translation;
+	/** Rotation component */
+	FRotator Rotation;
+	/** Scale component */
+	FVector Scale;
+
+	FTransformParameterNameAndValue(FName InParameterName,const FVector& InTranslation,
+		const FRotator& InRotation, const FVector& InScale) : Translation(InTranslation),
+		Rotation(InRotation), Scale(InScale)
+	{
+		ParameterName = InParameterName;	
+	}
+};
 
 /**
  * Structure representing an animated scalar parameter and it's associated animation curve.
  */
 USTRUCT()
-struct FScalarParameterNameAndCurve
+struct MOVIESCENETRACKS_API FScalarParameterNameAndCurve
 {
 	GENERATED_USTRUCT_BODY()
 
@@ -101,7 +120,7 @@ struct FScalarParameterNameAndCurve
  * Structure representing an animated vector parameter and it's associated animation curve.
  */
 USTRUCT()
-struct FVectorParameterNameAndCurves
+struct MOVIESCENETRACKS_API FVectorParameterNameAndCurves
 {
 	GENERATED_USTRUCT_BODY()
 
@@ -130,10 +149,10 @@ struct FVectorParameterNameAndCurves
 
 
 /**
-* Structure representing an animated vector parameter and it's associated animation curve.
+* Structure representing an animated color parameter and it's associated animation curve.
 */
 USTRUCT()
-struct FColorParameterNameAndCurves
+struct MOVIESCENETRACKS_API FColorParameterNameAndCurves
 {
 	GENERATED_USTRUCT_BODY()
 
@@ -143,7 +162,7 @@ struct FColorParameterNameAndCurves
 	/** Creates a new FVectorParameterNameAndCurve for a specific color parameter. */
 	FColorParameterNameAndCurves(FName InParameterName);
 
-	/** The name of the vector parameter which is being animated. */
+	/** The name of the color parameter which is being animated. */
 	UPROPERTY()
 	FName ParameterName;
 
@@ -164,25 +183,57 @@ struct FColorParameterNameAndCurves
 	FMovieSceneFloatChannel AlphaCurve;
 };
 
+/**
+* Structure representing an animated transform parameter and it's associated animation curve.
+*/
+USTRUCT()
+struct MOVIESCENETRACKS_API FTransformParameterNameAndCurves
+{
+	GENERATED_USTRUCT_BODY()
+
+	FTransformParameterNameAndCurves()
+	{}
+
+	/** Creates a new FVectorParameterNameAndCurve for a specific color parameter. */
+	FTransformParameterNameAndCurves(FName InParameterName);
+
+	/** The name of the transform  parameter which is being animated. */
+	UPROPERTY()
+	FName ParameterName;
 
+	/** Translation curves */
+	UPROPERTY()
+	FMovieSceneFloatChannel Translation[3];
+
+	/** Rotation curves */
+	UPROPERTY()
+	FMovieSceneFloatChannel Rotation[3];
+
+	/** Scale curves */
+	UPROPERTY()
+	FMovieSceneFloatChannel Scale[3];
+};
 /**
  * A single movie scene section which can contain data for multiple named parameters.
  */
-UCLASS(MinimalAPI)
-class UMovieSceneParameterSection
+UCLASS()
+class MOVIESCENETRACKS_API UMovieSceneParameterSection
 	: public UMovieSceneSection
 {
 	GENERATED_UCLASS_BODY()
 
 public:
 	/** Adds a a key for a specific scalar parameter at the specified time with the specified value. */
-	MOVIESCENETRACKS_API void AddScalarParameterKey(FName InParameterName, FFrameNumber InTime, float InValue);
+	 void AddScalarParameterKey(FName InParameterName, FFrameNumber InTime, float InValue);
 
 	/** Adds a a key for a specific vector parameter at the specified time with the specified value. */
-	MOVIESCENETRACKS_API void AddVectorParameterKey(FName InParameterName, FFrameNumber InTime, FVector InValue);
+	 void AddVectorParameterKey(FName InParameterName, FFrameNumber InTime, FVector InValue);
+
+	/** Adds a a key for a specific color parameter at the specified time with the specified value. */
+	 void AddColorParameterKey(FName InParameterName, FFrameNumber InTime, FLinearColor InValue);
 
 	/** Adds a a key for a specific color parameter at the specified time with the specified value. */
-	MOVIESCENETRACKS_API void AddColorParameterKey(FName InParameterName, FFrameNumber InTime, FLinearColor InValue);
+	 void AddTransformParameterKey(FName InParameterName, FFrameNumber InTime, const FTransform& InValue);
 
 	/** 
 	 * Removes a scalar parameter from this section. 
@@ -190,7 +241,7 @@ public:
 	 * @param InParameterName The name of the scalar parameter to remove.
 	 * @returns True if a parameter with that name was found and removed, otherwise false.
 	 */
-	MOVIESCENETRACKS_API bool RemoveScalarParameter(FName InParameterName);
+	 bool RemoveScalarParameter(FName InParameterName);
 
 	/**
 	 * Removes a vector parameter from this section.
@@ -198,7 +249,7 @@ public:
 	 * @param InParameterName The name of the vector parameter to remove.
 	 * @returns True if a parameter with that name was found and removed, otherwise false.
 	 */
-	MOVIESCENETRACKS_API bool RemoveVectorParameter(FName InParameterName);
+	 bool RemoveVectorParameter(FName InParameterName);
 
 	/**
 	 * Removes a color parameter from this section.
@@ -206,32 +257,43 @@ public:
 	 * @param InParameterName The name of the color parameter to remove.
 	 * @returns True if a parameter with that name was found and removed, otherwise false.
 	 */
-	MOVIESCENETRACKS_API bool RemoveColorParameter(FName InParameterName);
+	 bool RemoveColorParameter(FName InParameterName);
+
+	/**
+	 * Removes a transform parameter from this section.
+	 *
+	 * @param InParameterName The name of the transform parameter to remove.
+	 * @returns True if a parameter with that name was found and removed, otherwise false.
+	 */
+	 bool RemoveTransformParameter(FName InParameterName);
 
 	/** Gets the animated scalar parameters and their associated curves. */
-	MOVIESCENETRACKS_API TArray<FScalarParameterNameAndCurve>& GetScalarParameterNamesAndCurves();
-	MOVIESCENETRACKS_API const TArray<FScalarParameterNameAndCurve>& GetScalarParameterNamesAndCurves() const;
+	 TArray<FScalarParameterNameAndCurve>& GetScalarParameterNamesAndCurves();
+	 const TArray<FScalarParameterNameAndCurve>& GetScalarParameterNamesAndCurves() const;
 
 	/** Gets the animated vector parameters and their associated curves. */
-	MOVIESCENETRACKS_API TArray<FVectorParameterNameAndCurves>& GetVectorParameterNamesAndCurves();
-	MOVIESCENETRACKS_API const TArray<FVectorParameterNameAndCurves>& GetVectorParameterNamesAndCurves() const;
+	 TArray<FVectorParameterNameAndCurves>& GetVectorParameterNamesAndCurves();
+	 const TArray<FVectorParameterNameAndCurves>& GetVectorParameterNamesAndCurves() const;
 
 	/** Gets the animated color parameters and their associated curves. */
-	MOVIESCENETRACKS_API TArray<FColorParameterNameAndCurves>& GetColorParameterNamesAndCurves();
-	MOVIESCENETRACKS_API const TArray<FColorParameterNameAndCurves>& GetColorParameterNamesAndCurves() const;
+	 TArray<FColorParameterNameAndCurves>& GetColorParameterNamesAndCurves();
+	 const TArray<FColorParameterNameAndCurves>& GetColorParameterNamesAndCurves() const;
+
+	/** Gets the animated transform parameters and their associated curves. */
+	 TArray<FTransformParameterNameAndCurves>& GetTransformParameterNamesAndCurves();
+     const TArray<FTransformParameterNameAndCurves>& GetTransformParameterNamesAndCurves() const;
 
 	/** Gets the set of all parameter names used by this section. */
-	MOVIESCENETRACKS_API void GetParameterNames(TSet<FName>& ParameterNames) const;
+	void GetParameterNames(TSet<FName>& ParameterNames) const;
 
 protected:
 
 	//~ UMovieSceneSection interface
 	virtual void Serialize(FArchive& Ar) override;
 	virtual void PostEditImport() override;
+	virtual void ReconstructChannelProxy();
 
-	void ReconstructChannelProxy();
-
-private:
+protected:
 
 	/** The scalar parameter names and their associated curves. */
 	UPROPERTY()
@@ -244,4 +306,8 @@ private:
 	/** The vector parameter names and their associated curves. */
 	UPROPERTY()
 	TArray<FColorParameterNameAndCurves> ColorParameterNamesAndCurves;
+
+	/** The vector parameter names and their associated curves. */
+	UPROPERTY()
+	TArray<FTransformParameterNameAndCurves> TransformParameterNamesAndCurves;
 };
\ No newline at end of file
diff --git a/Engine/Source/Runtime/MovieSceneTracks/Public/Tracks/MovieScenePropertyTrack.h b/Engine/Source/Runtime/MovieSceneTracks/Public/Tracks/MovieScenePropertyTrack.h
--- a/Engine/Source/Runtime/MovieSceneTracks/Public/Tracks/MovieScenePropertyTrack.h
+++ b/Engine/Source/Runtime/MovieSceneTracks/Public/Tracks/MovieScenePropertyTrack.h
@@ -90,13 +90,13 @@ public:
 	 *
 	 * @param Section The section that changes.
 	 */
-	void SetSectionToKey(UMovieSceneSection* Section);
+	virtual void SetSectionToKey(UMovieSceneSection* Section) override;
 
 	/**
 	 * Finds a section we want to key and recieve globally changed values.
 	 * @return The Section that changes.
 	 */
-	class UMovieSceneSection* GetSectionToKey();
+	virtual UMovieSceneSection* GetSectionToKey() const override;
 #if WITH_EDITORONLY_DATA
 public:
 	/** Unique name for this track to afford multiple tracks on a given object (i.e. for array properties) */
diff --git a/Engine/Source/Runtime/SlateCore/Private/Rendering/DrawElements.cpp b/Engine/Source/Runtime/SlateCore/Private/Rendering/DrawElements.cpp
--- a/Engine/Source/Runtime/SlateCore/Private/Rendering/DrawElements.cpp
+++ b/Engine/Source/Runtime/SlateCore/Private/Rendering/DrawElements.cpp
@@ -554,7 +554,8 @@ void FSlateDrawElement::MakeCustomVerts(FSlateWindowElementList& ElementList, ui
 	FSlateCustomVertsPayload& DataPayload = ElementList.CreatePayload<FSlateCustomVertsPayload>(Element);
 
 	const FSlateShaderResourceProxy* RenderingProxy = InRenderResourceHandle.GetResourceProxy();
-	DataPayload.SetCustomVerts(RenderingProxy, InVerts, InIndexes, InInstanceData->GetRenderProxy(), InInstanceOffset, InNumInstances);
+	ISlateUpdatableInstanceBufferRenderProxy* RenderProxy = InInstanceData ? InInstanceData->GetRenderProxy() : nullptr;
+	DataPayload.SetCustomVerts(RenderingProxy, InVerts, InIndexes, RenderProxy, InInstanceOffset, InNumInstances);
 
 	Element.Init(ElementList, EElementType::ET_CustomVerts, InLayer, FPaintGeometry(), InDrawEffects);
 	Element.RenderTransform = FSlateRenderTransform();
