Sequencer: Fix event evaluation error output

diff --git a/Engine/Source/Runtime/MovieSceneTracks/Private/Evaluation/MovieSceneEventTemplate.cpp b/Engine/Source/Runtime/MovieSceneTracks/Private/Evaluation/MovieSceneEventTemplate.cpp
--- a/Engine/Source/Runtime/MovieSceneTracks/Private/Evaluation/MovieSceneEventTemplate.cpp
+++ b/Engine/Source/Runtime/MovieSceneTracks/Private/Evaluation/MovieSceneEventTemplate.cpp
@@ -337,7 +337,7 @@ struct FEventTriggerExecutionToken
 					->AddToken(FUObjectToken::Create(Player.GetEvaluationTemplate().GetSequence(SequenceID)))
 					->AddToken(FUObjectToken::Create(BoundObjectProperty->GetOuter()))
 					->AddToken(FUObjectToken::Create(BoundObject))
-					->AddToken(FTextToken::Create(FText::Format(LOCTEXT("LevelBP_InvalidCast_Error", "Failed to trigger event: Cast to %s failed."), FText::FromName(ObjectParameter->PropertyClass->GetFName()))));
+					->AddToken(FTextToken::Create(FText::Format(LOCTEXT("LevelBP_InvalidCast_Error", "Failed to trigger event: Cast to {0} failed."), FText::FromName(ObjectParameter->PropertyClass->GetFName()))));
 
 				return false;
 			}
