Sequencer: Fixed crash caused by events invoking functions that had non CPF_Param properties in their property list

We were only allocating enough space for ParmsSize, byt initializing (and destroying) PropertiesSize worth of properties which may be larger. We now manually initialize properties marked CPF_Param instead of the whole struct to avoid needlessly allocating and initializing local variables that will not be used here.

diff --git a/Engine/Source/Runtime/MovieSceneTracks/Private/Evaluation/MovieSceneEventTemplate.cpp b/Engine/Source/Runtime/MovieSceneTracks/Private/Evaluation/MovieSceneEventTemplate.cpp
--- a/Engine/Source/Runtime/MovieSceneTracks/Private/Evaluation/MovieSceneEventTemplate.cpp
+++ b/Engine/Source/Runtime/MovieSceneTracks/Private/Evaluation/MovieSceneEventTemplate.cpp
@@ -270,7 +270,19 @@ struct FEventTriggerExecutionToken
 		uint8* Parameters = (uint8*)FMemory_Alloca(Event.Function->ParmsSize + Event.Function->MinAlignment);
 		Parameters = Align(Parameters, Event.Function->MinAlignment);
 
-		Event.Function->InitializeStruct(Parameters);
+		// Mem zero the parameter list
+		FMemory::Memzero(Parameters, Event.Function->ParmsSize);
+
+		// Initialize all CPF_Param properties - these are aways at the head of the list
+		for (TFieldIterator<UProperty> It(Event.Function); It && It->HasAnyPropertyFlags(CPF_Parm); ++It)
+		{
+			UProperty* LocalProp = *It;
+			checkSlow(LocalProp);
+			if (!LocalProp->HasAnyPropertyFlags(CPF_ZeroConstructor))
+			{
+				LocalProp->InitializeValue_InContainer(Parameters);
+			}
+		}
 
 		for (UObject* BoundObject : EventContexts)
 		{
@@ -292,7 +304,11 @@ struct FEventTriggerExecutionToken
 			DirectorInstance->ProcessEvent(Event.Function, Parameters);
 		}
 
-		Event.Function->DestroyStruct(Parameters);
+		// Destroy all parameter properties one by one
+		for (TFieldIterator<UProperty> It(Event.Function); It && It->HasAnyPropertyFlags(CPF_Parm); ++It)
+		{
+			It->DestroyValue_InContainer(Parameters);
+		}
 	}
 
 	bool PatchBoundObject(uint8* Parameters, UObject* BoundObject, UProperty* BoundObjectProperty, IMovieScenePlayer& Player, FMovieSceneSequenceID SequenceID)
