Sequencer: Fixed overriden spawnables getting destroyed when the spawn track executes Restore State

The call to DestroyObjectDirectly was the body of the preanimated token in its original form - before binding overrides were even a thing, and before we cached the object binding ID and sequence ID in the token. We must still cache pre-animated state for all spawnable objects so that manipulating a spawn track does not cause the object to be destroyed and re-spawned (since the spawn track will get a new template). There doesn't appear to be any way that a spawnable object could exist without an associated pre animated token for the operand, so I think it is safe to remove the call to DestroyObjectDirectly and just allow the spawn register to do the work.

diff --git a/Engine/Source/Runtime/MovieSceneTracks/Private/Evaluation/MovieSceneSpawnTemplate.cpp b/Engine/Source/Runtime/MovieSceneTracks/Private/Evaluation/MovieSceneSpawnTemplate.cpp
--- a/Engine/Source/Runtime/MovieSceneTracks/Private/Evaluation/MovieSceneSpawnTemplate.cpp
+++ b/Engine/Source/Runtime/MovieSceneTracks/Private/Evaluation/MovieSceneSpawnTemplate.cpp
@@ -25,11 +25,7 @@ struct FSpawnTrackPreAnimatedTokenProducer : IMovieScenePreAnimatedTokenProducer
 
 			virtual void RestoreState(UObject& InObject, IMovieScenePlayer& Player) override
 			{
-				bool bDestroyed = Player.GetSpawnRegister().DestroySpawnedObject(OperandToDestroy.ObjectBindingID, OperandToDestroy.SequenceID, Player);
-				if (!bDestroyed)
-				{
-					Player.GetSpawnRegister().DestroyObjectDirectly(InObject);
-				}
+				Player.GetSpawnRegister().DestroySpawnedObject(OperandToDestroy.ObjectBindingID, OperandToDestroy.SequenceID, Player);
 			}
 		};
 		
