Sequencer: Swap identifiers for camera cut section functions. GetCameraBindingID should be BlueprintPure and SetCameraBindingID should be BlueprintCallable.

diff --git a/Engine/Source/Runtime/MovieSceneTracks/Public/Sections/MovieSceneCameraCutSection.h b/Engine/Source/Runtime/MovieSceneTracks/Public/Sections/MovieSceneCameraCutSection.h
--- a/Engine/Source/Runtime/MovieSceneTracks/Public/Sections/MovieSceneCameraCutSection.h
+++ b/Engine/Source/Runtime/MovieSceneTracks/Public/Sections/MovieSceneCameraCutSection.h
@@ -41,14 +41,14 @@ public:
 	}
 
 	/** Gets the camera binding for this CameraCut section */
-	UFUNCTION(BlueprintCallable, Category = "Movie Scene Section")
+	UFUNCTION(BlueprintPure, Category = "Movie Scene Section")
 	const FMovieSceneObjectBindingID& GetCameraBindingID() const
 	{
 		return CameraBindingID;
 	}
 
 	/** Sets the camera binding for this CameraCut section */
-	UFUNCTION(BlueprintPure, Category = "Movie Scene Section")
+	UFUNCTION(BlueprintCallable, Category = "Movie Scene Section")
 	void SetCameraBindingID(const FMovieSceneObjectBindingID& InCameraBindingID)
 	{
 		CameraBindingID = InCameraBindingID;
