Sequencer: Use app delta time instead of sequencer delta time in computing the velocity. Sequencer delta time has dilation factored in, which means the computed velocity will be the same if the sequence is played at a different rate.

Tested with different combinations of sequencer play rate - using a time dilation track, setting the play rate on the level sequence, setting a play rate on a sub level sequence.

diff --git a/Engine/Source/Runtime/MovieSceneTracks/Private/Evaluation/MovieScene3DTransformTemplate.cpp b/Engine/Source/Runtime/MovieSceneTracks/Private/Evaluation/MovieScene3DTransformTemplate.cpp
--- a/Engine/Source/Runtime/MovieSceneTracks/Private/Evaluation/MovieScene3DTransformTemplate.cpp
+++ b/Engine/Source/Runtime/MovieSceneTracks/Private/Evaluation/MovieScene3DTransformTemplate.cpp
@@ -94,7 +94,7 @@ struct FComponentTransformActuator : TMovieSceneBlendingActuator<F3DTransformTra
 
 			SceneComponent->SetMobility(EComponentMobility::Movable);
 
-			InFinalValue.Apply(*SceneComponent, Context.GetDelta() / Context.GetFrameRate());
+			InFinalValue.Apply(*SceneComponent);
 		}
 	}
 
diff --git a/Engine/Source/Runtime/MovieSceneTracks/Private/Evaluation/MovieSceneTemplateCommon.cpp b/Engine/Source/Runtime/MovieSceneTracks/Private/Evaluation/MovieSceneTemplateCommon.cpp
--- a/Engine/Source/Runtime/MovieSceneTracks/Private/Evaluation/MovieSceneTemplateCommon.cpp
+++ b/Engine/Source/Runtime/MovieSceneTracks/Private/Evaluation/MovieSceneTemplateCommon.cpp
@@ -3,6 +3,7 @@
 #include "Evaluation/MovieSceneTemplateCommon.h"
 #include "Components/SceneComponent.h"
 #include "GameFramework/Actor.h"
+#include "Misc/App.h"
 
 /** A movie scene pre-animated token that stores a pre-animated mobility */
 struct FMobilityPreAnimatedToken : IMovieScenePreAnimatedToken
@@ -34,8 +35,7 @@ IMovieScenePreAnimatedTokenPtr FMobilityTokenProducer::CacheExistingState(UObjec
 	return FMobilityPreAnimatedToken(*SceneComponent);
 }
 
-
-void F3DTransformTrackToken::Apply(USceneComponent& SceneComponent, float DeltaTime) const
+void F3DTransformTrackToken::Apply(USceneComponent& SceneComponent) const
 {
 	/* Cache initial absolute position */
 	FVector PreviousPosition = SceneComponent.GetComponentLocation();
@@ -53,7 +53,8 @@ void F3DTransformTrackToken::Apply(USceneComponent& SceneComponent, float DeltaT
 	SceneComponent.SetRelativeLocation_Direct(Translation);
 	SceneComponent.SetRelativeRotation_Direct(Rotation);
 
-	if (DeltaTime > 0.f)
+	double DeltaTime = FApp::GetDeltaTime();
+	if (DeltaTime > 0)
 	{
 		/* Get current absolute position and set component velocity */
 		FVector CurrentPosition = SceneComponent.GetComponentLocation();
diff --git a/Engine/Source/Runtime/MovieSceneTracks/Private/Evaluation/MovieSceneTemplateCommon.h b/Engine/Source/Runtime/MovieSceneTracks/Private/Evaluation/MovieSceneTemplateCommon.h
--- a/Engine/Source/Runtime/MovieSceneTracks/Private/Evaluation/MovieSceneTemplateCommon.h
+++ b/Engine/Source/Runtime/MovieSceneTracks/Private/Evaluation/MovieSceneTemplateCommon.h
@@ -32,7 +32,7 @@ struct F3DTransformTrackToken
 	FRotator Rotation;
 	FVector Scale;
 
-	void Apply(USceneComponent& SceneComponent, float DeltaTime = 0.f) const;
+	void Apply(USceneComponent& SceneComponent) const;
 };
 
 struct F3DTransformTokenProducer : IMovieScenePreAnimatedTokenProducer
