Sequencer: Fixed error in audio section looping calculation

  - CL 8675950 introduced code that took duration as ms, whereas it is actually in seconds

diff --git a/Engine/Source/Runtime/MovieSceneTracks/Private/Evaluation/MovieSceneAudioTemplate.cpp b/Engine/Source/Runtime/MovieSceneTracks/Private/Evaluation/MovieSceneAudioTemplate.cpp
--- a/Engine/Source/Runtime/MovieSceneTracks/Private/Evaluation/MovieSceneAudioTemplate.cpp
+++ b/Engine/Source/Runtime/MovieSceneTracks/Private/Evaluation/MovieSceneAudioTemplate.cpp
@@ -366,7 +366,7 @@ struct FAudioSectionExecutionToken : IMovieSceneExecutionToken
 			if (AudioTime >= 0.f && AudioComponent.Sound)
 			{
 				float Duration = AudioComponent.Sound->GetDuration();
-				AudioTime = FMath::Fmod(AudioTime, Duration/1000.f);
+				AudioTime = FMath::Fmod(AudioTime, Duration);
 				AudioComponent.Play(AudioTime);
 			}
 
