Sequencer: Set material and particle range to playback range when creating a new section.

diff --git a/Engine/Source/Runtime/MovieSceneTracks/Private/Tracks/MovieSceneMaterialTrack.cpp b/Engine/Source/Runtime/MovieSceneTracks/Private/Tracks/MovieSceneMaterialTrack.cpp
--- a/Engine/Source/Runtime/MovieSceneTracks/Private/Tracks/MovieSceneMaterialTrack.cpp
+++ b/Engine/Source/Runtime/MovieSceneTracks/Private/Tracks/MovieSceneMaterialTrack.cpp
@@ -72,7 +72,11 @@ void UMovieSceneMaterialTrack::AddScalarParameterKey(FName ParameterName, FFrame
 	if (NearestSection == nullptr)
 	{
 		NearestSection = Cast<UMovieSceneParameterSection>(CreateNewSection());
-		NearestSection->SetRange(TRange<FFrameNumber>::Inclusive(Time, Time));
+
+		UMovieScene* MovieScene = GetTypedOuter<UMovieScene>();
+		check(MovieScene);
+
+		NearestSection->SetRange(MovieScene->GetPlaybackRange());
 		Sections.Add(NearestSection);
 	}
 	if (NearestSection->TryModify())
@@ -88,7 +92,12 @@ void UMovieSceneMaterialTrack::AddColorParameterKey(FName ParameterName, FFrameN
 	if (NearestSection == nullptr)
 	{
 		NearestSection = Cast<UMovieSceneParameterSection>(CreateNewSection());
-		NearestSection->SetRange(TRange<FFrameNumber>::Inclusive(Time, Time));
+
+		UMovieScene* MovieScene = GetTypedOuter<UMovieScene>();
+		check(MovieScene);
+
+		NearestSection->SetRange(MovieScene->GetPlaybackRange());
+
 		Sections.Add(NearestSection);
 	}
 	if (NearestSection->TryModify())
diff --git a/Engine/Source/Runtime/MovieSceneTracks/Private/Tracks/MovieSceneParticleParameterTrack.cpp b/Engine/Source/Runtime/MovieSceneTracks/Private/Tracks/MovieSceneParticleParameterTrack.cpp
--- a/Engine/Source/Runtime/MovieSceneTracks/Private/Tracks/MovieSceneParticleParameterTrack.cpp
+++ b/Engine/Source/Runtime/MovieSceneTracks/Private/Tracks/MovieSceneParticleParameterTrack.cpp
@@ -79,7 +79,12 @@ void UMovieSceneParticleParameterTrack::AddScalarParameterKey( FName ParameterNa
 	if ( NearestSection == nullptr )
 	{
 		NearestSection = Cast<UMovieSceneParameterSection>(CreateNewSection());
-		NearestSection->SetRange(TRange<FFrameNumber>::Inclusive(Time, Time));
+
+		UMovieScene* MovieScene = GetTypedOuter<UMovieScene>();
+		check(MovieScene);
+
+		NearestSection->SetRange(MovieScene->GetPlaybackRange());
+
 		Sections.Add( NearestSection );
 	}
 	NearestSection->AddScalarParameterKey(ParameterName, Time, Value);
@@ -91,7 +96,12 @@ void UMovieSceneParticleParameterTrack::AddVectorParameterKey( FName ParameterNa
 	if ( NearestSection == nullptr )
 	{
 		NearestSection = Cast<UMovieSceneParameterSection>( CreateNewSection() );
-		NearestSection->SetRange(TRange<FFrameNumber>::Inclusive(Time, Time));
+
+		UMovieScene* MovieScene = GetTypedOuter<UMovieScene>();
+		check(MovieScene);
+
+		NearestSection->SetRange(MovieScene->GetPlaybackRange());
+
 		Sections.Add( NearestSection );
 	}
 	NearestSection->AddVectorParameterKey( ParameterName, Time, Value );
@@ -103,7 +113,12 @@ void UMovieSceneParticleParameterTrack::AddColorParameterKey( FName ParameterNam
 	if ( NearestSection == nullptr )
 	{
 		NearestSection = Cast<UMovieSceneParameterSection>( CreateNewSection() );
-		NearestSection->SetRange(TRange<FFrameNumber>::Inclusive(Time, Time));
+
+		UMovieScene* MovieScene = GetTypedOuter<UMovieScene>();
+		check(MovieScene);
+
+		NearestSection->SetRange(MovieScene->GetPlaybackRange());
+
 		Sections.Add( NearestSection );
 	}
 	NearestSection->AddColorParameterKey( ParameterName, Time, Value );
