Fix some bugs with sequencer anim-nodes and anim-notifies not firing correctly in various cases.

- Provide the previous/next times to the anim-node so it can make the animation system update that whole time range, including when the node is resetting.
- Do the same thing for the anim-montage update, because when time loops around back to 0, it was updating backwards and firing notifies incorrectly.

Disable anim-notifies when scrubbing because we currently can't make all cases work consistently.

There are some other edge cases that are not solved by this CL yet, though.

diff --git a/Engine/Source/Runtime/AnimGraphRuntime/Private/AnimSequencerInstance.cpp b/Engine/Source/Runtime/AnimGraphRuntime/Private/AnimSequencerInstance.cpp
--- a/Engine/Source/Runtime/AnimGraphRuntime/Private/AnimSequencerInstance.cpp
+++ b/Engine/Source/Runtime/AnimGraphRuntime/Private/AnimSequencerInstance.cpp
@@ -30,6 +30,11 @@ void UAnimSequencerInstance::UpdateAnimTrack(UAnimSequenceBase* InAnimSequence,
 	GetProxyOnGameThread<FAnimSequencerInstanceProxy>().UpdateAnimTrack(InAnimSequence, SequenceId, InPosition, Weight, bFireNotifies);
 }
 
+void UAnimSequencerInstance::UpdateAnimTrack(UAnimSequenceBase* InAnimSequence, int32 SequenceId, float InFromPosition, float InToPosition, float Weight, bool bFireNotifies)
+{
+	GetProxyOnGameThread<FAnimSequencerInstanceProxy>().UpdateAnimTrack(InAnimSequence, SequenceId, InFromPosition, InToPosition, Weight, bFireNotifies);
+}
+
 void UAnimSequencerInstance::ResetNodes()
 {
 	GetProxyOnGameThread<FAnimSequencerInstanceProxy>().ResetNodes();
diff --git a/Engine/Source/Runtime/AnimGraphRuntime/Private/AnimSequencerInstanceProxy.cpp b/Engine/Source/Runtime/AnimGraphRuntime/Private/AnimSequencerInstanceProxy.cpp
--- a/Engine/Source/Runtime/AnimGraphRuntime/Private/AnimSequencerInstanceProxy.cpp
+++ b/Engine/Source/Runtime/AnimGraphRuntime/Private/AnimSequencerInstanceProxy.cpp
@@ -143,11 +143,24 @@ void FAnimSequencerInstanceProxy::TermAnimTrack(int32 SequenceId)
 }*/
 
 void FAnimSequencerInstanceProxy::UpdateAnimTrack(UAnimSequenceBase* InAnimSequence, uint32 SequenceId, float InPosition, float Weight, bool bFireNotifies)
+{
+	UpdateAnimTrack(InAnimSequence, SequenceId, TOptional<float>(), InPosition, Weight, bFireNotifies);
+}
+
+void FAnimSequencerInstanceProxy::UpdateAnimTrack(UAnimSequenceBase* InAnimSequence, uint32 SequenceId, TOptional<float> InFromPosition, float InToPosition, float Weight, bool bFireNotifies)
 {
 	EnsureAnimTrack(InAnimSequence, SequenceId);
 
 	FSequencerPlayerAnimSequence* PlayerState = FindPlayer<FSequencerPlayerAnimSequence>(SequenceId);
-	PlayerState->PlayerNode.ExplicitTime = InPosition;
+	PlayerState->PlayerNode.ExplicitTime = InToPosition;
+	if (InFromPosition.IsSet())
+	{
+		// Set the internal time accumulator at the "from" time so that the player node will correctly evaluate the
+		// desired "from/to" range. We also disable the reinitialization code so it doesn't mess up that time we
+		// just set.
+		PlayerState->PlayerNode.SetExplicitPreviousTime(InFromPosition.GetValue());
+		PlayerState->PlayerNode.ReinitializationBehavior = ESequenceEvalReinit::NoReset;
+	}
 	// if no fire notifies, we can teleport to explicit time
 	PlayerState->PlayerNode.bTeleportToExplicitTime = !bFireNotifies;
 	// if moving to 0.f, we mark this to teleport. Otherwise, do not use explicit time
diff --git a/Engine/Source/Runtime/AnimGraphRuntime/Public/AnimNodes/AnimNode_SequenceEvaluator.h b/Engine/Source/Runtime/AnimGraphRuntime/Public/AnimNodes/AnimNode_SequenceEvaluator.h
--- a/Engine/Source/Runtime/AnimGraphRuntime/Public/AnimNodes/AnimNode_SequenceEvaluator.h
+++ b/Engine/Source/Runtime/AnimGraphRuntime/Public/AnimNodes/AnimNode_SequenceEvaluator.h
@@ -90,4 +90,6 @@ public:
 	virtual void SetAccumulatedTime(const float& NewTime) {ExplicitTime = NewTime;}
 	virtual UAnimationAsset* GetAnimAsset() {return Sequence;}
 	// End of FAnimNode_AssetPlayerBase Interface
+
+	void SetExplicitPreviousTime(float PreviousTime) { InternalTimeAccumulator = PreviousTime; }
 };
diff --git a/Engine/Source/Runtime/AnimGraphRuntime/Public/AnimSequencerInstance.h b/Engine/Source/Runtime/AnimGraphRuntime/Public/AnimSequencerInstance.h
--- a/Engine/Source/Runtime/AnimGraphRuntime/Public/AnimSequencerInstance.h
+++ b/Engine/Source/Runtime/AnimGraphRuntime/Public/AnimSequencerInstance.h
@@ -20,6 +20,7 @@ public:
 
 	/** Update an animation sequence player in this instance */
 	void UpdateAnimTrack(UAnimSequenceBase* InAnimSequence, int32 SequenceId, float InPosition, float Weight, bool bFireNotifies);
+	void UpdateAnimTrack(UAnimSequenceBase* InAnimSequence, int32 SequenceId, float InFromPosition, float InToPosition, float Weight, bool bFireNotifies);
 
 	/** Reset all nodes in this instance */
 	void ResetNodes();
diff --git a/Engine/Source/Runtime/AnimGraphRuntime/Public/AnimSequencerInstanceProxy.h b/Engine/Source/Runtime/AnimGraphRuntime/Public/AnimSequencerInstanceProxy.h
--- a/Engine/Source/Runtime/AnimGraphRuntime/Public/AnimSequencerInstanceProxy.h
+++ b/Engine/Source/Runtime/AnimGraphRuntime/Public/AnimSequencerInstanceProxy.h
@@ -83,6 +83,7 @@ public:
 
 	/** Update an animation sequence player in this instance */
 	void UpdateAnimTrack(UAnimSequenceBase* InAnimSequence, uint32 SequenceId, float InPosition, float Weight, bool bFireNotifies);
+	void UpdateAnimTrack(UAnimSequenceBase* InAnimSequence, uint32 SequenceId, TOptional<float> InFromPosition, float InToPosition, float Weight, bool bFireNotifies);
 
 	/** Reset all nodes in this instance */
 	virtual void ResetNodes();
diff --git a/Engine/Source/Runtime/Engine/Classes/Animation/AnimMontage.h b/Engine/Source/Runtime/Engine/Classes/Animation/AnimMontage.h
--- a/Engine/Source/Runtime/Engine/Classes/Animation/AnimMontage.h
+++ b/Engine/Source/Runtime/Engine/Classes/Animation/AnimMontage.h
@@ -366,10 +366,11 @@ private:
 	int32 SyncGroupIndex;
 
 	/**
-	 * Optional position to force next update (ignoring the real delta time).
+	 * Optional evaluation range to use next update (ignoring the real delta time).
 	 * Used by external systems that are setting animation times directly. Will fire off notifies and other events provided the animation system is ticking.
 	 */
-	TOptional<float> ForcedNextPosition;
+	TOptional<float> ForcedNextFromPosition;
+	TOptional<float> ForcedNextToPosition;
 
 	UPROPERTY(Transient)
 	int32 DisableRootMotionCount;
@@ -470,8 +471,10 @@ public:
 	void PopDisableRootMotion() { DisableRootMotionCount--; }
 	bool IsRootMotionDisabled() const {	return DisableRootMotionCount > 0; }
 
-	/** Set the position of this animation as part of the next animation update tick. Will trigger events and notifies for the delta time. */
-	void SetNextPositionWithEvents(float InPosition) { ForcedNextPosition = InPosition; }
+	/** Set the next position of this animation for the next animation update tick. Will trigger events and notifies since last position. */
+	void SetNextPositionWithEvents(float InPosition) { ForcedNextFromPosition.Reset();  ForcedNextToPosition = InPosition; }
+	/** Set the evaluation range of this animation for the next animation update tick. Will trigger events and notifies for that range. */
+	void SetNextPositionWithEvents(float InFromPosition, float InToPosition) { ForcedNextFromPosition = InFromPosition; ForcedNextToPosition = InToPosition; }
 
 	/**
 	 * Montage Tick happens in 2 phases
@@ -523,8 +526,8 @@ public:
 	static UAnimMontage* SetMatineeAnimPositionInner(FName SlotName, USkeletalMeshComponent* SkeletalMeshComponent, UAnimSequenceBase* InAnimSequence, float InPosition, bool bLooping);
 	static UAnimMontage* PreviewMatineeSetAnimPositionInner(FName SlotName, USkeletalMeshComponent* SkeletalMeshComponent, UAnimSequenceBase* InAnimSequence, float InPosition, bool bLooping, bool bFireNotifies, float DeltaTime);
 	/** static functions that are used by sequencer montage support*/
-	static UAnimMontage* SetSequencerMontagePosition(FName SlotName, USkeletalMeshComponent* SkeletalMeshComponent, int32& InOutInstanceId, UAnimSequenceBase* InAnimSequence, float InPosition, float Weight, bool bLooping, bool bPlaying);
-	static UAnimMontage* PreviewSequencerMontagePosition(FName SlotName, USkeletalMeshComponent* SkeletalMeshComponent, int32& InOutInstanceId, UAnimSequenceBase* InAnimSequence, float InPosition, float Weight, bool bLooping, bool bFireNotifies, bool bPlaying);
+	static UAnimMontage* SetSequencerMontagePosition(FName SlotName, USkeletalMeshComponent* SkeletalMeshComponent, int32& InOutInstanceId, UAnimSequenceBase* InAnimSequence, float InFromPosition, float InToPosition, float Weight, bool bLooping, bool bPlaying);
+	static UAnimMontage* PreviewSequencerMontagePosition(FName SlotName, USkeletalMeshComponent* SkeletalMeshComponent, int32& InOutInstanceId, UAnimSequenceBase* InAnimSequence, float InFromPosition, float InToPosition, float Weight, bool bLooping, bool bFireNotifies, bool bPlaying);
 private:
 	static UAnimMontage* InitializeMatineeControl(FName SlotName, USkeletalMeshComponent* SkeletalMeshComponent, UAnimSequenceBase* InAnimSequence, bool bLooping);
 };
diff --git a/Engine/Source/Runtime/Engine/Private/Animation/AnimMontage.cpp b/Engine/Source/Runtime/Engine/Private/Animation/AnimMontage.cpp
--- a/Engine/Source/Runtime/Engine/Private/Animation/AnimMontage.cpp
+++ b/Engine/Source/Runtime/Engine/Private/Animation/AnimMontage.cpp
@@ -1828,9 +1828,13 @@ EMontageSubStepResult FMontageSubStepper::Advance(float& InOut_P_Original, const
 
 	// If we're forcing next position, this is our DeltaMove.
 	// We don't use play rate and delta time to move.
-	if (MontageInstance->ForcedNextPosition.IsSet())
+	if (MontageInstance->ForcedNextToPosition.IsSet())
 	{
-		const float NewPosition = MontageInstance->ForcedNextPosition.GetValue();
+		const float NewPosition = MontageInstance->ForcedNextToPosition.GetValue();
+		if (MontageInstance->ForcedNextFromPosition.IsSet())
+		{
+			InOut_P_Original = MontageInstance->ForcedNextFromPosition.GetValue();
+		}
 		DeltaMove = NewPosition - InOut_P_Original;
 		PlayRate = DeltaMove / TimeRemaining;
 		bPlayingForward = (DeltaMove >= 0.f);
@@ -2320,7 +2324,8 @@ void FAnimMontageInstance::Advance(float DeltaTime, struct FRootMotionMovementPa
 			}
 		
 			// if we had a ForcedNextPosition set, reset it.
-			ForcedNextPosition.Reset();
+			ForcedNextToPosition.Reset();
+			ForcedNextFromPosition.Reset();
 		}
 	}
 
@@ -2668,7 +2673,7 @@ UAnimMontage* FAnimMontageInstance::PreviewMatineeSetAnimPositionInner(FName Slo
 	return PlayingMontage;
 }
 
-UAnimMontage* FAnimMontageInstance::SetSequencerMontagePosition(FName SlotName, USkeletalMeshComponent* SkeletalMeshComponent, int32& InOutInstanceId, UAnimSequenceBase* InAnimSequence, float InPosition, float Weight, bool bLooping, bool bInPlaying)
+UAnimMontage* FAnimMontageInstance::SetSequencerMontagePosition(FName SlotName, USkeletalMeshComponent* SkeletalMeshComponent, int32& InOutInstanceId, UAnimSequenceBase* InAnimSequence, float InFromPosition, float InToPosition, float Weight, bool bLooping, bool bInPlaying)
 {
 	UAnimInstance* AnimInst = SkeletalMeshComponent->GetAnimInstance();
 	if (AnimInst)
@@ -2698,11 +2703,11 @@ UAnimMontage* FAnimMontageInstance::SetSequencerMontagePosition(FName SlotName,
 			
 			if (bInPlaying)
 			{
-				MontageInstanceToUpdate->SetNextPositionWithEvents(InPosition);
+				MontageInstanceToUpdate->SetNextPositionWithEvents(InFromPosition, InToPosition);
 			}
 			else
 			{
-				MontageInstanceToUpdate->SetPosition(InPosition);
+				MontageInstanceToUpdate->SetPosition(InToPosition);
 			}
 
 			MontageInstanceToUpdate->bPlaying = bInPlaying;
@@ -2717,21 +2722,20 @@ UAnimMontage* FAnimMontageInstance::SetSequencerMontagePosition(FName SlotName,
 	return nullptr;
 }
 
-UAnimMontage* FAnimMontageInstance::PreviewSequencerMontagePosition(FName SlotName, USkeletalMeshComponent* SkeletalMeshComponent, int32& InOutInstanceId, UAnimSequenceBase* InAnimSequence, float InPosition, float Weight, bool bLooping, bool bFireNotifies, bool bInPlaying)
+UAnimMontage* FAnimMontageInstance::PreviewSequencerMontagePosition(FName SlotName, USkeletalMeshComponent* SkeletalMeshComponent, int32& InOutInstanceId, UAnimSequenceBase* InAnimSequence, float InFromPosition, float InToPosition, float Weight, bool bLooping, bool bFireNotifies, bool bInPlaying)
 {
 	UAnimInstance* AnimInst = SkeletalMeshComponent->GetAnimInstance();
 	if (AnimInst)
 	{
 		FAnimMontageInstance* MontageInstanceToUpdate = AnimInst->GetMontageInstanceForID(InOutInstanceId);
-		float PreviousPosition = (MontageInstanceToUpdate) ? MontageInstanceToUpdate->GetPosition() : InPosition;
 
-		UAnimMontage* PlayingMontage = SetSequencerMontagePosition(SlotName, SkeletalMeshComponent, InOutInstanceId, InAnimSequence, InPosition, Weight, bLooping, bInPlaying);
+		UAnimMontage* PlayingMontage = SetSequencerMontagePosition(SlotName, SkeletalMeshComponent, InOutInstanceId, InAnimSequence, InFromPosition, InToPosition, Weight, bLooping, bInPlaying);
 		if (PlayingMontage)
 		{
 			// we have to get it again in case if this is new
 			MontageInstanceToUpdate = AnimInst->GetMontageInstanceForID(InOutInstanceId);
 			// since we don't advance montage in the tick, we manually have to handle notifies
-			MontageInstanceToUpdate->HandleEvents(PreviousPosition, InPosition, NULL);
+			MontageInstanceToUpdate->HandleEvents(InFromPosition, InToPosition, NULL);
 			if (!bFireNotifies)
 			{
 				AnimInst->NotifyQueue.Reset(SkeletalMeshComponent);
diff --git a/Engine/Source/Runtime/MovieSceneTracks/Private/Evaluation/MovieSceneSkeletalAnimationTemplate.cpp b/Engine/Source/Runtime/MovieSceneTracks/Private/Evaluation/MovieSceneSkeletalAnimationTemplate.cpp
--- a/Engine/Source/Runtime/MovieSceneTracks/Private/Evaluation/MovieSceneSkeletalAnimationTemplate.cpp
+++ b/Engine/Source/Runtime/MovieSceneTracks/Private/Evaluation/MovieSceneSkeletalAnimationTemplate.cpp
@@ -149,9 +149,10 @@ struct FPreAnimatedAnimationTokenProducer : IMovieScenePreAnimatedTokenProducer
 
 struct FMinimalAnimParameters
 {
-	FMinimalAnimParameters(UAnimSequenceBase* InAnimation, float InEvalTime, float InBlendWeight, const FMovieSceneEvaluationScope& InScope, FName InSlotName, FObjectKey InSection, bool InSkipAnimationNotifiers, bool InForceCustomMode)
+	FMinimalAnimParameters(UAnimSequenceBase* InAnimation, float InFromEvalTime, float InToEvalTime, float InBlendWeight, const FMovieSceneEvaluationScope& InScope, FName InSlotName, FObjectKey InSection, bool InSkipAnimationNotifiers, bool InForceCustomMode)
 		: Animation(InAnimation)
-		, EvalTime(InEvalTime)
+		, FromEvalTime(InFromEvalTime)
+		, ToEvalTime(InToEvalTime)
 		, BlendWeight(InBlendWeight)
 		, EvaluationScope(InScope)
 		, SlotName(InSlotName)
@@ -161,7 +162,8 @@ struct FMinimalAnimParameters
 	{}
 	
 	UAnimSequenceBase* Animation;
-	float EvalTime;
+	float FromEvalTime;
+	float ToEvalTime;
 	float BlendWeight;
 	FMovieSceneEvaluationScope EvaluationScope;
 	FName SlotName;
@@ -248,9 +250,10 @@ namespace MovieScene
 
 			const EMovieScenePlayerStatus::Type PlayerStatus = Player.GetPlaybackStatus();
 
-			// If the playback status is jumping, ie. one such occurrence is setting the time for thumbnail generation, disable anim notifies updates because it could fire audio
-			// We now layer this with the passed in notify toggleto force a disable in this case.
-			const bool bFireNotifies = !bPreviewPlayback || (PlayerStatus != EMovieScenePlayerStatus::Jumping && PlayerStatus != EMovieScenePlayerStatus::Stopped);
+			// If the playback status is jumping, ie. one such occurrence is setting the time for thumbnail generation, disable anim notifies updates because it could fire audio.
+			// If the playback status is scrubbing, we disable notifies for now because we can't properly fire them in all cases until we get evaluation range info.
+			// We now layer this with the passed in notify toggle to force a disable in this case.
+			const bool bFireNotifies = !bPreviewPlayback || (PlayerStatus != EMovieScenePlayerStatus::Jumping && PlayerStatus != EMovieScenePlayerStatus::Stopped && PlayerStatus != EMovieScenePlayerStatus::Scrubbing);
 
 			// When jumping from one cut to another cut, the delta time should be 0 so that anim notifies before the current position are not evaluated. Note, anim notifies at the current time should still be evaluated.
 			const double DeltaTime = ( Context.HasJumped() ? FFrameTime(0) : Context.GetRange().Size<FFrameTime>() ) / Context.GetFrameRate();
@@ -350,14 +353,14 @@ namespace MovieScene
 				if (bPreviewPlayback)
 				{
 					PreviewSetAnimPosition(PersistentData, Player, SkeletalMeshComponent,
-						AnimParams.SlotName, AnimParams.Section, AnimParams.Animation, AnimParams.EvalTime, AnimParams.BlendWeight,
+						AnimParams.SlotName, AnimParams.Section, AnimParams.Animation, AnimParams.FromEvalTime, AnimParams.ToEvalTime, AnimParams.BlendWeight,
 						bFireNotifies && !AnimParams.bSkipAnimNotifiers, DeltaTime, PlayerStatus == EMovieScenePlayerStatus::Playing, 
 						bResetDynamics, AnimParams.bForceCustomMode);
 				}
 				else
 				{
 					SetAnimPosition(PersistentData, Player, SkeletalMeshComponent,
-						AnimParams.SlotName, AnimParams.Section, AnimParams.Animation, AnimParams.EvalTime, AnimParams.BlendWeight,
+						AnimParams.SlotName, AnimParams.Section, AnimParams.Animation, AnimParams.FromEvalTime, AnimParams.ToEvalTime, AnimParams.BlendWeight,
 						PlayerStatus == EMovieScenePlayerStatus::Playing, bFireNotifies && !AnimParams.bSkipAnimNotifiers,
 						AnimParams.bForceCustomMode
 					);
@@ -365,7 +368,7 @@ namespace MovieScene
 			}
 		}
 
-		void SetAnimPosition(FPersistentEvaluationData& PersistentData, IMovieScenePlayer& Player, USkeletalMeshComponent* SkeletalMeshComponent, FName SlotName, FObjectKey Section, UAnimSequenceBase* InAnimSequence, float InPosition, float Weight, bool bPlaying, bool bFireNotifies, bool bForceCustomMode)
+		void SetAnimPosition(FPersistentEvaluationData& PersistentData, IMovieScenePlayer& Player, USkeletalMeshComponent* SkeletalMeshComponent, FName SlotName, FObjectKey Section, UAnimSequenceBase* InAnimSequence, float InFromPosition, float InToPosition, float Weight, bool bPlaying, bool bFireNotifies, bool bForceCustomMode)
 		{
 			static const bool bLooping = false;
 
@@ -385,14 +388,14 @@ namespace MovieScene
 				Player.SavePreAnimatedState(*SequencerInst, AnimTypeID, FStatelessPreAnimatedTokenProducer(&ResetAnimSequencerInstance));
 
 				// Set position and weight
-				SequencerInst->UpdateAnimTrack(InAnimSequence, GetTypeHash(AnimTypeID), InPosition, Weight, bFireNotifies);
+				SequencerInst->UpdateAnimTrack(InAnimSequence, GetTypeHash(AnimTypeID), InFromPosition, InToPosition, Weight, bFireNotifies);
 			}
 			else if (UAnimInstance* AnimInst = SkeletalMeshComponent->GetAnimInstance())
 			{
 				FMontagePlayerPerSectionData* SectionData = MontageData.Find(Section);
 
 				int32 InstanceId = (SectionData) ? SectionData->MontageInstanceId : INDEX_NONE;
-				TWeakObjectPtr<UAnimMontage> Montage = FAnimMontageInstance::SetSequencerMontagePosition(SlotName, SkeletalMeshComponent, InstanceId, InAnimSequence, InPosition, Weight, bLooping, bPlaying);
+				TWeakObjectPtr<UAnimMontage> Montage = FAnimMontageInstance::SetSequencerMontagePosition(SlotName, SkeletalMeshComponent, InstanceId, InAnimSequence, InFromPosition, InToPosition, Weight, bLooping, bPlaying);
 
 				if (Montage.IsValid())
 				{
@@ -410,7 +413,7 @@ namespace MovieScene
 			}
 		}
 
-		void PreviewSetAnimPosition(FPersistentEvaluationData& PersistentData, IMovieScenePlayer& Player, USkeletalMeshComponent* SkeletalMeshComponent, FName SlotName, FObjectKey Section, UAnimSequenceBase* InAnimSequence, float InPosition, float Weight, bool bFireNotifies, float DeltaTime, bool bPlaying, bool bResetDynamics, bool bForceCustomMode)
+		void PreviewSetAnimPosition(FPersistentEvaluationData& PersistentData, IMovieScenePlayer& Player, USkeletalMeshComponent* SkeletalMeshComponent, FName SlotName, FObjectKey Section, UAnimSequenceBase* InAnimSequence, float InFromPosition, float InToPosition, float Weight, bool bFireNotifies, float DeltaTime, bool bPlaying, bool bResetDynamics, bool bForceCustomMode)
 		{
 			static const bool bLooping = false;
 
@@ -430,14 +433,14 @@ namespace MovieScene
 				Player.SavePreAnimatedState(*SequencerInst, AnimTypeID, FStatelessPreAnimatedTokenProducer(&ResetAnimSequencerInstance));
 
 				// Set position and weight
-				SequencerInst->UpdateAnimTrack(InAnimSequence, GetTypeHash(AnimTypeID), InPosition, Weight, bFireNotifies);
+				SequencerInst->UpdateAnimTrack(InAnimSequence, GetTypeHash(AnimTypeID), InFromPosition, InToPosition, Weight, bFireNotifies);
 			}
 			else if (UAnimInstance* AnimInst = SkeletalMeshComponent->GetAnimInstance())
 			{
 				FMontagePlayerPerSectionData* SectionData = MontageData.Find(Section);
 
 				int32 InstanceId = (SectionData)? SectionData->MontageInstanceId : INDEX_NONE;
-				TWeakObjectPtr<UAnimMontage> Montage = FAnimMontageInstance::PreviewSequencerMontagePosition(SlotName, SkeletalMeshComponent, InstanceId, InAnimSequence, InPosition, Weight, bLooping, bFireNotifies, bPlaying);
+				TWeakObjectPtr<UAnimMontage> Montage = FAnimMontageInstance::PreviewSequencerMontagePosition(SlotName, SkeletalMeshComponent, InstanceId, InAnimSequence, InFromPosition, InToPosition, Weight, bLooping, bFireNotifies, bPlaying);
 
 				if (Montage.IsValid())
 				{
@@ -481,7 +484,6 @@ void FMovieSceneSkeletalAnimationSectionTemplate::Evaluate(const FMovieSceneEval
 {
 	if (Params.Animation)
 	{
-
 		FOptionalMovieSceneBlendType BlendType = GetSourceSection()->GetBlendType();
 		check(BlendType.IsValid());
 
@@ -493,8 +495,9 @@ void FMovieSceneSkeletalAnimationSectionTemplate::Evaluate(const FMovieSceneEval
 			Accumulator.DefineActuator(ActuatorTypeID, MakeShared<MovieScene::FComponentAnimationActuator>());
 		}
 
-		// calculate the time at which to evaluate the animation
+		// Calculate the time at which to evaluate the animation
 		float EvalTime = Params.MapTimeToAnimation(Context.GetTime(), Context.GetFrameRate());
+		float PreviousEvalTime = Params.MapTimeToAnimation(Context.GetPreviousTime(), Context.GetFrameRate());
 
 		float ManualWeight = 1.f;
 		Params.Weight.Evaluate(Context.GetTime(), ManualWeight);
@@ -503,7 +506,7 @@ void FMovieSceneSkeletalAnimationSectionTemplate::Evaluate(const FMovieSceneEval
 
 		// Add the blendable to the accumulator
 		FMinimalAnimParameters AnimParams(
-			Params.Animation, EvalTime, Weight, ExecutionTokens.GetCurrentScope(), Params.SlotName, GetSourceSection(), Params.bSkipAnimNotifiers, 
+			Params.Animation, PreviousEvalTime, EvalTime, Weight, ExecutionTokens.GetCurrentScope(), Params.SlotName, GetSourceSection(), Params.bSkipAnimNotifiers, 
 			Params.bForceCustomMode
 		);
 		ExecutionTokens.BlendToken(ActuatorTypeID, TBlendableToken<MovieScene::FBlendedAnimation>(AnimParams, BlendType.Get(), 1.f));
@@ -512,7 +515,8 @@ void FMovieSceneSkeletalAnimationSectionTemplate::Evaluate(const FMovieSceneEval
 		{
 			FFrameTime SimulatedTime = IMovieSceneMotionVectorSimulation::GetSimulationTime(Context);
 
-			// calculate the time at which to evaluate the animation
+			// Calculate the time at which to evaluate the animation
+			float CurrentEvalTime = Params.MapTimeToAnimation(Context.GetTime(), Context.GetFrameRate());
 			float SimulatedEvalTime = Params.MapTimeToAnimation(SimulatedTime, Context.GetFrameRate());
 
 			float SimulatedManualWeight = 1.f;
@@ -521,7 +525,8 @@ void FMovieSceneSkeletalAnimationSectionTemplate::Evaluate(const FMovieSceneEval
 			const float SimulatedWeight = SimulatedManualWeight * EvaluateEasing(SimulatedTime);
 
 			FSimulatedAnimParameters SimulatedAnimParams{ AnimParams };
-			SimulatedAnimParams.AnimParams.EvalTime = SimulatedEvalTime;
+			SimulatedAnimParams.AnimParams.FromEvalTime = CurrentEvalTime;
+			SimulatedAnimParams.AnimParams.ToEvalTime = SimulatedEvalTime;
 			SimulatedAnimParams.AnimParams.BlendWeight = SimulatedWeight;
 			ExecutionTokens.BlendToken(ActuatorTypeID, TBlendableToken<MovieScene::FBlendedAnimation>(SimulatedAnimParams, BlendType.Get(), 1.f));
 		}
