Sequencer: Fix for allowing spatialized audio on scene components, not just actors

diff --git a/Engine/Source/Runtime/MovieSceneTracks/Private/Evaluation/MovieSceneAudioTemplate.cpp b/Engine/Source/Runtime/MovieSceneTracks/Private/Evaluation/MovieSceneAudioTemplate.cpp
--- a/Engine/Source/Runtime/MovieSceneTracks/Private/Evaluation/MovieSceneAudioTemplate.cpp
+++ b/Engine/Source/Runtime/MovieSceneTracks/Private/Evaluation/MovieSceneAudioTemplate.cpp
@@ -114,41 +114,64 @@ struct FCachedAudioTrackData : IPersistentEvaluationData
 	}
 
 	/** Only to be called on the game thread */
-	UAudioComponent* AddAudioComponentForRow(int32 RowIndex, FObjectKey SectionKey, AActor& PrincipalActor, IMovieScenePlayer& Player)
+	UAudioComponent* AddAudioComponentForRow(int32 RowIndex, FObjectKey SectionKey, UObject& PrincipalObject, IMovieScenePlayer& Player)
 	{
-		FObjectKey ActorKey(&PrincipalActor);
+		FObjectKey ObjectKey(&PrincipalObject);
 		
-		if (!AudioComponentsByActorKey.Contains(ActorKey))
+		if (!AudioComponentsByActorKey.Contains(ObjectKey))
 		{
-			AudioComponentsByActorKey.Add(ActorKey, TMap<FObjectKey, TWeakObjectPtr<UAudioComponent>>());
+			AudioComponentsByActorKey.Add(ObjectKey, TMap<FObjectKey, TWeakObjectPtr<UAudioComponent>>());
 		}
 
-		UAudioComponent* ExistingComponent = GetAudioComponent(ActorKey, SectionKey);
+		UAudioComponent* ExistingComponent = GetAudioComponent(ObjectKey, SectionKey);
 		if (!ExistingComponent)
 		{
 			USoundCue* TempPlaybackAudioCue = NewObject<USoundCue>();
 
-			FAudioDevice::FCreateComponentParams Params(PrincipalActor.GetWorld(), &PrincipalActor);
-			ExistingComponent = FAudioDevice::CreateComponent(TempPlaybackAudioCue, Params);
+			AActor* Actor = nullptr;
+			USceneComponent* SceneComponent = nullptr;
+			FString ObjectName;
 
-			if (!ExistingComponent)
+			if (PrincipalObject.IsA<AActor>())
 			{
-					FString ActorName = 
+				Actor = Cast<AActor>(&PrincipalObject);
+				SceneComponent = Actor->GetRootComponent();
+				ObjectName =
 #if WITH_EDITOR
-						PrincipalActor.GetActorLabel();
+					Actor->GetActorLabel();
 #else
-						PrincipalActor.GetName();
+					Actor->GetName();
 #endif
-				UE_LOG(LogMovieScene, Warning, TEXT("Failed to create audio component for spatialized audio track (row %d on %s)."), RowIndex, *ActorName);
+			}
+			else if (PrincipalObject.IsA<UActorComponent>())
+			{
+				UActorComponent* ActorComponent = Cast<UActorComponent>(&PrincipalObject);
+				Actor = ActorComponent->GetOwner();
+				SceneComponent = Cast<USceneComponent>(ActorComponent);
+				ObjectName = ActorComponent->GetName();
+			}
+
+			if (!Actor || !SceneComponent)
+			{
+				UE_LOG(LogMovieScene, Warning, TEXT("Failed to find scene component for spatialized audio track (row %d)."), RowIndex);
+				return nullptr;
+			}
+
+			FAudioDevice::FCreateComponentParams Params(Actor->GetWorld(), Actor);
+			ExistingComponent = FAudioDevice::CreateComponent(TempPlaybackAudioCue, Params);
+
+			if (!ExistingComponent)
+			{
+				UE_LOG(LogMovieScene, Warning, TEXT("Failed to create audio component for spatialized audio track (row %d on %s)."), RowIndex, *ObjectName);
 				return nullptr;
 			}
 
 			Player.SavePreAnimatedState(*ExistingComponent, FMovieSceneAnimTypeID::Unique(), FDestroyAudioPreAnimatedToken::FProducer());
 
-			AudioComponentsByActorKey[ActorKey].Add(SectionKey, ExistingComponent);
+			AudioComponentsByActorKey[ObjectKey].Add(SectionKey, ExistingComponent);
 
 			ExistingComponent->SetFlags(RF_Transient);
-			ExistingComponent->AttachToComponent(PrincipalActor.GetRootComponent(), FAttachmentTransformRules::KeepRelativeTransform);
+			ExistingComponent->AttachToComponent(SceneComponent, FAttachmentTransformRules::KeepRelativeTransform);
 		}
 
 		return ExistingComponent;
@@ -260,17 +283,11 @@ struct FAudioSectionExecutionToken : IMovieSceneExecutionToken
 		{
 			for (TWeakObjectPtr<> Object : Player.FindBoundObjects(Operand))
 			{
-				AActor* Actor = Cast<AActor>(Object.Get());
-				if (!Actor)
-				{
-					continue;
-				}
-
-				UAudioComponent* AudioComponent = TrackData.GetAudioComponent(Actor, SectionKey);
+				UAudioComponent* AudioComponent = TrackData.GetAudioComponent(Object.Get(), SectionKey);
 				if (!AudioComponent)
 				{
 					// Initialize the sound
-					AudioComponent = TrackData.AddAudioComponentForRow(AudioSection->GetRowIndex(), SectionKey, *Actor, Player);
+					AudioComponent = TrackData.AddAudioComponentForRow(AudioSection->GetRowIndex(), SectionKey, *Object.Get(), Player);
 
 					if (AudioComponent)
 					{
