Sequencer: Fix issue where anims with a rate scale other than 1.0 will not loop correctly when using anim blueprint

diff --git a/Engine/Source/Editor/MovieSceneTools/Private/TrackEditors/SkeletalAnimationTrackEditor.cpp b/Engine/Source/Editor/MovieSceneTools/Private/TrackEditors/SkeletalAnimationTrackEditor.cpp
--- a/Engine/Source/Editor/MovieSceneTools/Private/TrackEditors/SkeletalAnimationTrackEditor.cpp
+++ b/Engine/Source/Editor/MovieSceneTools/Private/TrackEditors/SkeletalAnimationTrackEditor.cpp
@@ -311,7 +311,7 @@ int32 FSkeletalAnimationSection::OnPaintSection( FSequencerSectionPainter& Paint
 	FFrameRate TickResolution = TimeToPixelConverter.GetTickResolution();
 
 	// Add lines where the animation starts and ends/loops
-	const float AnimPlayRate = FMath::IsNearlyZero(Section.Params.PlayRate) ? 1.0f : Section.Params.PlayRate;
+	const float AnimPlayRate = FMath::IsNearlyZero(Section.Params.PlayRate) ? 1.0f : Section.Params.PlayRate * Section.Params.Animation->RateScale;
 	const float SeqLength = (Section.Params.GetSequenceLength() - TickResolution.AsSeconds(Section.Params.StartFrameOffset + Section.Params.EndFrameOffset)) / AnimPlayRate;
 	const float FirstLoopSeqLength = SeqLength - TickResolution.AsSeconds(Section.Params.FirstLoopStartFrameOffset) / AnimPlayRate;
 
diff --git a/Engine/Source/Runtime/MovieSceneTracks/Private/Evaluation/MovieSceneSkeletalAnimationTemplate.cpp b/Engine/Source/Runtime/MovieSceneTracks/Private/Evaluation/MovieSceneSkeletalAnimationTemplate.cpp
--- a/Engine/Source/Runtime/MovieSceneTracks/Private/Evaluation/MovieSceneSkeletalAnimationTemplate.cpp
+++ b/Engine/Source/Runtime/MovieSceneTracks/Private/Evaluation/MovieSceneSkeletalAnimationTemplate.cpp
@@ -395,7 +395,9 @@ namespace MovieScene
 				FMontagePlayerPerSectionData* SectionData = MontageData.Find(Section);
 
 				int32 InstanceId = (SectionData) ? SectionData->MontageInstanceId : INDEX_NONE;
-				TWeakObjectPtr<UAnimMontage> Montage = FAnimMontageInstance::SetSequencerMontagePosition(SlotName, SkeletalMeshComponent, InstanceId, InAnimSequence, InFromPosition, InToPosition, Weight, bLooping, bPlaying);
+
+				const float AssetPlayRate = FMath::IsNearlyZero(InAnimSequence->RateScale) ? 1.0f : InAnimSequence->RateScale;
+				TWeakObjectPtr<UAnimMontage> Montage = FAnimMontageInstance::SetSequencerMontagePosition(SlotName, SkeletalMeshComponent, InstanceId, InAnimSequence, InFromPosition / AssetPlayRate, InToPosition / AssetPlayRate, Weight, bLooping, bPlaying);
 
 				if (Montage.IsValid())
 				{
@@ -440,7 +442,9 @@ namespace MovieScene
 				FMontagePlayerPerSectionData* SectionData = MontageData.Find(Section);
 
 				int32 InstanceId = (SectionData)? SectionData->MontageInstanceId : INDEX_NONE;
-				TWeakObjectPtr<UAnimMontage> Montage = FAnimMontageInstance::PreviewSequencerMontagePosition(SlotName, SkeletalMeshComponent, InstanceId, InAnimSequence, InFromPosition, InToPosition, Weight, bLooping, bFireNotifies, bPlaying);
+
+				const float AssetPlayRate = FMath::IsNearlyZero(InAnimSequence->RateScale) ? 1.0f : InAnimSequence->RateScale;
+				TWeakObjectPtr<UAnimMontage> Montage = FAnimMontageInstance::PreviewSequencerMontagePosition(SlotName, SkeletalMeshComponent, InstanceId, InAnimSequence, InFromPosition / AssetPlayRate, InToPosition / AssetPlayRate, Weight, bLooping, bFireNotifies, bPlaying);
 
 				if (Montage.IsValid())
 				{
diff --git a/Engine/Source/Runtime/MovieSceneTracks/Private/Sections/MovieSceneSkeletalAnimationSection.cpp b/Engine/Source/Runtime/MovieSceneTracks/Private/Sections/MovieSceneSkeletalAnimationSection.cpp
--- a/Engine/Source/Runtime/MovieSceneTracks/Private/Sections/MovieSceneSkeletalAnimationSection.cpp
+++ b/Engine/Source/Runtime/MovieSceneTracks/Private/Sections/MovieSceneSkeletalAnimationSection.cpp
@@ -248,7 +248,7 @@ void UMovieSceneSkeletalAnimationSection::GetSnapTimes(TArray<FFrameNumber>& Out
 	const FFrameNumber StartFrame = GetInclusiveStartFrame();
 	const FFrameNumber EndFrame   = GetExclusiveEndFrame() - 1; // -1 because we don't need to add the end frame twice
 
-	const float AnimPlayRate     = FMath::IsNearlyZero(Params.PlayRate) ? 1.0f : Params.PlayRate;
+	const float AnimPlayRate     = FMath::IsNearlyZero(Params.PlayRate) ? 1.0f : Params.PlayRate * Params.Animation->RateScale;
 	const float SeqLengthSeconds = Params.GetSequenceLength() - FrameRate.AsSeconds(Params.StartFrameOffset + Params.EndFrameOffset) / AnimPlayRate;
 	const float FirstLoopSeqLengthSeconds = SeqLengthSeconds - FrameRate.AsSeconds(Params.FirstLoopStartFrameOffset) / AnimPlayRate;
 
