Sequencer: Fix unbind skeletal mesh component.

Revert 7662465 which delayed the unbind until the end after resetting the pose. Implemented a more direct fix for motion vectors not getting cleared after tearing down sequencer.

diff --git a/Engine/Source/Runtime/MovieSceneTracks/Private/Evaluation/MovieSceneSkeletalAnimationTemplate.cpp b/Engine/Source/Runtime/MovieSceneTracks/Private/Evaluation/MovieSceneSkeletalAnimationTemplate.cpp
--- a/Engine/Source/Runtime/MovieSceneTracks/Private/Evaluation/MovieSceneSkeletalAnimationTemplate.cpp
+++ b/Engine/Source/Runtime/MovieSceneTracks/Private/Evaluation/MovieSceneSkeletalAnimationTemplate.cpp
@@ -101,6 +101,8 @@ struct FPreAnimatedAnimationTokenProducer : IMovieScenePreAnimatedTokenProducer
 					SequencerInst->ResetNodes();
 				}
 
+				UAnimSequencerInstance::UnbindFromSkeletalMeshComponent(Component);
+
 				// Reset the mesh component update flag and animation mode to what they were before we animated the object
 				Component->VisibilityBasedAnimTickOption = VisibilityBasedAnimTickOption;
 				if (Component->GetAnimationMode() != AnimationMode)
@@ -129,8 +131,11 @@ struct FPreAnimatedAnimationTokenProducer : IMovieScenePreAnimatedTokenProducer
 #if WITH_EDITOR
 				Component->SetUpdateAnimationInEditor(bUpdateAnimationInEditor);
 #endif
-
-				UAnimSequencerInstance::UnbindFromSkeletalMeshComponent(Component);
+				// if not game world, don't clean this up
+				if (Component->GetWorld() != nullptr && Component->GetWorld()->IsGameWorld() == false)
+				{
+					Component->ClearMotionVector();
+				}
 			}
 
 			EVisibilityBasedAnimTickOption VisibilityBasedAnimTickOption;
