Sequencer: Added ensure to handle mismatched property/function pairs which may result from nativization for the time being

diff --git a/Engine/Source/Runtime/MovieSceneTracks/Private/Evaluation/MovieSceneEventTemplate.cpp b/Engine/Source/Runtime/MovieSceneTracks/Private/Evaluation/MovieSceneEventTemplate.cpp
--- a/Engine/Source/Runtime/MovieSceneTracks/Private/Evaluation/MovieSceneEventTemplate.cpp
+++ b/Engine/Source/Runtime/MovieSceneTracks/Private/Evaluation/MovieSceneEventTemplate.cpp
@@ -261,6 +261,11 @@ struct FEventTriggerExecutionToken
 
 	void TriggerEventWithParameters(UObject* DirectorInstance, const FMovieSceneEventPtrs& Event, TArrayView<UObject* const> EventContexts, IMovieScenePlayer& Player, FMovieSceneSequenceID SequenceID)
 	{
+		if (!ensureMsgf(!Event.BoundObjectProperty || (Event.BoundObjectProperty->GetOuter() == Event.Function && Event.BoundObjectProperty->GetOffset_ForUFunction() < Event.Function->ParmsSize), TEXT("Bound object property belongs to the wrong function or has an offset greater than the parameter size! This should never happen and indicates a BP compilation or nativization error.")))
+		{
+			return;
+		}
+
 		// Parse all function parameters.
 		uint8* Parameters = (uint8*)FMemory_Alloca(Event.Function->ParmsSize + Event.Function->MinAlignment);
 		Parameters = Align(Parameters, Event.Function->MinAlignment);
