Merge fix for FMovieSceneEvent::PostSerialize missing #WITH_EDITORONLY_DATA

diff --git a/Engine/Source/Runtime/MovieSceneTracks/Public/Channels/MovieSceneEvent.h b/Engine/Source/Runtime/MovieSceneTracks/Public/Channels/MovieSceneEvent.h
--- a/Engine/Source/Runtime/MovieSceneTracks/Public/Channels/MovieSceneEvent.h
+++ b/Engine/Source/Runtime/MovieSceneTracks/Public/Channels/MovieSceneEvent.h
@@ -53,7 +53,9 @@ struct FMovieSceneEvent
 	UPROPERTY()
 	FMovieSceneEventPtrs Ptrs;
 
+#if WITH_EDITORONLY_DATA
 	void PostSerialize(const FArchive& Ar);
+#endif
 
 public:
 
