Sequencer: Move data upgrade to ::Serialize in order to ensure it performs early enough during widget blueprint loads

diff --git a/Engine/Source/Editor/MovieSceneTools/Private/MovieSceneToolsModule.cpp b/Engine/Source/Editor/MovieSceneTools/Private/MovieSceneToolsModule.cpp
--- a/Engine/Source/Editor/MovieSceneTools/Private/MovieSceneToolsModule.cpp
+++ b/Engine/Source/Editor/MovieSceneTools/Private/MovieSceneToolsModule.cpp
@@ -293,6 +293,9 @@ void FMovieSceneToolsModule::HandleGenerateEventEntryPoints(UMovieSceneEventSect
 		FLinkerLoad* Linker = EventSection->GetLinker();
 		Linker->Preload(EventSection);
 	}
+
+	EventSection->AttemptUpgrade();
+
 	for (FMovieSceneEvent& EntryPoint : EventSection->GetAllEntryPoints())
 	{
 		UEdGraphNode* Endpoint = FMovieSceneEventUtils::FindEndpoint(&EntryPoint, EventSection, Params.CompilerContext->Blueprint);
diff --git a/Engine/Source/Runtime/MovieSceneTracks/Private/Sections/MovieSceneEventSectionBase.cpp b/Engine/Source/Runtime/MovieSceneTracks/Private/Sections/MovieSceneEventSectionBase.cpp
--- a/Engine/Source/Runtime/MovieSceneTracks/Private/Sections/MovieSceneEventSectionBase.cpp
+++ b/Engine/Source/Runtime/MovieSceneTracks/Private/Sections/MovieSceneEventSectionBase.cpp
@@ -56,28 +56,49 @@ void UMovieSceneEventSectionBase::OnPostCompile(UBlueprint* Blueprint)
 	Blueprint->OnCompiled().RemoveAll(this);
 }
 
-#endif
-
-void UMovieSceneEventSectionBase::PostLoad()
+void UMovieSceneEventSectionBase::AttemptUpgrade()
 {
-	Super::PostLoad();
-
-#if WITH_EDITORONLY_DATA
-	if (UBlueprint* Blueprint = DirectorBlueprint_DEPRECATED.Get())
+	UBlueprint* Blueprint = DirectorBlueprint_DEPRECATED.Get();
+	if (!Blueprint)
 	{
-		UpgradeLegacyEventEndpoint.Broadcast(this, Blueprint);
+		return;
+	}
 
-		Blueprint->GenerateFunctionGraphsEvent.AddUniqueDynamic(this, &UMovieSceneEventSectionBase::HandleGenerateEntryPoints);
+	// Bind onto this event to ensure that we get another chance to upgrade when the BP compiles
+	Blueprint->GenerateFunctionGraphsEvent.AddUniqueDynamic(this, &UMovieSceneEventSectionBase::HandleGenerateEntryPoints);
 
-		// If the BP has already been compiled (eg regenerate on load) we must perform PostCompile fixup immediately since
-		// We will not have had a chance to generate function entries. In this case we just bind directly to the already compiled functions.
-		if (Blueprint->bHasBeenRegenerated)
+	for (UEdGraph* EdGraph : Blueprint->FunctionGraphs)
+	{
+		if (EdGraph->HasAnyFlags(RF_NeedLoad))
 		{
-			OnPostCompile(Blueprint);
+			return;
 		}
+	}
+
+	UpgradeLegacyEventEndpoint.Broadcast(this, Blueprint);
 
-		// We're done with data upgrade now
-		DirectorBlueprint_DEPRECATED = nullptr;
+	// If the BP has already been compiled (eg regenerate on load) we must perform PostCompile fixup immediately since
+	// We will not have had a chance to generate function entries. In this case we just bind directly to the already compiled functions.
+	if (Blueprint->bHasBeenRegenerated)
+	{
+		OnPostCompile(Blueprint);
+	}
+
+	// We're done with data upgrade now
+	DirectorBlueprint_DEPRECATED = nullptr;
+}
+
+#endif
+
+void UMovieSceneEventSectionBase::Serialize(FArchive& Ar)
+{
+	Super::Serialize(Ar);
+
+#if WITH_EDITORONLY_DATA
+	
+	if (Ar.IsLoading())
+	{
+		AttemptUpgrade();
 	}
 #endif
 }
\ No newline at end of file
diff --git a/Engine/Source/Runtime/MovieSceneTracks/Public/Sections/MovieSceneEventSectionBase.h b/Engine/Source/Runtime/MovieSceneTracks/Public/Sections/MovieSceneEventSectionBase.h
--- a/Engine/Source/Runtime/MovieSceneTracks/Public/Sections/MovieSceneEventSectionBase.h
+++ b/Engine/Source/Runtime/MovieSceneTracks/Public/Sections/MovieSceneEventSectionBase.h
@@ -20,10 +20,12 @@ public:
 
 	virtual TArrayView<FMovieSceneEvent> GetAllEntryPoints() { return TArrayView<FMovieSceneEvent>(); }
 
-	virtual void PostLoad() override;
+	virtual void Serialize(FArchive& Ar) override;
 
 #if WITH_EDITOR
 
+	MOVIESCENETRACKS_API void AttemptUpgrade();
+
 	DECLARE_MULTICAST_DELEGATE_TwoParams(FGenerateEventEntryPointFunctionsEvent, UMovieSceneEventSectionBase*, const FGenerateBlueprintFunctionParams&);
 	DECLARE_MULTICAST_DELEGATE_FourParams(FFixupPayloadParameterNameEvent, UMovieSceneEventSectionBase*, UK2Node*, FName, FName);
 	DECLARE_MULTICAST_DELEGATE_TwoParams(FUpgradeLegacyEventEndpoint, UMovieSceneEventSectionBase*, UBlueprint*);
