Sequencer: Changes Sequencer Audiotemplates to check for a 0.2s desync (cvar Sequencer.Audio.MaxDesyncTolerance) between expected audio time and actual audio time. There did not appear to be a good way to see how long the Audio Component had been playing audio for, so we track locally how long it's been since we last requested Play. This lets us detect when the Time thinks it should have been playing the section for say, 35s but it's only been 5s since we last called Play (which maybe the case when jumping backwards in a replay that is then fast-forwarded)

diff --git a/Engine/Source/Runtime/MovieSceneTracks/Private/Evaluation/MovieSceneAudioTemplate.cpp b/Engine/Source/Runtime/MovieSceneTracks/Private/Evaluation/MovieSceneAudioTemplate.cpp
--- a/Engine/Source/Runtime/MovieSceneTracks/Private/Evaluation/MovieSceneAudioTemplate.cpp
+++ b/Engine/Source/Runtime/MovieSceneTracks/Private/Evaluation/MovieSceneAudioTemplate.cpp
@@ -21,6 +21,13 @@ DECLARE_CYCLE_STAT(TEXT("Audio Track Evaluate"), MovieSceneEval_AudioTrack_Evalu
 DECLARE_CYCLE_STAT(TEXT("Audio Track Tear Down"), MovieSceneEval_AudioTrack_TearDown, STATGROUP_MovieSceneEval);
 DECLARE_CYCLE_STAT(TEXT("Audio Track Token Execute"), MovieSceneEval_AudioTrack_TokenExecute, STATGROUP_MovieSceneEval);
 
+static float MaxSequenceAudioDesyncToleranceCVar = 0.2f;
+FAutoConsoleVariableRef CVarMaxSequenceAudioDesyncTolerance(
+	TEXT("Sequencer.Audio.MaxDesyncTolerance"),
+	MaxSequenceAudioDesyncToleranceCVar,
+	TEXT("Controls how many seconds an audio track can be out of sync in a Sequence before we attempt a time correction.\n"),
+	ECVF_Default);
+
 
 /** Stop audio from playing */
 struct FStopAudioPreAnimatedToken : IMovieScenePreAnimatedToken
@@ -80,11 +87,15 @@ struct FCachedAudioTrackData : IPersistentEvaluationData
 	TMap<FObjectKey, TMap<FObjectKey, TWeakObjectPtr<UAudioComponent>>> AudioComponentsByActorKey;
 	
 	FCachedAudioTrackData()
+		: SoundLastPlayedAtTime(0.f)
 	{
 		// Create the container for master tracks, which do not have an actor to attach to
 		AudioComponentsByActorKey.Add(FObjectKey(), TMap<FObjectKey, TWeakObjectPtr<UAudioComponent>>());
 	}
 
+	/** Set whenever we ask the Audio component to start playing a sound. Used to detect desyncs caused when Sequencer evaluates at more-than-real-time. */
+	float SoundLastPlayedAtTime;
+
 	UAudioComponent* GetAudioComponent(FObjectKey ActorKey, FObjectKey SectionKey)
 	{
 		if (TMap<FObjectKey, TWeakObjectPtr<UAudioComponent>>* Map = AudioComponentsByActorKey.Find(ActorKey))
@@ -318,8 +329,6 @@ struct FAudioSectionExecutionToken : IMovieSceneExecutionToken
 	{
 		Player.SavePreAnimatedState(AudioComponent, FStopAudioPreAnimatedToken::GetAnimTypeID(), FStopAudioPreAnimatedToken::FProducer());
 
-		bool bPlaySound = !AudioComponent.IsPlaying() || AudioComponent.Sound != AudioSection->GetSound();
-
 		float AudioVolume = 1.f;
 		AudioSection->GetSoundVolumeChannel().Evaluate(Context.GetTime(), AudioVolume);
 		AudioVolume = AudioVolume * AudioSection->EvaluateEasing(Context.GetTime());
@@ -335,7 +344,41 @@ struct FAudioSectionExecutionToken : IMovieSceneExecutionToken
 			AudioComponent.SetPitchMultiplier(PitchMultiplier);
 		}
 
-		if (bPlaySound)
+		float SectionStartTimeSeconds = (AudioSection->HasStartFrame() ? AudioSection->GetInclusiveStartFrame() : 0) / AudioSection->GetTypedOuter<UMovieScene>()->GetTickResolution();
+
+		FCachedAudioTrackData& TrackData = PersistentData.GetOrAddTrackData<FCachedAudioTrackData>();
+		const FFrameNumber AudioStartOffset = AudioSection->GetStartOffset();
+		USoundBase* Sound = AudioSection->GetSound();
+
+		float AudioTime = (Context.GetTime() / Context.GetFrameRate()) - SectionStartTimeSeconds + (float)Context.GetFrameRate().AsSeconds(AudioStartOffset);
+		if (AudioTime >= 0.f && Sound)
+		{
+			float Duration = Sound->GetDuration();
+			AudioTime = FMath::Fmod(AudioTime, Duration / 1000.f);
+		}
+
+		// If the audio component is not playing we (may) need a state change. If the audio component is playing
+		// the wrong sound then we need a state change. If the audio playback time is significantly out of sync 
+		// with the desired time then we need a state change.
+		bool bSoundNeedsStateChange = !AudioComponent.IsPlaying() || AudioComponent.Sound != Sound;
+		bool bSoundNeedsTimeSync = false;
+
+		if (Player.GetPlaybackContext()->GetWorld())
+		{
+			float CurrentGameTime = Player.GetPlaybackContext()->GetWorld()->GetAudioTimeSeconds();
+
+			// This tells us how much time has passed in the game world (and thus, reasonably, the audio playback)
+			// so if we calculate that we should be playing say, 15s into the section during evaluation, but
+			// we're only 5s since the last Play call, then we know we're out of sync. 
+			float GameTimeDelta = CurrentGameTime - TrackData.SoundLastPlayedAtTime;
+			if (FMath::Abs(GameTimeDelta - AudioTime) > MaxSequenceAudioDesyncToleranceCVar)
+			{
+				UE_LOG(LogMovieScene, Verbose, TEXT("Audio Component detected a significant mismatch in (assumed) playback time versus the desired time. Time since last play call: %6.2f(s) Desired Time: %6.2f(s). Component: %s sound: %s"), GameTimeDelta, AudioTime, *AudioComponent.GetName(), *GetNameSafe(AudioComponent.Sound));
+				bSoundNeedsTimeSync = true;
+			}
+		}
+
+		if (bSoundNeedsStateChange || bSoundNeedsTimeSync)
 		{
 			AudioComponent.bAllowSpatialization = bAllowSpatialization;
 
@@ -344,8 +387,20 @@ struct FAudioSectionExecutionToken : IMovieSceneExecutionToken
 				AudioComponent.AttenuationSettings = AudioSection->GetAttenuationSettings();
 			}
 
-			AudioComponent.Stop();
-			AudioComponent.SetSound(AudioSection->GetSound());
+			// Only call stop on the sound if it is actually playing. This prevents spamming
+			// stop calls when a sound cue with a duration of zero is played.
+			if (AudioComponent.IsPlaying() || bSoundNeedsTimeSync)
+			{
+				UE_LOG(LogMovieScene, Verbose, TEXT("Audio Component stopped due to needing a state change bIsPlaying: %d bNeedsTimeSync: %d. Component: %s sound: %s"), AudioComponent.IsPlaying(), bSoundNeedsTimeSync, *AudioComponent.GetName(), *AudioComponent.Sound->GetName());
+				AudioComponent.Stop();
+			}
+
+			// Only change the sound clip if it has actually changed. This calls Stop internally if needed.
+			if (AudioComponent.Sound != Sound)
+			{
+				UE_LOG(LogMovieScene, Verbose, TEXT("Audio Component calling SetSound due to new sound. Component: %s OldSound: %s NewSound: %s"), *AudioComponent.GetName(), *GetNameSafe(AudioComponent.Sound), *AudioComponent.Sound->GetName());
+				AudioComponent.SetSound(Sound);
+			}
 #if WITH_EDITOR
 			UObject* PlaybackContext = Player.GetPlaybackContext();
 			UWorld* World = PlaybackContext ? PlaybackContext->GetWorld() : nullptr;
@@ -360,14 +415,16 @@ struct FAudioSectionExecutionToken : IMovieSceneExecutionToken
 				AudioComponent.bIsUISound = false;
 			}
 
-			float SectionStartTimeSeconds = (AudioSection->HasStartFrame() ? AudioSection->GetInclusiveStartFrame() : 0) / AudioSection->GetTypedOuter<UMovieScene>()->GetTickResolution();
-
-			float AudioTime = (Context.GetTime() / Context.GetFrameRate()) - SectionStartTimeSeconds + (float)Context.GetFrameRate().AsSeconds(AudioSection->GetStartOffset());
-			if (AudioTime >= 0.f && AudioComponent.Sound)
+			if (AudioTime >= 0.f)
 			{
-				float Duration = AudioComponent.Sound->GetDuration();
-				AudioTime = FMath::Fmod(AudioTime, Duration/1000.f);
+				UE_LOG(LogMovieScene, Verbose, TEXT("Audio Component Play at Local Time: %6.2f CurrentTime: %6.2f(s) SectionStart: %6.2f(s), SoundDur: %6.2f OffsetIntoClip: %6.2f sound: %s"), AudioTime, (Context.GetTime() / Context.GetFrameRate()), SectionStartTimeSeconds, AudioComponent.Sound->GetDuration(), (float)Context.GetFrameRate().AsSeconds(AudioStartOffset), *AudioComponent.Sound->GetName());
 				AudioComponent.Play(AudioTime);
+
+				// Keep track of when we asked this audio clip to play (in game time) so that we can figure out if there's a significant desync in the future.
+				if (Player.GetPlaybackContext()->GetWorld())
+				{
+					TrackData.SoundLastPlayedAtTime = Player.GetPlaybackContext()->GetWorld()->GetAudioTimeSeconds() - AudioTime;
+				}
 			}
 
 			if (Context.GetStatus() == EMovieScenePlayerStatus::Scrubbing)
