Sequencer: Attach object binding for master track audio sections

diff --git a/Engine/Source/Runtime/MovieSceneTracks/Private/Evaluation/MovieSceneAudioTemplate.cpp b/Engine/Source/Runtime/MovieSceneTracks/Private/Evaluation/MovieSceneAudioTemplate.cpp
--- a/Engine/Source/Runtime/MovieSceneTracks/Private/Evaluation/MovieSceneAudioTemplate.cpp
+++ b/Engine/Source/Runtime/MovieSceneTracks/Private/Evaluation/MovieSceneAudioTemplate.cpp
@@ -178,13 +178,34 @@ struct FCachedAudioTrackData : IPersistentEvaluationData
 	}
 
 	/** Only to be called on the game thread */
-	UAudioComponent* AddMasterAudioComponentForRow(int32 RowIndex, FObjectKey SectionKey, UWorld* World, IMovieScenePlayer& Player)
+	UAudioComponent* AddMasterAudioComponentForRow(int32 RowIndex, FObjectKey SectionKey, UObject* PrincipalObject, UWorld* World, IMovieScenePlayer& Player)
 	{
 		UAudioComponent* ExistingComponent = GetAudioComponent(FObjectKey(), SectionKey);
 		if (!ExistingComponent)
 		{
 			USoundCue* TempPlaybackAudioCue = NewObject<USoundCue>();
-			ExistingComponent = FAudioDevice::CreateComponent(TempPlaybackAudioCue, FAudioDevice::FCreateComponentParams(World));
+
+			AActor* Actor = nullptr;
+			USceneComponent* SceneComponent = nullptr;
+
+			if (PrincipalObject)
+			{
+				if (PrincipalObject->IsA<AActor>())
+				{
+					Actor = Cast<AActor>(PrincipalObject);
+					SceneComponent = Actor->GetRootComponent();
+				}
+				else if (PrincipalObject->IsA<UActorComponent>())
+				{
+					UActorComponent* ActorComponent = Cast<UActorComponent>(PrincipalObject);
+					Actor = ActorComponent->GetOwner();
+					SceneComponent = Cast<USceneComponent>(ActorComponent);
+				}
+			}
+
+			FAudioDevice::FCreateComponentParams Params(Actor ? Actor->GetWorld() : World, Actor);
+
+			ExistingComponent = FAudioDevice::CreateComponent(TempPlaybackAudioCue, Params);
 
 			if (!ExistingComponent)
 			{
@@ -195,6 +216,11 @@ struct FCachedAudioTrackData : IPersistentEvaluationData
 			Player.SavePreAnimatedState(*ExistingComponent, FMovieSceneAnimTypeID::Unique(), FDestroyAudioPreAnimatedToken::FProducer());
 			
 			ExistingComponent->SetFlags(RF_Transient);
+			if (SceneComponent)
+			{
+				ExistingComponent->AttachToComponent(SceneComponent, FAttachmentTransformRules::KeepRelativeTransform);
+			}
+			
 			AudioComponentsByActorKey[FObjectKey()].Add(SectionKey, ExistingComponent);
 		}
 
@@ -252,8 +278,33 @@ struct FAudioSectionExecutionToken : IMovieSceneExecutionToken
 			UAudioComponent* AudioComponent = TrackData.GetAudioComponent(FObjectKey(), SectionKey);
 			if (!AudioComponent)
 			{
+				FMovieSceneObjectBindingID AttachBindingID = AudioSection->GetAttachBindingID();
+
+				UObject* AttachObject = nullptr;
+
+				if (AttachBindingID.IsValid())
+				{
+					FMovieSceneSequenceID SequenceID = Operand.SequenceID;
+					if (AttachBindingID.GetSequenceID().IsValid())
+					{
+						// Ensure that this ID is resolvable from the root, based on the current local sequence ID
+						FMovieSceneObjectBindingID RootBindingID = AttachBindingID.ResolveLocalToRoot(SequenceID, Player.GetEvaluationTemplate().GetHierarchy());
+						SequenceID = RootBindingID.GetSequenceID();
+					}
+
+					// If the transform is set, otherwise use the bound actor's transform
+					FMovieSceneEvaluationOperand ObjectOperand(SequenceID, AttachBindingID.GetGuid());
+
+					TArrayView<TWeakObjectPtr<>> Objects = Player.FindBoundObjects(ObjectOperand);
+					// Only ever deal with one object
+					if (Objects.Num() > 0)
+					{
+						AttachObject = Objects[0].Get();
+					}
+				}
+
 				// Initialize the sound
-				AudioComponent = TrackData.AddMasterAudioComponentForRow(AudioSection->GetRowIndex(), SectionKey, PlaybackContext ? PlaybackContext->GetWorld() : nullptr, Player);
+				AudioComponent = TrackData.AddMasterAudioComponentForRow(AudioSection->GetRowIndex(), SectionKey, AttachObject, PlaybackContext ? PlaybackContext->GetWorld() : nullptr, Player);
 
 				if (AudioComponent)
 				{
@@ -274,7 +325,7 @@ struct FAudioSectionExecutionToken : IMovieSceneExecutionToken
 
 			if (AudioComponent)
 			{
-				EnsureAudioIsPlaying(*AudioComponent, PersistentData, Context, false, Player);
+				EnsureAudioIsPlaying(*AudioComponent, PersistentData, Context, AudioComponent->GetAttachParent() != nullptr, Player);
 			}
 		}
 
diff --git a/Engine/Source/Runtime/MovieSceneTracks/Public/Sections/MovieSceneAudioSection.h b/Engine/Source/Runtime/MovieSceneTracks/Public/Sections/MovieSceneAudioSection.h
--- a/Engine/Source/Runtime/MovieSceneTracks/Public/Sections/MovieSceneAudioSection.h
+++ b/Engine/Source/Runtime/MovieSceneTracks/Public/Sections/MovieSceneAudioSection.h
@@ -103,6 +103,20 @@ public:
 		return AttenuationSettings;
 	}
 
+	/** Gets the attach binding for this Audio section */
+	UFUNCTION(BlueprintPure, Category = "Movie Scene Section")
+	const FMovieSceneObjectBindingID& GetAttachBindingID() const
+	{
+		return AttachBindingID;
+	}
+
+	/** Sets the attach binding for this Audio section */
+	UFUNCTION(BlueprintCallable, Category = "Movie Scene Section")
+	void SetAttachBindingID(const FMovieSceneObjectBindingID& InAttachBindingID)
+	{
+		AttachBindingID = InAttachBindingID;
+	}
+
 	/** ~UObject interface */
 	virtual void PostLoad() override;
 
@@ -194,6 +208,10 @@ private:
 	UPROPERTY( EditAnywhere, Category="Attenuation" )
 	class USoundAttenuation* AttenuationSettings;
 
+	/** The object binding that this audio is attached to */
+	UPROPERTY(EditAnywhere, Category = "Audio")
+	FMovieSceneObjectBindingID AttachBindingID;
+
 	/** Called when subtitles are sent to the SubtitleManager.  Set this delegate if you want to hijack the subtitles for other purposes */
 	UPROPERTY()
 	FOnQueueSubtitles OnQueueSubtitles;
