Sequencer: Corrected code that was incorrectly applying interface ptrs when triggering events on interfaces

#note Copy from 4.24[REDACTED_HANDLE]

diff --git a/Engine/Source/Runtime/MovieSceneTracks/Private/Evaluation/MovieSceneEventTemplate.cpp b/Engine/Source/Runtime/MovieSceneTracks/Private/Evaluation/MovieSceneEventTemplate.cpp
--- a/Engine/Source/Runtime/MovieSceneTracks/Private/Evaluation/MovieSceneEventTemplate.cpp
+++ b/Engine/Source/Runtime/MovieSceneTracks/Private/Evaluation/MovieSceneEventTemplate.cpp
@@ -325,7 +325,7 @@ struct FEventTriggerExecutionToken
 				FScriptInterface Interface;
 				Interface.SetObject(BoundObject);
 				Interface.SetInterface(BoundObject->GetInterfaceAddress(InterfaceParameter->InterfaceClass));
-				InterfaceParameter->SetPropertyValue_InContainer(Parameters, BoundObject);
+				InterfaceParameter->SetPropertyValue_InContainer(Parameters, Interface);
 				return true;
 			}
 
