Sequencer: Fixed looped audio playback not starting at the correct time if the sound wave was also looped

#note Copy from 4.24[REDACTED_HANDLE]

diff --git a/Engine/Source/Runtime/MovieSceneTracks/Private/Evaluation/MovieSceneAudioTemplate.cpp b/Engine/Source/Runtime/MovieSceneTracks/Private/Evaluation/MovieSceneAudioTemplate.cpp
--- a/Engine/Source/Runtime/MovieSceneTracks/Private/Evaluation/MovieSceneAudioTemplate.cpp
+++ b/Engine/Source/Runtime/MovieSceneTracks/Private/Evaluation/MovieSceneAudioTemplate.cpp
@@ -12,6 +12,7 @@
 #include "ActiveSound.h"
 #include "Sections/MovieSceneAudioSection.h"
 #include "Tracks/MovieSceneAudioTrack.h"
+#include "Sound/SoundWave.h"
 #include "MovieScene.h"
 #include "Evaluation/MovieSceneEvaluation.h"
 #include "IMovieScenePlayer.h"
@@ -416,7 +417,7 @@ struct FAudioSectionExecutionToken : IMovieSceneExecutionToken
 			float AudioTime = (Context.GetTime() / Context.GetFrameRate()) - SectionStartTimeSeconds + (float)Context.GetFrameRate().AsSeconds(AudioSection->GetStartOffset());
 			if (AudioTime >= 0.f && AudioComponent.Sound)
 			{
-				float Duration = AudioComponent.Sound->GetDuration();
+				const float Duration = MovieSceneHelpers::GetSoundDuration(AudioComponent.Sound);
 				AudioTime = FMath::Fmod(AudioTime, Duration);
 				AudioComponent.Play(AudioTime);
 			}
