Sequencer: Fix audio duration, resolved changes from 1017477 which stomped on fixes from 10800746, 10286096

diff --git a/Engine/Source/Runtime/MovieSceneTracks/Private/Evaluation/MovieSceneAudioTemplate.cpp b/Engine/Source/Runtime/MovieSceneTracks/Private/Evaluation/MovieSceneAudioTemplate.cpp
--- a/Engine/Source/Runtime/MovieSceneTracks/Private/Evaluation/MovieSceneAudioTemplate.cpp
+++ b/Engine/Source/Runtime/MovieSceneTracks/Private/Evaluation/MovieSceneAudioTemplate.cpp
@@ -405,8 +405,8 @@ struct FAudioSectionExecutionToken : IMovieSceneExecutionToken
 		float AudioTime = (Context.GetTime() / Context.GetFrameRate()) - SectionStartTimeSeconds + (float)Context.GetFrameRate().AsSeconds(AudioStartOffset);
 		if (AudioTime >= 0.f && Sound)
 		{
-			float Duration = Sound->GetDuration();
-			AudioTime = FMath::Fmod(AudioTime, Duration / 1000.f);
+			const float Duration = MovieSceneHelpers::GetSoundDuration(Sound);
+			AudioTime = Duration > 0.f ? FMath::Fmod(AudioTime, Duration) : AudioTime;
 		}
 
 		// If the audio component is not playing we (may) need a state change. If the audio component is playing
