Sequencer: Return early if the event is specified by an object binding, but there aren't any bound objects (ie. spawnable hasn't been spawned)

diff --git a/Engine/Source/Runtime/MovieSceneTracks/Private/Evaluation/MovieSceneEventTemplate.cpp b/Engine/Source/Runtime/MovieSceneTracks/Private/Evaluation/MovieSceneEventTemplate.cpp
--- a/Engine/Source/Runtime/MovieSceneTracks/Private/Evaluation/MovieSceneEventTemplate.cpp
+++ b/Engine/Source/Runtime/MovieSceneTracks/Private/Evaluation/MovieSceneEventTemplate.cpp
@@ -199,6 +199,12 @@ struct FEventTriggerExecutionToken
 					EventContexts.Add(EventContext);
 				}
 			}
+
+			// If there aren't any bound objects (ie. spawnable hasn't been spawned), no need to trigger the event
+			if (EventContexts.Num() == 0)
+			{
+				return;
+			}
 		}
 
 		// If we have specified event receivers
