Sequencer: Added the ability to animate the actor to attach audio to per audio section

Removed the attach object binding property in favor of FMovieSceneActorReferenceData.

diff --git a/Engine/Source/Runtime/MovieSceneTracks/Private/Evaluation/MovieSceneAudioTemplate.cpp b/Engine/Source/Runtime/MovieSceneTracks/Private/Evaluation/MovieSceneAudioTemplate.cpp
--- a/Engine/Source/Runtime/MovieSceneTracks/Private/Evaluation/MovieSceneAudioTemplate.cpp
+++ b/Engine/Source/Runtime/MovieSceneTracks/Private/Evaluation/MovieSceneAudioTemplate.cpp
@@ -190,32 +190,14 @@ struct FCachedAudioTrackData : IPersistentEvaluationData
 	}
 
 	/** Only to be called on the game thread */
-	UAudioComponent* AddMasterAudioComponentForRow(int32 RowIndex, FObjectKey SectionKey, UObject* PrincipalObject, UWorld* World, IMovieScenePlayer& Player)
+	UAudioComponent* AddMasterAudioComponentForRow(int32 RowIndex, FObjectKey SectionKey, UWorld* World, IMovieScenePlayer& Player)
 	{
 		UAudioComponent* ExistingComponent = GetAudioComponent(FObjectKey(), SectionKey);
 		if (!ExistingComponent)
 		{
 			USoundCue* TempPlaybackAudioCue = NewObject<USoundCue>();
 
-			AActor* Actor = nullptr;
-			USceneComponent* SceneComponent = nullptr;
-
-			if (PrincipalObject)
-			{
-				if (PrincipalObject->IsA<AActor>())
-				{
-					Actor = Cast<AActor>(PrincipalObject);
-					SceneComponent = Actor->GetRootComponent();
-				}
-				else if (PrincipalObject->IsA<UActorComponent>())
-				{
-					UActorComponent* ActorComponent = Cast<UActorComponent>(PrincipalObject);
-					Actor = ActorComponent->GetOwner();
-					SceneComponent = Cast<USceneComponent>(ActorComponent);
-				}
-			}
-
-			FAudioDevice::FCreateComponentParams Params(Actor ? Actor->GetWorld() : World, Actor);
+			FAudioDevice::FCreateComponentParams Params(World);
 
 			ExistingComponent = FAudioDevice::CreateComponent(TempPlaybackAudioCue, Params);
 
@@ -228,11 +210,7 @@ struct FCachedAudioTrackData : IPersistentEvaluationData
 			Player.SavePreAnimatedState(*ExistingComponent, FMovieSceneAnimTypeID::Unique(), FDestroyAudioPreAnimatedToken::FProducer());
 			
 			ExistingComponent->SetFlags(RF_Transient);
-			if (SceneComponent)
-			{
-				ExistingComponent->AttachToComponent(SceneComponent, FAttachmentTransformRules::KeepRelativeTransform);
-			}
-			
+
 			AudioComponentsByActorKey[FObjectKey()].Add(SectionKey, ExistingComponent);
 		}
 
@@ -287,36 +265,52 @@ struct FAudioSectionExecutionToken : IMovieSceneExecutionToken
 		{
 			UObject* PlaybackContext = Player.GetPlaybackContext();
 
-			UAudioComponent* AudioComponent = TrackData.GetAudioComponent(FObjectKey(), SectionKey);
-			if (!AudioComponent)
+			const FMovieSceneActorReferenceData& AttachActorData = AudioSection->GetAttachActorData();
+
+			USceneComponent* AttachComponent = nullptr;
+
+			FMovieSceneObjectBindingID AttachBindingID = AttachActorData.Evaluate(Context.GetTime()).Object;
+			if (AttachBindingID.IsValid())
 			{
-				FMovieSceneObjectBindingID AttachBindingID = AudioSection->GetAttachBindingID();
+				FMovieSceneSequenceID SequenceID = Operand.SequenceID;
+				if (AttachBindingID.GetSequenceID().IsValid())
+				{
+					// Ensure that this ID is resolvable from the root, based on the current local sequence ID
+					FMovieSceneObjectBindingID RootBindingID = AttachBindingID.ResolveLocalToRoot(SequenceID, Player.GetEvaluationTemplate().GetHierarchy());
+					SequenceID = RootBindingID.GetSequenceID();
+				}
 
-				UObject* AttachObject = nullptr;
+				// If the transform is set, otherwise use the bound actor's transform
+				FMovieSceneEvaluationOperand ObjectOperand(SequenceID, AttachBindingID.GetGuid());
 
-				if (AttachBindingID.IsValid())
+				for (TWeakObjectPtr<> WeakObject : Player.FindBoundObjects(ObjectOperand))
 				{
-					FMovieSceneSequenceID SequenceID = Operand.SequenceID;
-					if (AttachBindingID.GetSequenceID().IsValid())
+					UObject* AttachObject = WeakObject.Get();
+					if (AttachObject)
 					{
-						// Ensure that this ID is resolvable from the root, based on the current local sequence ID
-						FMovieSceneObjectBindingID RootBindingID = AttachBindingID.ResolveLocalToRoot(SequenceID, Player.GetEvaluationTemplate().GetHierarchy());
-						SequenceID = RootBindingID.GetSequenceID();
+						if (AttachObject->IsA<AActor>())
+						{
+							AActor* Actor = Cast<AActor>(AttachObject);
+							AttachComponent = Actor->GetRootComponent();
+						}
+						else if (AttachObject->IsA<UActorComponent>())
+						{
+							UActorComponent* ActorComponent = Cast<UActorComponent>(AttachObject);
+							AttachComponent = Cast<USceneComponent>(ActorComponent);
+						}
 					}
-
-					// If the transform is set, otherwise use the bound actor's transform
-					FMovieSceneEvaluationOperand ObjectOperand(SequenceID, AttachBindingID.GetGuid());
-
-					TArrayView<TWeakObjectPtr<>> Objects = Player.FindBoundObjects(ObjectOperand);
-					// Only ever deal with one object
-					if (Objects.Num() > 0)
+					if (AttachComponent)
 					{
-						AttachObject = Objects[0].Get();
+						break;
 					}
 				}
+			}
 
+			UAudioComponent* AudioComponent = TrackData.GetAudioComponent(FObjectKey(), SectionKey);
+			if (!AudioComponent)
+			{
 				// Initialize the sound
-				AudioComponent = TrackData.AddMasterAudioComponentForRow(AudioSection->GetRowIndex(), SectionKey, AttachObject, PlaybackContext ? PlaybackContext->GetWorld() : nullptr, Player);
+				AudioComponent = TrackData.AddMasterAudioComponentForRow(AudioSection->GetRowIndex(), SectionKey, PlaybackContext ? PlaybackContext->GetWorld() : nullptr, Player);
 
 				if (AudioComponent)
 				{
@@ -337,6 +331,11 @@ struct FAudioSectionExecutionToken : IMovieSceneExecutionToken
 
 			if (AudioComponent)
 			{
+				if (AttachComponent)
+				{
+					AudioComponent->AttachToComponent(AttachComponent, FAttachmentTransformRules::KeepRelativeTransform);
+				}
+
 				EnsureAudioIsPlaying(*AudioComponent, PersistentData, Context, AudioComponent->GetAttachParent() != nullptr, Player);
 			}
 		}
@@ -492,7 +491,7 @@ struct FAudioSectionExecutionToken : IMovieSceneExecutionToken
 			{
 				DECLARE_CYCLE_STAT(TEXT("FAudioThreadTask.MovieSceneUpdateAudioTransform"), STAT_MovieSceneUpdateAudioTransform, STATGROUP_TaskGraphTasks);
 
-				const FTransform ActorTransform = AudioComponent.GetOwner()->GetTransform();
+				const FTransform ActorTransform = AudioComponent.GetComponentTransform();
 				const uint64 ActorComponentID = AudioComponent.GetAudioComponentID();
 				FAudioThread::RunCommandOnAudioThread([AudioDevice, ActorComponentID, ActorTransform]()
 				{
diff --git a/Engine/Source/Runtime/MovieSceneTracks/Private/Sections/MovieSceneAudioSection.cpp b/Engine/Source/Runtime/MovieSceneTracks/Private/Sections/MovieSceneAudioSection.cpp
--- a/Engine/Source/Runtime/MovieSceneTracks/Private/Sections/MovieSceneAudioSection.cpp
+++ b/Engine/Source/Runtime/MovieSceneTracks/Private/Sections/MovieSceneAudioSection.cpp
@@ -16,9 +16,10 @@ struct FAudioChannelEditorData
 	{
 		Data[0].SetIdentifiers("Volume", NSLOCTEXT("MovieSceneAudioSection", "SoundVolumeText", "Volume"));
 		Data[1].SetIdentifiers("Pitch", NSLOCTEXT("MovieSceneAudioSection", "PitchText", "Pitch"));
+		Data[2].SetIdentifiers("AttachActor", NSLOCTEXT("MovieSceneAudioSection", "AttachActorText", "Attach"));
 	}
 
-	FMovieSceneChannelMetaData Data[2];
+	FMovieSceneChannelMetaData Data[3];
 };
 
 #endif // WITH_EDITOR
@@ -56,11 +57,13 @@ UMovieSceneAudioSection::UMovieSceneAudioSection( const FObjectInitializer& Obje
 	static const FAudioChannelEditorData EditorData;
 	Channels.Add(SoundVolume,     EditorData.Data[0], TMovieSceneExternalValue<float>());
 	Channels.Add(PitchMultiplier, EditorData.Data[1], TMovieSceneExternalValue<float>());
+	Channels.Add(AttachActorData, EditorData.Data[2]);
 
 #else
 
 	Channels.Add(SoundVolume);
 	Channels.Add(PitchMultiplier);
+	Channels.Add(AttachActorData);
 
 #endif
 
diff --git a/Engine/Source/Runtime/MovieSceneTracks/Public/Sections/MovieSceneAudioSection.h b/Engine/Source/Runtime/MovieSceneTracks/Public/Sections/MovieSceneAudioSection.h
--- a/Engine/Source/Runtime/MovieSceneTracks/Public/Sections/MovieSceneAudioSection.h
+++ b/Engine/Source/Runtime/MovieSceneTracks/Public/Sections/MovieSceneAudioSection.h
@@ -8,6 +8,7 @@
 #include "Runtime/Engine/Classes/Components/AudioComponent.h"
 #include "Sound/SoundAttenuation.h"
 #include "Channels/MovieSceneFloatChannel.h"
+#include "Sections/MovieSceneActorReferenceSection.h"
 #include "MovieSceneAudioSection.generated.h"
 
 class USoundBase;
@@ -103,19 +104,7 @@ public:
 		return AttenuationSettings;
 	}
 
-	/** Gets the attach binding for this Audio section */
-	UFUNCTION(BlueprintPure, Category = "Movie Scene Section")
-	const FMovieSceneObjectBindingID& GetAttachBindingID() const
-	{
-		return AttachBindingID;
-	}
-
-	/** Sets the attach binding for this Audio section */
-	UFUNCTION(BlueprintCallable, Category = "Movie Scene Section")
-	void SetAttachBindingID(const FMovieSceneObjectBindingID& InAttachBindingID)
-	{
-		AttachBindingID = InAttachBindingID;
-	}
+	const FMovieSceneActorReferenceData& GetAttachActorData() const { return AttachActorData; }
 
 	/** ~UObject interface */
 	virtual void PostLoad() override;
@@ -197,6 +186,9 @@ private:
 	UPROPERTY( )
 	FMovieSceneFloatChannel PitchMultiplier;
 
+	UPROPERTY()
+	FMovieSceneActorReferenceData AttachActorData;
+
 	UPROPERTY( EditAnywhere, Category="Audio" )
 	bool bSuppressSubtitles;
 
@@ -208,10 +200,6 @@ private:
 	UPROPERTY( EditAnywhere, Category="Attenuation" )
 	class USoundAttenuation* AttenuationSettings;
 
-	/** The object binding that this audio is attached to */
-	UPROPERTY(EditAnywhere, Category = "Audio")
-	FMovieSceneObjectBindingID AttachBindingID;
-
 	/** Called when subtitles are sent to the SubtitleManager.  Set this delegate if you want to hijack the subtitles for other purposes */
 	UPROPERTY()
 	FOnQueueSubtitles OnQueueSubtitles;
