Sequencer: More fixes for camera cut track taking control.

The previous fix had broken another case :(
Test with both sub-scene and shot track, in a couple of cases each.

diff --git a/Engine/Source/Runtime/MovieSceneTracks/Private/Evaluation/MovieSceneCameraCutTemplate.cpp b/Engine/Source/Runtime/MovieSceneTracks/Private/Evaluation/MovieSceneCameraCutTemplate.cpp
--- a/Engine/Source/Runtime/MovieSceneTracks/Private/Evaluation/MovieSceneCameraCutTemplate.cpp
+++ b/Engine/Source/Runtime/MovieSceneTracks/Private/Evaluation/MovieSceneCameraCutTemplate.cpp
@@ -128,17 +128,22 @@ namespace MovieScene
 	struct FBlendedCameraCut
 	{
 		FMovieSceneObjectBindingID CameraBindingID;
+		FMovieSceneSequenceID OperandSequenceID;
 
 		FBlendedCameraCutEasingInfo EaseIn;
 		FBlendedCameraCutEasingInfo EaseOut;
 		bool bIsFinalCut = false;
 
 		FMovieSceneObjectBindingID PreviousCameraBindingID;
+		FMovieSceneSequenceID PreviousOperandSequenceID;
 
 		float PreviewBlendFactor = -1.f;
 
 		FBlendedCameraCut() {}
-		FBlendedCameraCut(FMovieSceneObjectBindingID InCameraBindingID) : CameraBindingID(InCameraBindingID) {}
+		FBlendedCameraCut(FMovieSceneObjectBindingID InCameraBindingID, FMovieSceneSequenceID InOperandSequenceID) 
+			: CameraBindingID(InCameraBindingID)
+			, OperandSequenceID(InOperandSequenceID)
+		{}
 
 		FBlendedCameraCut& Resolve(TMovieSceneInitialValueStore<FBlendedCameraCut>& InitialValueStore)
 		{
@@ -163,18 +168,20 @@ namespace MovieScene
 			return TypeID;
 		}
 
-		static UObject* FindBoundObject(FMovieSceneObjectBindingID BindingID, IMovieScenePlayer& Player)
+		static UObject* FindBoundObject(FMovieSceneObjectBindingID BindingID, FMovieSceneSequenceID SequenceID, IMovieScenePlayer& Player)
 		{
-			if (BindingID.IsValid())
+			FMovieSceneSequenceID ResolvedSequenceID = SequenceID;
+			if (BindingID.GetSequenceID().IsValid())
 			{
-				FMovieSceneObjectBindingID RootBindingID = BindingID.ResolveLocalToRoot(BindingID.GetSequenceID(), Player.GetEvaluationTemplate().GetHierarchy());
+				FMovieSceneObjectBindingID RootBindingID = BindingID.ResolveLocalToRoot(SequenceID, Player.GetEvaluationTemplate().GetHierarchy());
+				ResolvedSequenceID = RootBindingID.GetSequenceID();
+			}
 
-				FMovieSceneEvaluationOperand Operand(RootBindingID.GetSequenceID(), RootBindingID.GetGuid());
-				TArrayView<TWeakObjectPtr<>> Objects = Player.FindBoundObjects(Operand);
-				if (Objects.Num() > 0)
-				{
-					return Objects[0].Get();
-				}
+			FMovieSceneEvaluationOperand Operand(ResolvedSequenceID, BindingID.GetGuid());
+			TArrayView<TWeakObjectPtr<>> Objects = Player.FindBoundObjects(Operand);
+			if (Objects.Num() > 0)
+			{
+				return Objects[0].Get();
 			}
 			return nullptr;
 		}
@@ -183,7 +190,7 @@ namespace MovieScene
 		{
 			OriginalStack.SavePreAnimatedState(Player, GetCameraCutTypeID(), FCameraCutPreAnimatedTokenProducer());
 
-			UObject* CameraActor = FindBoundObject(InFinalValue.CameraBindingID, Player);
+			UObject* CameraActor = FindBoundObject(InFinalValue.CameraBindingID, InFinalValue.OperandSequenceID, Player);
 
 			FCameraCutTrackData& CameraCutCache = PersistentData.GetOrAddTrackData<FCameraCutTrackData>();
 
@@ -193,7 +200,7 @@ namespace MovieScene
 			CameraCutParams.BlendType = InFinalValue.EaseIn.BlendType;
 
 #if WITH_EDITOR
-			UObject* PreviousCameraActor = FindBoundObject(InFinalValue.PreviousCameraBindingID, Player);
+			UObject* PreviousCameraActor = FindBoundObject(InFinalValue.PreviousCameraBindingID, InFinalValue.PreviousOperandSequenceID, Player);
 			CameraCutParams.PreviousCameraObject = PreviousCameraActor;
 			CameraCutParams.PreviewBlendFactor = InFinalValue.PreviewBlendFactor;
 #endif
@@ -235,8 +242,10 @@ namespace MovieScene
 		else
 		{
 			FMovieSceneObjectBindingID PreviousCameraBindingID = OutBlend.CameraBindingID;
+			FMovieSceneSequenceID PreviousOperandSequenceID = OutBlend.OperandSequenceID;
 			OutBlend = InValue;
 			OutBlend.PreviousCameraBindingID = PreviousCameraBindingID;
+			OutBlend.PreviousOperandSequenceID = PreviousOperandSequenceID;
 		}
 	}
 }
@@ -271,7 +280,7 @@ void FMovieSceneCameraCutSectionTemplate::Evaluate(const FMovieSceneEvaluationOp
 		// For now we only look at how long the camera blend is supposed to be, and we pass that on to
 		// the player controller via the execution token. Later we'll need to actually drive the blend
 		// ourselves so that the curve itself is actually matching.
-		MovieScene::FBlendedCameraCut Params(FMovieSceneObjectBindingID(CameraBindingID.GetGuid(), Operand.SequenceID));
+		MovieScene::FBlendedCameraCut Params(CameraBindingID, Operand.SequenceID);
 		Params.bIsFinalCut = bIsFinalSection;
 		if (SourceSectionPtr.IsValid())
 		{
