Sequencer: Fix crash copying an audio track. Reconstruct the channel proxies when the audio track has been renamed since it's being pasted from a transient outer to the movie scene.

In AddGivenMasterTrack, add the track to the movie scene's master tracks before renaming.

diff --git a/Engine/Source/Runtime/MovieSceneTracks/Private/Tracks/MovieSceneAudioTrack.cpp b/Engine/Source/Runtime/MovieSceneTracks/Private/Tracks/MovieSceneAudioTrack.cpp
--- a/Engine/Source/Runtime/MovieSceneTracks/Private/Tracks/MovieSceneAudioTrack.cpp
+++ b/Engine/Source/Runtime/MovieSceneTracks/Private/Tracks/MovieSceneAudioTrack.cpp
@@ -108,7 +108,8 @@ UMovieSceneSection* UMovieSceneAudioTrack::AddNewSoundOnRow(USoundBase* Sound, F
 
 bool UMovieSceneAudioTrack::IsAMasterTrack() const
 {
-	return Cast<UMovieScene>(GetOuter())->IsAMasterTrack(*this);
+	UMovieScene* MovieScene = Cast<UMovieScene>(GetOuter());
+	return MovieScene ? MovieScene->IsAMasterTrack(*this) : false;
 }
 
 
@@ -130,5 +131,16 @@ UMovieSceneSection* UMovieSceneAudioTrack::CreateNewSection()
 	return NewObject<UMovieSceneAudioSection>(this, NAME_None, RF_Transactional);
 }
 
+void UMovieSceneAudioTrack::PostRename(UObject* OldOuter, const FName OldName)
+{
+	// If this audio track is outered to something new, update the channel proxy because the channel can depend on whether this is a master track or not
+	for (UMovieSceneSection* Section : AudioSections)
+	{
+		if (UMovieSceneAudioSection* AudioSection =Cast<UMovieSceneAudioSection>(Section))
+		{
+			AudioSection->UpdateChannelProxy();
+		}
+	}
+}
 
 #undef LOCTEXT_NAMESPACE
diff --git a/Engine/Source/Runtime/MovieSceneTracks/Public/Tracks/MovieSceneAudioTrack.h b/Engine/Source/Runtime/MovieSceneTracks/Public/Tracks/MovieSceneAudioTrack.h
--- a/Engine/Source/Runtime/MovieSceneTracks/Public/Tracks/MovieSceneAudioTrack.h
+++ b/Engine/Source/Runtime/MovieSceneTracks/Public/Tracks/MovieSceneAudioTrack.h
@@ -57,6 +57,8 @@ public:
 	virtual FMovieSceneTrackRowSegmentBlenderPtr GetRowSegmentBlender() const override;
 	virtual UMovieSceneSection* CreateNewSection() override;
 
+	virtual void PostRename(UObject* OldOuter, const FName OldName) override;
+
 private:
 
 	/** List of all master audio sections */
