Sequencer:   Need to make sure we re-Construct the Animation Sequence Nodes when entering a section since it may have been changed to a pose snapshot node. Also made sure ConstructNodes really just sets them up lightly since it's used by both SequencerInst/ControlRigInst.

diff --git a/Engine/Plugins/Experimental/ControlRig/Source/ControlRig/Private/Sequencer/ControlRigLayerInstance.cpp b/Engine/Plugins/Experimental/ControlRig/Source/ControlRig/Private/Sequencer/ControlRigLayerInstance.cpp
--- a/Engine/Plugins/Experimental/ControlRig/Source/ControlRig/Private/Sequencer/ControlRigLayerInstance.cpp
+++ b/Engine/Plugins/Experimental/ControlRig/Source/ControlRig/Private/Sequencer/ControlRigLayerInstance.cpp
@@ -35,6 +35,11 @@ void UControlRigLayerInstance::UpdateAnimTrack(UAnimSequenceBase* InAnimSequence
 	GetProxyOnGameThread<FControlRigLayerInstanceProxy>().UpdateAnimTrack(InAnimSequence, SequenceId, InFromPosition, InToPosition, Weight, bFireNotifies);
 }
 
+void UControlRigLayerInstance::ConstructNodes()
+{
+	GetProxyOnGameThread<FControlRigLayerInstanceProxy>().ConstructNodes();
+}
+
 void UControlRigLayerInstance::ResetNodes()
 {
 	GetProxyOnGameThread<FControlRigLayerInstanceProxy>().ResetNodes();
diff --git a/Engine/Plugins/Experimental/ControlRig/Source/ControlRig/Private/Sequencer/ControlRigLayerInstanceProxy.h b/Engine/Plugins/Experimental/ControlRig/Source/ControlRig/Private/Sequencer/ControlRigLayerInstanceProxy.h
--- a/Engine/Plugins/Experimental/ControlRig/Source/ControlRig/Private/Sequencer/ControlRigLayerInstanceProxy.h
+++ b/Engine/Plugins/Experimental/ControlRig/Source/ControlRig/Private/Sequencer/ControlRigLayerInstanceProxy.h
@@ -92,6 +92,8 @@ public:
 	virtual void ResetNodes();
 	/** Reset the pose in this instance*/
 	virtual void ResetPose();
+	/** Construct and link the base part of the blend tree */
+	virtual void ConstructNodes();
 
 	/** return first available control rig from the node it has */
 	UControlRig* GetFirstAvailableControlRig() const;
@@ -113,9 +115,6 @@ protected:
 	/** Find ControlRig node of the */
 	FAnimNode_ControlRig_ExternalSource* FindControlRigNode(int32 ControlRigID) const;
 
-	/** Construct and link the base part of the blend tree */
-	void ConstructNodes();
-
 	/** Input pose anim node */
 	FAnimNode_ControlRigInputPose InputPose;
 
diff --git a/Engine/Plugins/Experimental/ControlRig/Source/ControlRig/Public/Sequencer/ControlRigLayerInstance.h b/Engine/Plugins/Experimental/ControlRig/Source/ControlRig/Public/Sequencer/ControlRigLayerInstance.h
--- a/Engine/Plugins/Experimental/ControlRig/Source/ControlRig/Public/Sequencer/ControlRigLayerInstance.h
+++ b/Engine/Plugins/Experimental/ControlRig/Source/ControlRig/Public/Sequencer/ControlRigLayerInstance.h
@@ -32,6 +32,8 @@ public:
 	virtual void UpdateAnimTrack(UAnimSequenceBase* InAnimSequence, int32 SequenceId, float InFromPosition, float InToPosition, float Weight, bool bFireNotifies) override;
 	void RemoveAnimation(int32 SequenceId);
 
+	/** Construct all nodes in this instance */
+	virtual void ConstructNodes() override;
 	/** Reset all nodes in this instance */
 	virtual void ResetNodes() override;
 	/** Reset the pose in this instance*/
diff --git a/Engine/Source/Runtime/AnimGraphRuntime/Private/AnimSequencerInstance.cpp b/Engine/Source/Runtime/AnimGraphRuntime/Private/AnimSequencerInstance.cpp
--- a/Engine/Source/Runtime/AnimGraphRuntime/Private/AnimSequencerInstance.cpp
+++ b/Engine/Source/Runtime/AnimGraphRuntime/Private/AnimSequencerInstance.cpp
@@ -35,6 +35,11 @@ void UAnimSequencerInstance::UpdateAnimTrack(UAnimSequenceBase* InAnimSequence,
 	GetProxyOnGameThread<FAnimSequencerInstanceProxy>().UpdateAnimTrack(InAnimSequence, SequenceId, InFromPosition, InToPosition, Weight, bFireNotifies);
 }
 
+void UAnimSequencerInstance::ConstructNodes()
+{
+	GetProxyOnGameThread<FAnimSequencerInstanceProxy>().ConstructNodes();
+}
+
 void UAnimSequencerInstance::ResetNodes()
 {
 	GetProxyOnGameThread<FAnimSequencerInstanceProxy>().ResetNodes();
diff --git a/Engine/Source/Runtime/AnimGraphRuntime/Public/AnimSequencerInstance.h b/Engine/Source/Runtime/AnimGraphRuntime/Public/AnimSequencerInstance.h
--- a/Engine/Source/Runtime/AnimGraphRuntime/Public/AnimSequencerInstance.h
+++ b/Engine/Source/Runtime/AnimGraphRuntime/Public/AnimSequencerInstance.h
@@ -22,6 +22,9 @@ public:
 	virtual void UpdateAnimTrack(UAnimSequenceBase* InAnimSequence, int32 SequenceId, float InPosition, float Weight, bool bFireNotifies);
 	virtual void UpdateAnimTrack(UAnimSequenceBase* InAnimSequence, int32 SequenceId, float InFromPosition, float InToPosition, float Weight, bool bFireNotifies);
 
+	/** Construct all nodes in this instance */
+	virtual void ConstructNodes() override;
+
 	/** Reset all nodes in this instance */
 	virtual void ResetNodes() override;
 
diff --git a/Engine/Source/Runtime/AnimGraphRuntime/Public/AnimSequencerInstanceProxy.h b/Engine/Source/Runtime/AnimGraphRuntime/Public/AnimSequencerInstanceProxy.h
--- a/Engine/Source/Runtime/AnimGraphRuntime/Public/AnimSequencerInstanceProxy.h
+++ b/Engine/Source/Runtime/AnimGraphRuntime/Public/AnimSequencerInstanceProxy.h
@@ -91,6 +91,8 @@ public:
 	/** Reset the pose in this instance*/
 	virtual void ResetPose();
 
+	/** Construct and link the base part of the blend tree */
+	virtual void ConstructNodes();
 protected:
 	/** Find a player of a specified type */
 	template<typename Type>
@@ -105,9 +107,6 @@ protected:
 		return nullptr;
 	}
 
-	/** Construct and link the base part of the blend tree */
-	void ConstructNodes();
-
 	/** custom root node for this sequencer player. Didn't want to use RootNode in AnimInstance because it's used with lots of AnimBP functionality */
 	struct FAnimNode_ApplyAdditive SequencerRootNode;
 	struct FAnimNode_MultiWayBlend FullBodyBlendNode;
diff --git a/Engine/Source/Runtime/AnimGraphRuntime/Public/SequencerAnimationSupport.h b/Engine/Source/Runtime/AnimGraphRuntime/Public/SequencerAnimationSupport.h
--- a/Engine/Source/Runtime/AnimGraphRuntime/Public/SequencerAnimationSupport.h
+++ b/Engine/Source/Runtime/AnimGraphRuntime/Public/SequencerAnimationSupport.h
@@ -31,6 +31,9 @@ class ANIMGRAPHRUNTIME_API ISequencerAnimationSupport
 	virtual void UpdateAnimTrack(UAnimSequenceBase* InAnimSequence, int32 SequenceId, float InPosition, float Weight, bool bFireNotifies) = 0;
 	virtual void UpdateAnimTrack(UAnimSequenceBase* InAnimSequence, int32 SequenceId, float InFromPosition, float InToPosition, float Weight, bool bFireNotifies) = 0;
 
+	/** Construct all nodes in this instance */
+	virtual void ConstructNodes() = 0;
+
 	/** Reset all nodes in this instance */
 	virtual void ResetNodes() = 0;
 
diff --git a/Engine/Source/Runtime/MovieSceneTracks/Private/Evaluation/MovieSceneSkeletalAnimationTemplate.cpp b/Engine/Source/Runtime/MovieSceneTracks/Private/Evaluation/MovieSceneSkeletalAnimationTemplate.cpp
--- a/Engine/Source/Runtime/MovieSceneTracks/Private/Evaluation/MovieSceneSkeletalAnimationTemplate.cpp
+++ b/Engine/Source/Runtime/MovieSceneTracks/Private/Evaluation/MovieSceneSkeletalAnimationTemplate.cpp
@@ -270,11 +270,20 @@ namespace MovieScene
 			UAnimInstance* ExistingAnimInstance = GetSourceAnimInstance(SkeletalMeshComponent);
 			bool bWasCreated = false;
 			ISequencerAnimationSupport* SequencerInstance = FAnimCustomInstanceHelper::BindToSkeletalMeshComponent<UAnimSequencerInstance>(SkeletalMeshComponent,bWasCreated);
-			if (SequencerInstance && bWasCreated)
+			if (SequencerInstance)
 			{
-				SequencerInstance->SavePose();
+				if (bWasCreated)
+				{
+					SequencerInstance->SavePose();
+				}
+				else
+				{
+					SequencerInstance->ConstructNodes();
+				}
 			}
 
+
+
 			const bool bPreviewPlayback = ShouldUsePreviewPlayback(Player, *SkeletalMeshComponent);
 
 			const EMovieScenePlayerStatus::Type PlayerStatus = Player.GetPlaybackStatus();
