Sequencer: Fix missing blend-out handle on camera cut track when there's only one section.

diff --git a/Engine/Source/Runtime/MovieSceneTracks/Private/Tracks/MovieSceneCameraCutTrack.cpp b/Engine/Source/Runtime/MovieSceneTracks/Private/Tracks/MovieSceneCameraCutTrack.cpp
--- a/Engine/Source/Runtime/MovieSceneTracks/Private/Tracks/MovieSceneCameraCutTrack.cpp
+++ b/Engine/Source/Runtime/MovieSceneTracks/Private/Tracks/MovieSceneCameraCutTrack.cpp
@@ -115,15 +115,23 @@ EMovieSceneTrackEasingSupportFlags UMovieSceneCameraCutTrack::SupportsEasing(FMo
 	{
 		return EMovieSceneTrackEasingSupportFlags::None;
 	}
-	if (Params.ForSection != nullptr && Sections.Num() > 0)
+	if (Params.ForSection != nullptr)
 	{
-		if (Params.ForSection == Sections[0])
+		const int32 NumSections = Sections.Num();
+		if (NumSections == 1)
 		{
-			return EMovieSceneTrackEasingSupportFlags::AutomaticEasing | EMovieSceneTrackEasingSupportFlags::ManualEaseIn;
+			return EMovieSceneTrackEasingSupportFlags::AutomaticEasing | EMovieSceneTrackEasingSupportFlags::ManualEasing;
 		}
-		if (Params.ForSection == Sections.Last())
+		else if (NumSections > 1)
 		{
-			return EMovieSceneTrackEasingSupportFlags::AutomaticEasing | EMovieSceneTrackEasingSupportFlags::ManualEaseOut;
+			if (Params.ForSection == Sections[0])
+			{
+				return EMovieSceneTrackEasingSupportFlags::AutomaticEasing | EMovieSceneTrackEasingSupportFlags::ManualEaseIn;
+			}
+			if (Params.ForSection == Sections.Last())
+			{
+				return EMovieSceneTrackEasingSupportFlags::AutomaticEasing | EMovieSceneTrackEasingSupportFlags::ManualEaseOut;
+			}
 		}
 	}
 	return EMovieSceneTrackEasingSupportFlags::AutomaticEasing;
