Sequencer: Fix assert on camera cut blend-out evaluation.

diff --git a/Engine/Source/Runtime/MovieSceneTracks/Private/Evaluation/MovieSceneCameraCutTemplate.cpp b/Engine/Source/Runtime/MovieSceneTracks/Private/Evaluation/MovieSceneCameraCutTemplate.cpp
--- a/Engine/Source/Runtime/MovieSceneTracks/Private/Evaluation/MovieSceneCameraCutTemplate.cpp
+++ b/Engine/Source/Runtime/MovieSceneTracks/Private/Evaluation/MovieSceneCameraCutTemplate.cpp
@@ -316,7 +316,9 @@ void FMovieSceneCameraCutSectionTemplate::Evaluate(const FMovieSceneEvaluationOp
 				Params.EaseIn = Params.EaseOut;
 				Params.EaseOut = FBlendedCameraCutEasingInfo();
 				Params.PreviousCameraBindingID = Params.CameraBindingID;
+				Params.PreviousOperandSequenceID = Params.OperandSequenceID;
 				Params.CameraBindingID = FMovieSceneObjectBindingID();
+				Params.OperandSequenceID = FMovieSceneSequenceID();
 			}
 		}
 
