Sequencer: Fix to get component transform for external value. This fixes a bug were adding a key for a bp component wouldn't use the existing value of the component.

diff --git a/Engine/Source/Runtime/MovieSceneTracks/Private/Sections/MovieScene3DTransformSection.cpp b/Engine/Source/Runtime/MovieSceneTracks/Private/Sections/MovieScene3DTransformSection.cpp
--- a/Engine/Source/Runtime/MovieSceneTracks/Private/Sections/MovieScene3DTransformSection.cpp
+++ b/Engine/Source/Runtime/MovieSceneTracks/Private/Sections/MovieScene3DTransformSection.cpp
@@ -127,6 +127,10 @@ struct F3DTransformChannelEditorData
 				}
 			}
 		}
+		else if (USceneComponent* SceneComponent = Cast<USceneComponent>(&InObject))
+		{
+			return SceneComponent->GetRelativeTransform().GetTranslation();		
+		}
 		else if (AActor* Actor = Cast<AActor>(&InObject))
 		{
 			if (USceneComponent* RootComponent = Actor->GetRootComponent())
@@ -159,6 +163,10 @@ struct F3DTransformChannelEditorData
 				}
 			}
 		}
+		else if (USceneComponent* SceneComponent = Cast<USceneComponent>(&InObject))
+		{
+			return SceneComponent->GetRelativeRotation();
+		}
 		else if (AActor* Actor = Cast<AActor>(&InObject))
 		{
 			if (USceneComponent* RootComponent = Actor->GetRootComponent())
@@ -191,6 +199,10 @@ struct F3DTransformChannelEditorData
 				}
 			}
 		}
+		else if (USceneComponent* SceneComponent = Cast<USceneComponent>(&InObject))
+		{
+			return SceneComponent->GetRelativeTransform().GetScale3D();
+		}
 		else if (AActor* Actor = Cast<AActor>(&InObject))
 		{
 			if (USceneComponent* RootComponent = Actor->GetRootComponent())
