Sequencer: Add dedicated track for playing camera shakes on a camera shake source component.

This allows us to have a place to write editor-only code that will preview the shake while working in sequencer.

diff --git a/Engine/Source/Editor/MovieSceneTools/Private/MovieSceneToolsModule.cpp b/Engine/Source/Editor/MovieSceneTools/Private/MovieSceneToolsModule.cpp
--- a/Engine/Source/Editor/MovieSceneTools/Private/MovieSceneToolsModule.cpp
+++ b/Engine/Source/Editor/MovieSceneTools/Private/MovieSceneToolsModule.cpp
@@ -48,6 +48,7 @@
 #include "TrackEditors/MaterialParameterCollectionTrackEditor.h"
 #include "TrackEditors/ObjectPropertyTrackEditor.h"
 #include "TrackEditors/PrimitiveMaterialTrackEditor.h"
+#include "TrackEditors/CameraShakeSourceShakeTrackEditor.h"
 
 #include "MovieSceneBuiltInEasingFunctionCustomization.h"
 #include "MovieSceneObjectBindingIDCustomization.h"
@@ -123,6 +124,7 @@ void FMovieSceneToolsModule::StartupModule()
 		CameraShakeTrackCreateEditorHandle = SequencerModule.RegisterTrackEditor(FOnCreateTrackEditor::CreateStatic(&FCameraShakeTrackEditor::CreateTrackEditor));
 		MPCTrackCreateEditorHandle = SequencerModule.RegisterTrackEditor(FOnCreateTrackEditor::CreateStatic(&FMaterialParameterCollectionTrackEditor::CreateTrackEditor));
 		PrimitiveMaterialCreateEditorHandle = SequencerModule.RegisterTrackEditor(FOnCreateTrackEditor::CreateStatic(&FPrimitiveMaterialTrackEditor::CreateTrackEditor));
+		CameraShakeSourceShakeCreateEditorHandle = SequencerModule.RegisterTrackEditor(FOnCreateTrackEditor::CreateStatic(&FCameraShakeSourceShakeTrackEditor::CreateTrackEditor));
 
 		RegisterClipboardConversions();
 
@@ -193,6 +195,7 @@ void FMovieSceneToolsModule::ShutdownModule()
 	SequencerModule.UnRegisterTrackEditor( VisibilityPropertyTrackCreateEditorHandle );
 	SequencerModule.UnRegisterTrackEditor( ActorReferencePropertyTrackCreateEditorHandle );
 	SequencerModule.UnRegisterTrackEditor( StringPropertyTrackCreateEditorHandle );
+	SequencerModule.UnRegisterTrackEditor( CameraShakeSourceShakeCreateEditorHandle );
 
 	// unregister specialty track editors
 	SequencerModule.UnRegisterTrackEditor( AnimationTrackCreateEditorHandle );
diff --git a/Engine/Source/Editor/MovieSceneTools/Public/MovieSceneToolsModule.h b/Engine/Source/Editor/MovieSceneTools/Public/MovieSceneToolsModule.h
--- a/Engine/Source/Editor/MovieSceneTools/Public/MovieSceneToolsModule.h
+++ b/Engine/Source/Editor/MovieSceneTools/Public/MovieSceneToolsModule.h
@@ -89,10 +89,11 @@ private:
 	FDelegateHandle CameraShakeTrackCreateEditorHandle;
 	FDelegateHandle MPCTrackCreateEditorHandle;
 	FDelegateHandle PrimitiveMaterialCreateEditorHandle;
+	FDelegateHandle CameraShakeSourceShakeCreateEditorHandle;
 
 	FDelegateHandle GenerateEventEntryPointsHandle;
 	FDelegateHandle FixupPayloadParameterNameHandle;
 	FDelegateHandle UpgradeLegacyEventEndpointHandle;
 
 	TArray<IMovieSceneToolsTakeData*> TakeDatas;
-};
\ No newline at end of file
+};
diff --git a/Engine/Source/Runtime/Engine/Private/Camera/CameraShakeSourceComponent.cpp b/Engine/Source/Runtime/Engine/Private/Camera/CameraShakeSourceComponent.cpp
--- a/Engine/Source/Runtime/Engine/Private/Camera/CameraShakeSourceComponent.cpp
+++ b/Engine/Source/Runtime/Engine/Private/Camera/CameraShakeSourceComponent.cpp
@@ -74,7 +74,10 @@ void UCameraShakeSourceComponent::EndPlay(const EEndPlayReason::Type EndPlayReas
 
 void UCameraShakeSourceComponent::Play()
 {
-	PlayCameraShake(CameraShake);
+	if (ensureMsgf(CameraShake.Get() != nullptr, TEXT("No camera shake was specified on this source!")))
+	{
+		PlayCameraShake(CameraShake);
+	}
 }
 
 void UCameraShakeSourceComponent::PlayCameraShake(TSubclassOf<UCameraShake> InCameraShake)
diff --git a/Engine/Source/Runtime/MovieSceneTracks/Private/Evaluation/MovieSceneCameraShakeSourceShakeTemplate.cpp b/Engine/Source/Runtime/MovieSceneTracks/Private/Evaluation/MovieSceneCameraShakeSourceShakeTemplate.cpp
new file mode 100644
--- /dev/null
+++ b/Engine/Source/Runtime/MovieSceneTracks/Private/Evaluation/MovieSceneCameraShakeSourceShakeTemplate.cpp
@@ -0,0 +1,291 @@
+// Copyright Epic Games, Inc. All Rights Reserved.
+
+#include "Evaluation/MovieSceneCameraShakeSourceShakeTemplate.h"
+#include "Evaluation/MovieSceneEvaluation.h"
+#include "Camera/CameraShakeSourceComponent.h"
+
+#if WITH_EDITOR
+
+#include "LevelEditorViewport.h"
+#include "Camera/CameraModifier_CameraShake.h"
+
+#endif
+
+#if WITH_EDITOR
+
+/**
+ * A class that owns a gameplay camera shake manager, so that we can us it to preview shakes in editor.
+ */
+class FCameraShakePreviewer : public FGCObject
+{
+public:
+	FCameraShakePreviewer();
+	~FCameraShakePreviewer();
+
+	void RegisterViewModifier();
+	void UnRegisterViewModifier();
+	void Update(float DeltaTime, bool bIsPlaying);
+
+	UCameraModifier_CameraShake* GetCameraShake() { return PreviewCameraShake; }
+
+private:
+	// FGCObject interface
+	virtual void AddReferencedObjects(FReferenceCollector& Collector) override { Collector.AddReferencedObject(PreviewCameraShake); }
+	virtual FString GetReferencerName() const override { return TEXT("SCameraShakePreviewer"); }
+
+private:
+	void OnModifyView(FMinimalViewInfo& InOutPOV);
+	void OnLevelViewportClientListChanged();
+
+private:
+	UCameraModifier_CameraShake* PreviewCameraShake;
+	TArray<FLevelEditorViewportClient*> RegisteredViewportClients;
+
+	TOptional<float> LastDeltaTime;
+	FVector LastLocationModifier;
+	FRotator LastRotationModifier;
+	float LastFOVModifier;
+};
+
+FCameraShakePreviewer::FCameraShakePreviewer()
+	: PreviewCameraShake(NewObject<UCameraModifier_CameraShake>())
+	, LastLocationModifier(FVector::ZeroVector)
+	, LastRotationModifier(FRotator::ZeroRotator)
+	, LastFOVModifier(0.f)
+{
+}
+
+FCameraShakePreviewer::~FCameraShakePreviewer()
+{
+	PreviewCameraShake = nullptr;
+}
+
+void FCameraShakePreviewer::Update(float DeltaTime, bool bIsPlaying)
+{
+	LastDeltaTime = DeltaTime;
+
+	if (!bIsPlaying)
+	{
+		LastLocationModifier = FVector::ZeroVector;
+		LastRotationModifier = FRotator::ZeroRotator;
+		LastFOVModifier = 0.f;
+	}
+}
+
+void FCameraShakePreviewer::OnModifyView(FMinimalViewInfo& InOutPOV)
+{
+	const float DeltaTime = LastDeltaTime.Get(-1.f);
+	if (DeltaTime > 0.f)
+	{
+		FMinimalViewInfo InPOV(InOutPOV);
+		PreviewCameraShake->ModifyCamera(DeltaTime, InOutPOV);
+
+		LastLocationModifier = InOutPOV.Location - InPOV.Location;
+		LastRotationModifier = InOutPOV.Rotation - InPOV.Rotation;
+		LastFOVModifier = InOutPOV.FOV - InPOV.FOV;
+
+		LastDeltaTime.Reset();
+	}
+	else
+	{
+		InOutPOV.Location += LastLocationModifier;
+		InOutPOV.Rotation += LastRotationModifier;
+		InOutPOV.FOV += LastFOVModifier;
+	}
+}
+
+void FCameraShakePreviewer::RegisterViewModifier()
+{
+	if (GEditor == nullptr)
+	{
+		return;
+	}
+
+	// Register our view modifier on all appropriate viewports, and remember which viewports we did that on.
+	// We will later make sure to unregister on the same list, except for any viewport that somehow disappeared since,
+	// which we will be notified about with the OnLevelViewportClientListChanged event.
+	RegisteredViewportClients.Reset();
+	for (FLevelEditorViewportClient* LevelVC : GEditor->GetLevelViewportClients())
+	{		
+		if (LevelVC && LevelVC->AllowsCinematicControl() && LevelVC->GetViewMode() != VMI_Unknown)
+		{
+			RegisteredViewportClients.Add(LevelVC);
+			LevelVC->ViewModifiers.AddRaw(this, &FCameraShakePreviewer::OnModifyView);
+		}
+	}
+
+	GEditor->OnLevelViewportClientListChanged().AddRaw(this, &FCameraShakePreviewer::OnLevelViewportClientListChanged);
+}
+
+void FCameraShakePreviewer::UnRegisterViewModifier()
+{
+	GEditor->OnLevelViewportClientListChanged().RemoveAll(this);
+
+	for (FLevelEditorViewportClient* ViewportClient : RegisteredViewportClients)
+	{
+		ViewportClient->ViewModifiers.RemoveAll(this);
+	}
+}
+
+void FCameraShakePreviewer::OnLevelViewportClientListChanged()
+{
+	if (GEditor != nullptr)
+	{
+		// If any viewports were removed while we were playing, simply get rid of them from our list of
+		// registered viewports.
+		TSet<FLevelEditorViewportClient*> PreviousViewportClients(RegisteredViewportClients);
+		TSet<FLevelEditorViewportClient*> NewViewportClients(GEditor->GetLevelViewportClients());
+		RegisteredViewportClients = PreviousViewportClients.Intersect(NewViewportClients).Array();
+	}
+}
+
+#endif
+
+struct FCameraShakeSourceShakeSectionInstanceData : IPersistentEvaluationData
+{
+	bool bStarted = false;
+
+#if WITH_EDITOR
+	FCameraShakePreviewer Previewer;
+#endif
+};
+
+struct FPreAnimatedCameraShakeSourceShakeTokenProducer : IMovieScenePreAnimatedGlobalTokenProducer
+{
+	FMovieSceneEvaluationOperand Operand;
+
+	FPreAnimatedCameraShakeSourceShakeTokenProducer(FMovieSceneEvaluationOperand InOperand)
+		: Operand(InOperand)
+	{
+	}
+
+	virtual IMovieScenePreAnimatedGlobalTokenPtr CacheExistingState() const override
+	{
+		struct FRestoreToken : IMovieScenePreAnimatedGlobalToken
+		{
+			FMovieSceneEvaluationOperand Operand;
+
+			FRestoreToken(FMovieSceneEvaluationOperand InOperand) : Operand(InOperand) {}
+
+			virtual void RestoreState(IMovieScenePlayer& Player) override
+			{
+				for (TWeakObjectPtr<> BoundObject : Player.FindBoundObjects(Operand))
+				{
+					if (BoundObject.IsValid())
+					{
+						if (UCameraShakeSourceComponent* ShakeSourceComponent = CastChecked<UCameraShakeSourceComponent>(BoundObject.Get()))
+						{
+							ShakeSourceComponent->StopAllCameraShakes();
+						}
+					}
+				}
+			}
+		};
+
+		return FRestoreToken(Operand);
+	}
+};
+
+struct FCameraShakeSourceShakeExecutionToken : IMovieSceneExecutionToken
+{
+	const FMovieSceneCameraShakeSectionData& SourceData;
+
+	FCameraShakeSourceShakeExecutionToken(const FMovieSceneCameraShakeSectionData& InSourceData)
+		: SourceData(InSourceData)
+	{
+	}
+
+	virtual void Execute(const FMovieSceneContext& Context, const FMovieSceneEvaluationOperand& Operand, FPersistentEvaluationData& PersistentData, IMovieScenePlayer& Player) override
+	{
+#if WITH_EDITOR
+		FCameraShakeSourceShakeSectionInstanceData& InstanceData = PersistentData.GetSectionData<FCameraShakeSourceShakeSectionInstanceData>();
+		UCameraModifier_CameraShake* const PreviewCameraShake = InstanceData.Previewer.GetCameraShake();
+#endif
+
+		for (TWeakObjectPtr<> BoundObject : Player.FindBoundObjects(Operand))
+		{
+			if (BoundObject.IsValid())
+			{
+				if (UCameraShakeSourceComponent* ShakeSourceComponent = CastChecked<UCameraShakeSourceComponent>(BoundObject.Get()))
+				{
+					static FMovieSceneAnimTypeID ShakeTypeID = TMovieSceneAnimTypeID<FCameraShakeSourceShakeExecutionToken, 0>();
+					Player.SavePreAnimatedState(ShakeTypeID, FPreAnimatedCameraShakeSourceShakeTokenProducer(Operand));
+
+					if (SourceData.ShakeClass)
+					{
+						ShakeSourceComponent->PlayCameraShake(SourceData.ShakeClass);
+
+#if WITH_EDITOR
+						FAddCameraShakeParams Params;
+						Params.SourceComponent = ShakeSourceComponent;
+						PreviewCameraShake->AddCameraShake(SourceData.ShakeClass, Params);
+#endif
+					}
+					else
+					{
+						ShakeSourceComponent->Play();
+
+#if WITH_EDITOR
+						if (ShakeSourceComponent->CameraShake.Get() != nullptr)
+						{
+							FAddCameraShakeParams Params;
+							Params.SourceComponent = ShakeSourceComponent;
+							PreviewCameraShake->AddCameraShake(ShakeSourceComponent->CameraShake, Params);
+						}
+#endif
+					}
+				}
+			}
+		}
+	}
+};
+
+FMovieSceneCameraShakeSourceShakeSectionTemplate::FMovieSceneCameraShakeSourceShakeSectionTemplate()
+{
+}
+
+FMovieSceneCameraShakeSourceShakeSectionTemplate::FMovieSceneCameraShakeSourceShakeSectionTemplate(const UMovieSceneCameraShakeSourceShakeSection& Section)
+	: SourceData(Section.ShakeData)
+	, SectionStartTime(Section.HasStartFrame() ? Section.GetInclusiveStartFrame() : 0)
+{
+}
+
+void FMovieSceneCameraShakeSourceShakeSectionTemplate::SetupOverrides()
+{
+	EnableOverrides(RequiresSetupFlag | RequiresTearDownFlag);
+}
+
+void FMovieSceneCameraShakeSourceShakeSectionTemplate::Setup(FPersistentEvaluationData& PersistentData, IMovieScenePlayer& Player) const
+{
+	FCameraShakeSourceShakeSectionInstanceData& InstanceData = PersistentData.AddSectionData<FCameraShakeSourceShakeSectionInstanceData>();
+	InstanceData.bStarted = false;
+
+#if WITH_EDITOR
+	InstanceData.Previewer.RegisterViewModifier();
+#endif
+}
+
+void FMovieSceneCameraShakeSourceShakeSectionTemplate::Evaluate(const FMovieSceneEvaluationOperand& Operand, const FMovieSceneContext& Context, const FPersistentEvaluationData& PersistentData, FMovieSceneExecutionTokens& ExecutionTokens) const
+{
+	FCameraShakeSourceShakeSectionInstanceData& InstanceData = PersistentData.GetSectionData<FCameraShakeSourceShakeSectionInstanceData>();
+	if (!InstanceData.bStarted)
+	{
+		ExecutionTokens.Add(FCameraShakeSourceShakeExecutionToken(SourceData));
+		InstanceData.bStarted = true;
+	}
+
+#if WITH_EDITOR
+	const float DeltaTime = Context.GetFrameRate().AsSeconds(Context.GetDelta());
+	const bool bIsPlaying = Context.GetStatus() == EMovieScenePlayerStatus::Playing;
+	InstanceData.Previewer.Update(DeltaTime, bIsPlaying);
+#endif
+}
+
+void FMovieSceneCameraShakeSourceShakeSectionTemplate::TearDown(FPersistentEvaluationData& PersistentData, IMovieScenePlayer& Player) const
+{
+	FCameraShakeSourceShakeSectionInstanceData& InstanceData = PersistentData.GetSectionData<FCameraShakeSourceShakeSectionInstanceData>();
+
+#if WITH_EDITOR
+	InstanceData.Previewer.UnRegisterViewModifier();
+#endif
+}
diff --git a/Engine/Source/Runtime/MovieSceneTracks/Private/Evaluation/MovieSceneCameraShakeSourceShakeTemplate.h b/Engine/Source/Runtime/MovieSceneTracks/Private/Evaluation/MovieSceneCameraShakeSourceShakeTemplate.h
new file mode 100644
--- /dev/null
+++ b/Engine/Source/Runtime/MovieSceneTracks/Private/Evaluation/MovieSceneCameraShakeSourceShakeTemplate.h
@@ -0,0 +1,34 @@
+// Copyright Epic Games, Inc. All Rights Reserved.
+
+#pragma once
+
+#include "CoreMinimal.h"
+#include "UObject/ObjectMacros.h"
+#include "Evaluation/PersistentEvaluationData.h"
+#include "Evaluation/MovieSceneEvalTemplate.h"
+#include "Sections/MovieSceneCameraShakeSourceShakeSection.h"
+#include "MovieSceneCameraShakeSourceShakeTemplate.generated.h"
+
+USTRUCT()
+struct FMovieSceneCameraShakeSourceShakeSectionTemplate : public FMovieSceneEvalTemplate
+{
+	GENERATED_BODY()
+
+	FMovieSceneCameraShakeSourceShakeSectionTemplate();
+	FMovieSceneCameraShakeSourceShakeSectionTemplate(const UMovieSceneCameraShakeSourceShakeSection& Section);
+
+private:
+	virtual UScriptStruct& GetScriptStructImpl() const override { return *StaticStruct(); }
+	virtual void SetupOverrides() override;
+	virtual void Setup(FPersistentEvaluationData& PersistentData, IMovieScenePlayer& Player) const override;
+	virtual void Evaluate(const FMovieSceneEvaluationOperand& Operand, const FMovieSceneContext& Context, const FPersistentEvaluationData& PersistentData, FMovieSceneExecutionTokens& ExecutionTokens) const override;
+	virtual void TearDown(FPersistentEvaluationData& PersistentData, IMovieScenePlayer& Player) const override;
+
+	/** Source data taken from the section */
+	UPROPERTY()
+	FMovieSceneCameraShakeSectionData SourceData;
+
+	/** Cached section start time */
+	UPROPERTY()
+	FFrameNumber SectionStartTime;
+};
diff --git a/Engine/Source/Runtime/MovieSceneTracks/Private/Sections/MovieSceneCameraShakeSourceShakeSection.cpp b/Engine/Source/Runtime/MovieSceneTracks/Private/Sections/MovieSceneCameraShakeSourceShakeSection.cpp
new file mode 100644
--- /dev/null
+++ b/Engine/Source/Runtime/MovieSceneTracks/Private/Sections/MovieSceneCameraShakeSourceShakeSection.cpp
@@ -0,0 +1,16 @@
+// Copyright Epic Games, Inc. All Rights Reserved.
+
+#include "Sections/MovieSceneCameraShakeSourceShakeSection.h"
+#include "Evaluation/MovieSceneCameraShakeSourceShakeTemplate.h"
+#include "Tracks/MovieSceneCameraShakeSourceShakeTrack.h"
+
+UMovieSceneCameraShakeSourceShakeSection::UMovieSceneCameraShakeSourceShakeSection(const FObjectInitializer& ObjectInitializer)
+	: Super( ObjectInitializer )
+{
+	EvalOptions.EnableAndSetCompletionMode(EMovieSceneCompletionMode::ProjectDefault);
+}
+
+FMovieSceneEvalTemplatePtr UMovieSceneCameraShakeSourceShakeSection::GenerateTemplate() const
+{
+	return FMovieSceneCameraShakeSourceShakeSectionTemplate(*this);
+}
diff --git a/Engine/Source/Runtime/MovieSceneTracks/Private/Tracks/MovieSceneCameraShakeSourceShakeTrack.cpp b/Engine/Source/Runtime/MovieSceneTracks/Private/Tracks/MovieSceneCameraShakeSourceShakeTrack.cpp
new file mode 100644
--- /dev/null
+++ b/Engine/Source/Runtime/MovieSceneTracks/Private/Tracks/MovieSceneCameraShakeSourceShakeTrack.cpp
@@ -0,0 +1,85 @@
+// Copyright Epic Games, Inc. All Rights Reserved.
+
+#include "Tracks/MovieSceneCameraShakeSourceShakeTrack.h"
+#include "Sections/MovieSceneCameraShakeSourceShakeSection.h"
+#include "Compilation/MovieSceneCompilerRules.h"
+
+#define LOCTEXT_NAMESPACE "MovieSceneCameraShakeSourceShakeTrack"
+
+UMovieSceneSection* UMovieSceneCameraShakeSourceShakeTrack::AddNewCameraShake(const FFrameNumber KeyTime, const TSubclassOf<UCameraShake> ShakeClass)
+{
+	Modify();
+
+	UMovieSceneCameraShakeSourceShakeSection* const NewSection = Cast<UMovieSceneCameraShakeSourceShakeSection>(CreateNewSection());
+	if (NewSection)
+	{
+		// #fixme get length
+		FFrameTime Duration = 5.0 * GetTypedOuter<UMovieScene>()->GetTickResolution();
+		NewSection->InitialPlacement(CameraShakeSections, KeyTime, Duration.FrameNumber.Value, SupportsMultipleRows());
+		NewSection->ShakeData.ShakeClass = ShakeClass;
+		
+		AddSection(*NewSection);
+	}
+
+	return NewSection;
+}
+
+FMovieSceneTrackSegmentBlenderPtr UMovieSceneCameraShakeSourceShakeTrack::GetTrackSegmentBlender() const
+{
+	return FMovieSceneAdditiveCameraTrackBlender();
+}
+
+#if WITH_EDITORONLY_DATA
+FText UMovieSceneCameraShakeSourceShakeTrack::GetDisplayName() const
+{
+	return LOCTEXT("TrackName", "Camera Shake");
+}
+#endif
+
+const TArray<UMovieSceneSection*>& UMovieSceneCameraShakeSourceShakeTrack::GetAllSections() const
+{
+	return CameraShakeSections;
+}
+
+bool UMovieSceneCameraShakeSourceShakeTrack::SupportsType(TSubclassOf<UMovieSceneSection> SectionClass) const
+{
+	return SectionClass == UMovieSceneCameraShakeSourceShakeSection::StaticClass();
+}
+
+UMovieSceneSection* UMovieSceneCameraShakeSourceShakeTrack::CreateNewSection()
+{
+	return NewObject<UMovieSceneCameraShakeSourceShakeSection>(this, NAME_None, RF_Transactional);
+}
+
+void UMovieSceneCameraShakeSourceShakeTrack::RemoveAllAnimationData()
+{
+	CameraShakeSections.Empty();
+}
+
+bool UMovieSceneCameraShakeSourceShakeTrack::HasSection(const UMovieSceneSection& Section) const
+{
+	return CameraShakeSections.Contains(&Section);
+}
+
+void UMovieSceneCameraShakeSourceShakeTrack::AddSection(UMovieSceneSection& Section)
+{
+	CameraShakeSections.Add(&Section);
+}
+
+void UMovieSceneCameraShakeSourceShakeTrack::RemoveSection(UMovieSceneSection& Section)
+{
+	CameraShakeSections.Remove(&Section);
+}
+
+void UMovieSceneCameraShakeSourceShakeTrack::RemoveSectionAt(int32 SectionIndex)
+{
+	CameraShakeSections.RemoveAt(SectionIndex);
+}
+
+bool UMovieSceneCameraShakeSourceShakeTrack::IsEmpty() const
+{
+	return CameraShakeSections.Num() == 0;
+}
+
+#undef LOCTEXT_NAMESPACE
+
diff --git a/Engine/Source/Runtime/MovieSceneTracks/Public/Sections/MovieSceneCameraShakeSourceShakeSection.h b/Engine/Source/Runtime/MovieSceneTracks/Public/Sections/MovieSceneCameraShakeSourceShakeSection.h
new file mode 100644
--- /dev/null
+++ b/Engine/Source/Runtime/MovieSceneTracks/Public/Sections/MovieSceneCameraShakeSourceShakeSection.h
@@ -0,0 +1,22 @@
+// Copyright Epic Games, Inc. All Rights Reserved.
+
+#pragma once
+
+#include "CoreMinimal.h"
+#include "MovieSceneCameraShakeSection.h"
+#include "MovieSceneCameraShakeSourceShakeSection.generated.h"
+
+UCLASS(MinimalAPI)
+class UMovieSceneCameraShakeSourceShakeSection : public UMovieSceneSection
+{
+	GENERATED_BODY()
+
+public:
+	UMovieSceneCameraShakeSourceShakeSection(const FObjectInitializer& ObjectInitializer);
+
+	virtual FMovieSceneEvalTemplatePtr GenerateTemplate() const override;
+	
+	UPROPERTY(EditAnywhere, Category="Camera Shake", meta=(ShowOnlyInnerProperties))
+	FMovieSceneCameraShakeSectionData ShakeData;
+};
+
diff --git a/Engine/Source/Runtime/MovieSceneTracks/Public/Tracks/MovieSceneCameraShakeSourceShakeTrack.h b/Engine/Source/Runtime/MovieSceneTracks/Public/Tracks/MovieSceneCameraShakeSourceShakeTrack.h
new file mode 100644
--- /dev/null
+++ b/Engine/Source/Runtime/MovieSceneTracks/Public/Tracks/MovieSceneCameraShakeSourceShakeTrack.h
@@ -0,0 +1,50 @@
+// Copyright Epic Games, Inc. All Rights Reserved.
+
+#pragma once
+
+#include "CoreMinimal.h"
+#include "UObject/ObjectMacros.h"
+#include "Templates/SubclassOf.h"
+#include "Camera/CameraShake.h"
+#include "Misc/InlineValue.h"
+#include "MovieSceneNameableTrack.h"
+#include "MovieSceneCameraShakeSourceShakeTrack.generated.h"
+
+struct FMovieSceneEvaluationTrack;
+struct FMovieSceneSegmentCompilerRules;
+
+/**
+ * 
+ */
+UCLASS(MinimalAPI)
+class UMovieSceneCameraShakeSourceShakeTrack : public UMovieSceneNameableTrack
+{
+	GENERATED_BODY()
+
+public:
+	virtual UMovieSceneSection* AddNewCameraShake(const FFrameNumber KeyTime, const TSubclassOf<UCameraShake> ShakeClass);
+	
+public:
+	// UMovieSceneTrack interface
+	virtual bool HasSection(const UMovieSceneSection& Section) const override;
+	virtual void AddSection(UMovieSceneSection& Section) override;
+	virtual void RemoveSection(UMovieSceneSection& Section) override;
+	virtual void RemoveSectionAt(int32 SectionIndex) override;
+	virtual bool IsEmpty() const override;
+	virtual bool SupportsMultipleRows() const override { return true; }
+	virtual const TArray<UMovieSceneSection*>& GetAllSections() const override;
+	virtual bool SupportsType(TSubclassOf<UMovieSceneSection> SectionClass) const override;
+	virtual UMovieSceneSection* CreateNewSection() override;
+	virtual void RemoveAllAnimationData() override;
+	virtual FMovieSceneTrackSegmentBlenderPtr GetTrackSegmentBlender() const override;
+
+#if WITH_EDITORONLY_DATA
+	virtual FText GetDisplayName() const override;
+#endif
+	
+private:
+	/** List of all sections */
+	UPROPERTY()
+	TArray<UMovieSceneSection*> CameraShakeSections;
+};
+
