Sequencer: Initialize F3DTransformTrackToken

#rnx

diff --git a/Engine/Source/Runtime/MovieSceneTracks/Private/Evaluation/MovieSceneTemplateCommon.h b/Engine/Source/Runtime/MovieSceneTracks/Private/Evaluation/MovieSceneTemplateCommon.h
--- a/Engine/Source/Runtime/MovieSceneTracks/Private/Evaluation/MovieSceneTemplateCommon.h
+++ b/Engine/Source/Runtime/MovieSceneTracks/Private/Evaluation/MovieSceneTemplateCommon.h
@@ -19,7 +19,7 @@ private:
 /** A token that sets a component's relative transform */
 struct F3DTransformTrackToken
 {
-	F3DTransformTrackToken()
+	F3DTransformTrackToken() : Translation(FVector::ZeroVector), Rotation(FRotator::ZeroRotator), Scale(FVector::OneVector)
 	{
 	}
 
