Sequencer: Remove cook optimizations for spawn track. Previous, spawn tracks that didn't have any active sections would be removed during cook. This would prevent the tracks to be rebound during runtime.

This is also a speculative fix for non deterministic cooks.

diff --git a/Engine/Source/Runtime/MovieSceneTracks/Private/Tracks/MovieSceneSpawnTrack.cpp b/Engine/Source/Runtime/MovieSceneTracks/Private/Tracks/MovieSceneSpawnTrack.cpp
--- a/Engine/Source/Runtime/MovieSceneTracks/Private/Tracks/MovieSceneSpawnTrack.cpp
+++ b/Engine/Source/Runtime/MovieSceneTracks/Private/Tracks/MovieSceneSpawnTrack.cpp
@@ -131,41 +131,6 @@ void UMovieSceneSpawnTrack::PostCompile(FMovieSceneEvaluationTrack& OutTrack, co
 
 #if WITH_EDITORONLY_DATA
 
-ECookOptimizationFlags UMovieSceneSpawnTrack::GetCookOptimizationFlags() const
-{
-	// Since the spawn track denotes the lifetime of a spawnable, if the object is never spawned, we can remove the entire object
-	for (UMovieSceneSection* Section : Sections)
-	{
-		UMovieSceneSpawnSection*     BoolSection = CastChecked<UMovieSceneSpawnSection>(Section);
-		TMovieSceneChannelData<bool> BoolChannel = BoolSection->GetChannel().GetData();
-
-		if (!BoolSection->IsActive())
-		{
-			continue;
-		}
-
-		// If this bool section doesn't have any keys but does have a default value, this shouldn't be cooked out
-		if (BoolChannel.GetTimes().Num() == 0 && BoolSection->GetChannel().GetDefault().Get(false))
-		{
-			return ECookOptimizationFlags::None;
-		}
-		// If there are any keys that will cause this object to be spawned, we can't cook the object out
-		else for (bool Key : BoolChannel.GetValues())
-		{
-			if (Key != false)
-			{
-				return ECookOptimizationFlags::None;
-			}
-		}
-	}
-
-	return ECookOptimizationFlags::RemoveObject;
-}
-
-#endif
-
-#if WITH_EDITORONLY_DATA
-
 FText UMovieSceneSpawnTrack::GetDisplayName() const
 {
 	return LOCTEXT("TrackName", "Spawned");
diff --git a/Engine/Source/Runtime/MovieSceneTracks/Public/Tracks/MovieSceneSpawnTrack.h b/Engine/Source/Runtime/MovieSceneTracks/Public/Tracks/MovieSceneSpawnTrack.h
--- a/Engine/Source/Runtime/MovieSceneTracks/Public/Tracks/MovieSceneSpawnTrack.h
+++ b/Engine/Source/Runtime/MovieSceneTracks/Public/Tracks/MovieSceneSpawnTrack.h
@@ -58,10 +58,6 @@ public:
 	//~ UObject interface
 	virtual void PostLoad() override;
 
-#if WITH_EDITOR
-	virtual ECookOptimizationFlags GetCookOptimizationFlags() const override;
-#endif
-
 #if WITH_EDITORONLY_DATA
 	virtual FText GetDisplayName() const override;
 #endif
