Sequencer: Fix audio playback with multiple rows. The SoundLastPLayedAtTime needs to be tracked per audio component.

#rnx

diff --git a/Engine/Source/Runtime/MovieSceneTracks/Private/Evaluation/MovieSceneAudioTemplate.cpp b/Engine/Source/Runtime/MovieSceneTracks/Private/Evaluation/MovieSceneAudioTemplate.cpp
--- a/Engine/Source/Runtime/MovieSceneTracks/Private/Evaluation/MovieSceneAudioTemplate.cpp
+++ b/Engine/Source/Runtime/MovieSceneTracks/Private/Evaluation/MovieSceneAudioTemplate.cpp
@@ -88,14 +88,13 @@ struct FCachedAudioTrackData : IPersistentEvaluationData
 	TMap<FObjectKey, TMap<FObjectKey, TWeakObjectPtr<UAudioComponent>>> AudioComponentsByActorKey;
 	
 	FCachedAudioTrackData()
-		: SoundLastPlayedAtTime(0.f)
 	{
 		// Create the container for master tracks, which do not have an actor to attach to
 		AudioComponentsByActorKey.Add(FObjectKey(), TMap<FObjectKey, TWeakObjectPtr<UAudioComponent>>());
 	}
 
 	/** Set whenever we ask the Audio component to start playing a sound. Used to detect desyncs caused when Sequencer evaluates at more-than-real-time. */
-	float SoundLastPlayedAtTime;
+	TMap<TWeakObjectPtr<UAudioComponent>, float> SoundLastPlayedAtTime;
 
 	UAudioComponent* GetAudioComponent(FObjectKey ActorKey, FObjectKey SectionKey)
 	{
@@ -414,11 +413,16 @@ struct FAudioSectionExecutionToken : IMovieSceneExecutionToken
 			// This tells us how much time has passed in the game world (and thus, reasonably, the audio playback)
 			// so if we calculate that we should be playing say, 15s into the section during evaluation, but
 			// we're only 5s since the last Play call, then we know we're out of sync. 
-			float GameTimeDelta = CurrentGameTime - TrackData.SoundLastPlayedAtTime;
-			if (FMath::Abs(GameTimeDelta - AudioTime) > MaxSequenceAudioDesyncToleranceCVar)
+			if (TrackData.SoundLastPlayedAtTime.Contains(&AudioComponent))
 			{
-				UE_LOG(LogMovieScene, Verbose, TEXT("Audio Component detected a significant mismatch in (assumed) playback time versus the desired time. Time since last play call: %6.2f(s) Desired Time: %6.2f(s). Component: %s sound: %s"), GameTimeDelta, AudioTime, *AudioComponent.GetName(), *GetNameSafe(AudioComponent.Sound));
-				bSoundNeedsTimeSync = true;
+				float SoundLastPlayedAtTime = TrackData.SoundLastPlayedAtTime[&AudioComponent];
+
+				float GameTimeDelta = CurrentGameTime - SoundLastPlayedAtTime;
+				if (FMath::Abs(GameTimeDelta - AudioTime) > MaxSequenceAudioDesyncToleranceCVar)
+				{
+					UE_LOG(LogMovieScene, Verbose, TEXT("Audio Component detected a significant mismatch in (assumed) playback time versus the desired time. Time since last play call: %6.2f(s) Desired Time: %6.2f(s). Component: %s sound: %s"), GameTimeDelta, AudioTime, *AudioComponent.GetName(), *GetNameSafe(AudioComponent.Sound));
+					bSoundNeedsTimeSync = true;
+				}
 			}
 		}
 
@@ -467,7 +471,12 @@ struct FAudioSectionExecutionToken : IMovieSceneExecutionToken
 				// Keep track of when we asked this audio clip to play (in game time) so that we can figure out if there's a significant desync in the future.
 				if (Player.GetPlaybackContext()->GetWorld())
 				{
-					TrackData.SoundLastPlayedAtTime = Player.GetPlaybackContext()->GetWorld()->GetAudioTimeSeconds() - AudioTime;
+					if (!TrackData.SoundLastPlayedAtTime.Contains(&AudioComponent))
+					{
+						TrackData.SoundLastPlayedAtTime.Add(&AudioComponent);
+					}
+
+					TrackData.SoundLastPlayedAtTime[&AudioComponent] = Player.GetPlaybackContext()->GetWorld()->GetAudioTimeSeconds() - AudioTime;
 				}
 			}
 
