Fixing up compilation for PC in movie scene camera code
#rnx

diff --git a/Engine/Source/Runtime/MovieSceneTracks/Private/Evaluation/MovieSceneCameraShakeSourceTriggerTemplate.cpp b/Engine/Source/Runtime/MovieSceneTracks/Private/Evaluation/MovieSceneCameraShakeSourceTriggerTemplate.cpp
--- a/Engine/Source/Runtime/MovieSceneTracks/Private/Evaluation/MovieSceneCameraShakeSourceTriggerTemplate.cpp
+++ b/Engine/Source/Runtime/MovieSceneTracks/Private/Evaluation/MovieSceneCameraShakeSourceTriggerTemplate.cpp
@@ -120,7 +120,7 @@ void FMovieSceneCameraShakeSourceTriggerSectionTemplate::EvaluateSwept(const FMo
 	TArray<FMovieSceneCameraShakeSourceTrigger> ShakesToTrigger;
 
 	const int32 TriggerCount = TriggerTimes.Num();
-	const float PositionInSeconds = Context.GetTime() * Context.GetRootToSequenceTransform().Inverse() / Context.GetFrameRate();
+	const float PositionInSeconds = Context.GetTime() * Context.GetRootToSequenceTransform().InverseLinearOnly() / Context.GetFrameRate();
 
 	for (int32 KeyIndex = 0; KeyIndex < TriggerCount; ++KeyIndex)
 	{
