Sequencer: Bring back lost "else" statement.

diff --git a/Engine/Source/Runtime/MovieSceneTracks/Private/Evaluation/MovieSceneCameraCutTemplate.cpp b/Engine/Source/Runtime/MovieSceneTracks/Private/Evaluation/MovieSceneCameraCutTemplate.cpp
--- a/Engine/Source/Runtime/MovieSceneTracks/Private/Evaluation/MovieSceneCameraCutTemplate.cpp
+++ b/Engine/Source/Runtime/MovieSceneTracks/Private/Evaluation/MovieSceneCameraCutTemplate.cpp
@@ -275,7 +275,7 @@ void FMovieSceneCameraCutSectionTemplate::Evaluate(const FMovieSceneEvaluationOp
 	{
 		ExecutionTokens.Add(FCameraCutPreRollExecutionToken(CameraBindingID, CutTransform, bHasCutTransform));
 	}
-
+	else
 	{
 		// For now we only look at how long the camera blend is supposed to be, and we pass that on to
 		// the player controller via the execution token. Later we'll need to actually drive the blend
