Sequencer: Transform origins no longer apply to components and actors that are attached to other components or actors

  - Also fixed attach tracks not having the correct evaluation priority meaning they wouldn't necessarily evaluate before transforms

diff --git a/Engine/Source/Runtime/MovieSceneTracks/Private/Evaluation/MovieScene3DTransformTemplate.cpp b/Engine/Source/Runtime/MovieSceneTracks/Private/Evaluation/MovieScene3DTransformTemplate.cpp
--- a/Engine/Source/Runtime/MovieSceneTracks/Private/Evaluation/MovieScene3DTransformTemplate.cpp
+++ b/Engine/Source/Runtime/MovieSceneTracks/Private/Evaluation/MovieScene3DTransformTemplate.cpp
@@ -40,6 +40,55 @@ namespace MovieScene
 		Out.Rotation = FRotator::MakeFromEuler(FVector(In[3], In[4], In[5]));
 		Out.Scale = FVector(In[6], In[7], In[8]);
 	}
+
+	struct FTransformInput
+	{
+		TMultiChannelValue<float, 9> Channels;
+		const FGlobalTransformPersistentData* GlobalTransformData;
+	};
+
+	void BlendValue(TMaskedBlendable<float, 9>& OutBlend, const FTransformInput& InValue, float Weight, EMovieSceneBlendType BlendType, TMovieSceneInitialValueStore<F3DTransformTrackToken>& InitialValueStore)
+	{
+		// Apply origin transformation if necessary 
+		if (BlendType == EMovieSceneBlendType::Absolute && InValue.GlobalTransformData)
+		{
+			USceneComponent* SceneComponent = MovieSceneHelpers::SceneComponentFromRuntimeObject(InitialValueStore.GetAnimatingObject());
+
+			// Apply global transform if there is one and there is no attach parent
+			if (!SceneComponent || !SceneComponent->GetAttachParent())
+			{
+				float Components[6] = {
+					InValue.Channels.Get(0, 0.f), InValue.Channels.Get(1, 0.f), InValue.Channels.Get(2, 0.f),
+					InValue.Channels.Get(3, 0.f), InValue.Channels.Get(4, 0.f), InValue.Channels.Get(5, 0.f),
+				};
+
+				FTransform AnimatedTransform(FRotator(Components[4], Components[5], Components[3]), FVector(Components[0], Components[1], Components[2]));
+				AnimatedTransform = AnimatedTransform * InValue.GlobalTransformData->Origin;
+
+				FVector Location = AnimatedTransform.GetTranslation();
+				Components[0] = Location.X;
+				Components[1] = Location.Y;
+				Components[2] = Location.Z;
+
+				FVector Rotation = AnimatedTransform.GetRotation().Euler();
+				Components[3] = Rotation.X;
+				Components[4] = Rotation.Y;
+				Components[5] = Rotation.Z;
+
+				for (int32 Index = 0; Index < 9; ++Index)
+				{
+					if (InValue.Channels.IsSet(Index))
+					{
+						float Value = Index < UE_ARRAY_COUNT(Components) ? Components[Index] : InValue.Channels[Index];
+						BlendValue(OutBlend, Value, Index, Weight, BlendType, InitialValueStore);
+					}
+				}
+				return;
+			}
+		}
+
+		BlendValue(OutBlend, InValue.Channels, Weight, BlendType, InitialValueStore);
+	}
 }
 
 // Specify a unique runtime type identifier for 3d transform track tokens
@@ -155,47 +204,6 @@ FMovieSceneComponentTransformSectionTemplate::FMovieSceneComponentTransformSecti
 {
 }
 
-MovieScene::TMultiChannelValue<float, 9> FMovieSceneComponentTransformSectionTemplate::EvaluateTransform(FFrameTime Time, const FGlobalTransformPersistentData* GlobalTransformData) const
-{
-	MovieScene::TMultiChannelValue<float, 9> TransformValue = TemplateData.Evaluate(Time);
-	if (TransformValue.IsEmpty())
-	{
-		return TransformValue;
-	}
-
-	// Apply origin transformation if necessary 
-	if (TemplateData.BlendType == EMovieSceneBlendType::Absolute && GlobalTransformData)
-	{
-		float Components[6] = {
-			TransformValue.Get(0, 0.f), TransformValue.Get(1, 0.f), TransformValue.Get(2, 0.f),
-			TransformValue.Get(3, 0.f), TransformValue.Get(4, 0.f), TransformValue.Get(5, 0.f),
-		};
-
-		FTransform AnimatedTransform(FRotator(Components[4], Components[5], Components[3]), FVector(Components[0], Components[1], Components[2]));
-		AnimatedTransform = AnimatedTransform * GlobalTransformData->Origin;
-
-		FVector Location = AnimatedTransform.GetTranslation();
-		Components[0] = Location.X;
-		Components[1] = Location.Y;
-		Components[2] = Location.Z;
-
-		FVector Rotation = AnimatedTransform.GetRotation().Euler();
-		Components[3] = Rotation.X;
-		Components[4] = Rotation.Y;
-		Components[5] = Rotation.Z;
-
-		for (int32 Index = 0; Index < UE_ARRAY_COUNT(Components); ++Index)
-		{
-			if (TransformValue.IsSet(Index))
-			{
-				TransformValue.Set(Index, Components[Index]);
-			}
-		}
-	}
-
-	return TransformValue;
-}
-
 void FMovieSceneComponentTransformSectionTemplate::Evaluate(const FMovieSceneEvaluationOperand& Operand, const FMovieSceneContext& Context, const FPersistentEvaluationData& PersistentData, FMovieSceneExecutionTokens& ExecutionTokens) const
 {
 	using namespace MovieScene;
@@ -203,8 +211,9 @@ void FMovieSceneComponentTransformSectionTemplate::Evaluate(const FMovieSceneEva
 	static FSharedPersistentDataKey GlobalTransformDataKey = FGlobalTransformPersistentData::GetDataKey();
 	const FGlobalTransformPersistentData* GlobalTransformData = PersistentData.Find<FGlobalTransformPersistentData>(GlobalTransformDataKey);
 
-	TMultiChannelValue<float, 9> TransformValue = EvaluateTransform(Context.GetTime(), GlobalTransformData);
-	if (!TransformValue.IsEmpty())
+	FTransformInput TransformInput { TemplateData.Evaluate(Context.GetTime()), GlobalTransformData };
+
+	if (!TransformInput.Channels.IsEmpty())
 	{
 		// Ensure the accumulator knows how to actually apply component transforms
 		FMovieSceneBlendingActuatorID ActuatorTypeID = FComponentTransformActuator::GetActuatorTypeID();
@@ -222,15 +231,15 @@ void FMovieSceneComponentTransformSectionTemplate::Evaluate(const FMovieSceneEva
 			Weight *= ManualWeight;
 		}
 
-		ExecutionTokens.BlendToken(ActuatorTypeID, TBlendableToken<F3DTransformTrackToken>(TransformValue, TemplateData.BlendType, Weight));
+		ExecutionTokens.BlendToken(ActuatorTypeID, TBlendableToken<F3DTransformTrackToken>(TransformInput, TemplateData.BlendType, Weight));
 	}
 
 	if (IMovieSceneMotionVectorSimulation::IsEnabled(PersistentData, Context))
 	{
 		FFrameTime SimulationTime = IMovieSceneMotionVectorSimulation::GetSimulationTime(Context);
-		TMultiChannelValue<float, 9> SimulatedValue = EvaluateTransform(SimulationTime, GlobalTransformData);
+		FTransformInput SimulatedValue{ TemplateData.Evaluate(SimulationTime), GlobalTransformData };
 
-		if (!SimulatedValue.IsEmpty())
+		if (!SimulatedValue.Channels.IsEmpty())
 		{
 			// Ensure the accumulator knows how to actually apply component transforms
 			FMovieSceneBlendingActuatorID SimulationActuatorTypeID = FSimulatedComponentTransformActuator::GetActuatorTypeID();
@@ -251,9 +260,9 @@ void FMovieSceneComponentTransformSectionTemplate::Evaluate(const FMovieSceneEva
 			// Reproject backwards
 			for (int32 Index = 0; Index < 9; ++Index)
 			{
-				if (SimulatedValue.IsSet(Index) && TransformValue.IsSet(Index))
+				if (SimulatedValue.Channels.IsSet(Index) && TransformInput.Channels.IsSet(Index))
 				{
-					SimulatedValue.Set(Index, 2.f*TransformValue[Index] - SimulatedValue[Index]);
+					SimulatedValue.Channels.Set(Index, 2.f*TransformInput.Channels[Index] - SimulatedValue.Channels[Index]);
 				}
 			}
 
diff --git a/Engine/Source/Runtime/MovieSceneTracks/Private/Tracks/MovieScene3DAttachTrack.cpp b/Engine/Source/Runtime/MovieSceneTracks/Private/Tracks/MovieScene3DAttachTrack.cpp
--- a/Engine/Source/Runtime/MovieSceneTracks/Private/Tracks/MovieScene3DAttachTrack.cpp
+++ b/Engine/Source/Runtime/MovieSceneTracks/Private/Tracks/MovieScene3DAttachTrack.cpp
@@ -45,6 +45,12 @@ UMovieSceneSection* UMovieScene3DAttachTrack::CreateNewSection()
 	return NewSection;
 }
 
+void UMovieScene3DAttachTrack::PostCompile(FMovieSceneEvaluationTrack& Track, const FMovieSceneTrackCompilerArgs& Args) const
+{
+	// Set the evaluation priority to be higher than transforms so that attachments always happen first
+	Track.SetEvaluationPriority(1100);
+}
+
 #if WITH_EDITORONLY_DATA
 FText UMovieScene3DAttachTrack::GetDisplayName() const
 {
diff --git a/Engine/Source/Runtime/MovieSceneTracks/Public/Evaluation/MovieScene3DTransformTemplate.h b/Engine/Source/Runtime/MovieSceneTracks/Public/Evaluation/MovieScene3DTransformTemplate.h
--- a/Engine/Source/Runtime/MovieSceneTracks/Public/Evaluation/MovieScene3DTransformTemplate.h
+++ b/Engine/Source/Runtime/MovieSceneTracks/Public/Evaluation/MovieScene3DTransformTemplate.h
@@ -73,7 +73,4 @@ protected:
 	virtual UScriptStruct& GetScriptStructImpl() const override { return *StaticStruct(); }
 	virtual void Evaluate(const FMovieSceneEvaluationOperand& Operand, const FMovieSceneContext& Context, const FPersistentEvaluationData& PersistentData, FMovieSceneExecutionTokens& ExecutionTokens) const override;
 	virtual void Interrogate(const FMovieSceneContext& Context, FMovieSceneInterrogationData& Container, UObject* BindingOverride) const override;
-
-private:
-	MovieScene::TMultiChannelValue<float, 9> EvaluateTransform(FFrameTime Time, const FGlobalTransformPersistentData* GlobalTransformData) const;
 };
diff --git a/Engine/Source/Runtime/MovieSceneTracks/Public/Tracks/MovieScene3DAttachTrack.h b/Engine/Source/Runtime/MovieSceneTracks/Public/Tracks/MovieScene3DAttachTrack.h
--- a/Engine/Source/Runtime/MovieSceneTracks/Public/Tracks/MovieScene3DAttachTrack.h
+++ b/Engine/Source/Runtime/MovieSceneTracks/Public/Tracks/MovieScene3DAttachTrack.h
@@ -23,6 +23,7 @@ public:
 	virtual void AddConstraint( FFrameNumber Time, int32 Duration, const FName SocketName, const FName ComponentName, const FMovieSceneObjectBindingID& ConstraintBindingID ) override;
 	virtual bool SupportsType(TSubclassOf<UMovieSceneSection> SectionClass) const override;
 	virtual class UMovieSceneSection* CreateNewSection() override;
+	virtual void PostCompile(FMovieSceneEvaluationTrack& Track, const FMovieSceneTrackCompilerArgs& Args) const override;
 
 public:
 
