Sequencer: If we change a skeleton on a skelmesh while it's sequence is playing and then do a restore state we could crash since anim instance is referencing old skeleton which has a different number of bones.

We make sure to InitalizeAnimScriptInstance so bones/morphs/curves get recalculated.

diff --git a/Engine/Source/Runtime/MovieSceneTracks/Private/Evaluation/MovieSceneSkeletalAnimationTemplate.cpp b/Engine/Source/Runtime/MovieSceneTracks/Private/Evaluation/MovieSceneSkeletalAnimationTemplate.cpp
--- a/Engine/Source/Runtime/MovieSceneTracks/Private/Evaluation/MovieSceneSkeletalAnimationTemplate.cpp
+++ b/Engine/Source/Runtime/MovieSceneTracks/Private/Evaluation/MovieSceneSkeletalAnimationTemplate.cpp
@@ -139,6 +139,13 @@ struct FPreAnimatedAnimationTokenProducer : IMovieScenePreAnimatedTokenProducer
 				{
 					Component->AnimScriptInstance = CachedAnimInstance.Get();
 					CachedAnimInstance.Reset();
+					if (Component->AnimScriptInstance && Component->SkeletalMesh)
+					{
+						//the skeleton may have changed so need to recalc required bones as needed.
+						Component->AnimScriptInstance->CurrentSkeleton = Component->SkeletalMesh->Skeleton;
+						//Need at least RecalcRequiredbones and UpdateMorphTargetrs
+						Component->InitializeAnimScriptInstance(true);
+					}
 				}
 
 				// Restore pose after unbinding to force the restored pose
