Cameras: Add support for PlayScale, PlaySpace, and UserSpaceRotation for camera shake source components

diff --git a/Engine/Source/Runtime/Engine/Classes/Camera/CameraShakeSourceComponent.h b/Engine/Source/Runtime/Engine/Classes/Camera/CameraShakeSourceComponent.h
--- a/Engine/Source/Runtime/Engine/Classes/Camera/CameraShakeSourceComponent.h
+++ b/Engine/Source/Runtime/Engine/Classes/Camera/CameraShakeSourceComponent.h
@@ -3,6 +3,7 @@
 #pragma once
 
 #include "CoreMinimal.h"
+#include "Camera/CameraTypes.h"
 #include "Components/SceneComponent.h"
 #include "UObject/ObjectMacros.h"
 #include "CameraShakeSourceComponent.generated.h"
@@ -77,7 +78,7 @@ public:
 
 	/** Starts a new camera shake originating from this source, and apply it on all player controllers */
 	UFUNCTION(BlueprintCallable, Category = CameraShake)
-	void PlayCameraShake(TSubclassOf<UCameraShake> InCameraShake);
+	void PlayCameraShake(TSubclassOf<UCameraShake> InCameraShake, float Scale=1.f, enum ECameraAnimPlaySpace::Type PlaySpace = ECameraAnimPlaySpace::CameraLocal, FRotator UserPlaySpaceRot = FRotator::ZeroRotator);
 
 	/** Stops a camera shake originating from this source */
 	UFUNCTION(BlueprintCallable, Category = CameraShake)
diff --git a/Engine/Source/Runtime/Engine/Classes/Camera/PlayerCameraManager.h b/Engine/Source/Runtime/Engine/Classes/Camera/PlayerCameraManager.h
--- a/Engine/Source/Runtime/Engine/Classes/Camera/PlayerCameraManager.h
+++ b/Engine/Source/Runtime/Engine/Classes/Camera/PlayerCameraManager.h
@@ -786,9 +786,12 @@ public:
 	 * Plays a camera shake on this camera.
 	 * @param Shake - The class of camera shake to play.
 	 * @param SourceComponent - The source from which the camera shake originates.
+	 * @param Scale - Applies an additional constant scale on top of the dynamic scale computed with the distance to the source
+	 * @param PlaySpace - Which coordinate system to play the shake in (affects oscillations and camera anims)
+	 * @param UserPlaySpaceRot - Coordinate system to play shake when PlaySpace == CAPS_UserDefined.
 	 */
 	UFUNCTION(BlueprintCallable, Category = "Camera Shakes")
-	virtual class UCameraShake* PlayCameraShakeFromSource(TSubclassOf<class UCameraShake> ShakeClass, class UCameraShakeSourceComponent* SourceComponent);
+	virtual class UCameraShake* PlayCameraShakeFromSource(TSubclassOf<class UCameraShake> ShakeClass, class UCameraShakeSourceComponent* SourceComponent, float Scale=1.f, enum ECameraAnimPlaySpace::Type PlaySpace = ECameraAnimPlaySpace::CameraLocal, FRotator UserPlaySpaceRot = FRotator::ZeroRotator);
 	
 	/** Immediately stops the given shake instance and invalidates it. */
 	UFUNCTION(BlueprintCallable, Category = "Camera Shakes")
diff --git a/Engine/Source/Runtime/Engine/Private/Camera/CameraShakeSourceComponent.cpp b/Engine/Source/Runtime/Engine/Private/Camera/CameraShakeSourceComponent.cpp
--- a/Engine/Source/Runtime/Engine/Private/Camera/CameraShakeSourceComponent.cpp
+++ b/Engine/Source/Runtime/Engine/Private/Camera/CameraShakeSourceComponent.cpp
@@ -80,14 +80,14 @@ void UCameraShakeSourceComponent::Play()
 	}
 }
 
-void UCameraShakeSourceComponent::PlayCameraShake(TSubclassOf<UCameraShake> InCameraShake)
+void UCameraShakeSourceComponent::PlayCameraShake(TSubclassOf<UCameraShake> InCameraShake, float Scale, enum ECameraAnimPlaySpace::Type PlaySpace, FRotator UserPlaySpaceRot)
 {
 	for (FConstPlayerControllerIterator Iterator = GetWorld()->GetPlayerControllerIterator(); Iterator; ++Iterator)
 	{
 		APlayerController* PlayerController = Iterator->Get();
 		if (PlayerController != nullptr && PlayerController->PlayerCameraManager != nullptr)
 		{
-			PlayerController->PlayerCameraManager->PlayCameraShakeFromSource(InCameraShake, this);
+			PlayerController->PlayerCameraManager->PlayCameraShakeFromSource(InCameraShake, this, Scale, PlaySpace, UserPlaySpaceRot);
 		}
 	}
 }
diff --git a/Engine/Source/Runtime/Engine/Private/PlayerCameraManager.cpp b/Engine/Source/Runtime/Engine/Private/PlayerCameraManager.cpp
--- a/Engine/Source/Runtime/Engine/Private/PlayerCameraManager.cpp
+++ b/Engine/Source/Runtime/Engine/Private/PlayerCameraManager.cpp
@@ -1293,11 +1293,11 @@ UCameraShake* APlayerCameraManager::PlayCameraShake(TSubclassOf<UCameraShake> Sh
 	return nullptr;
 }
 
-UCameraShake* APlayerCameraManager::PlayCameraShakeFromSource(TSubclassOf<class UCameraShake> ShakeClass, class UCameraShakeSourceComponent* SourceComponent)
+UCameraShake* APlayerCameraManager::PlayCameraShakeFromSource(TSubclassOf<class UCameraShake> ShakeClass, class UCameraShakeSourceComponent* SourceComponent, float Scale, enum ECameraAnimPlaySpace::Type PlaySpace, FRotator UserPlaySpaceRot)
 {
 	if (ShakeClass && CachedCameraShakeMod)
 	{
-		return CachedCameraShakeMod->AddCameraShake(ShakeClass, FAddCameraShakeParams(1.0f, ECameraAnimPlaySpace::CameraLocal, FRotator::ZeroRotator, SourceComponent));
+		return CachedCameraShakeMod->AddCameraShake(ShakeClass, FAddCameraShakeParams(Scale, PlaySpace, UserPlaySpaceRot, SourceComponent));
 	}
 
 	return nullptr;
diff --git a/Engine/Source/Runtime/MovieSceneTracks/Private/Evaluation/MovieSceneCameraShakeSourceShakeTemplate.cpp b/Engine/Source/Runtime/MovieSceneTracks/Private/Evaluation/MovieSceneCameraShakeSourceShakeTemplate.cpp
--- a/Engine/Source/Runtime/MovieSceneTracks/Private/Evaluation/MovieSceneCameraShakeSourceShakeTemplate.cpp
+++ b/Engine/Source/Runtime/MovieSceneTracks/Private/Evaluation/MovieSceneCameraShakeSourceShakeTemplate.cpp
@@ -107,7 +107,7 @@ struct FCameraShakeSourceShakeStartExecutionToken : IMovieSceneExecutionToken
 						}
 
 						// Start playing the shake.
-						ShakeSourceComponent->PlayCameraShake(ShakeClass);
+						ShakeSourceComponent->PlayCameraShake(ShakeClass, SourceData.PlayScale, SourceData.PlaySpace, SourceData.UserDefinedPlaySpace);
 
 #if WITH_EDITOR
 						// Also start playing the shake in our editor preview.
@@ -115,6 +115,9 @@ struct FCameraShakeSourceShakeStartExecutionToken : IMovieSceneExecutionToken
 
 						FAddCameraShakeParams Params;
 						Params.SourceComponent = ShakeSourceComponent;
+						Params.Scale = SourceData.PlayScale;
+						Params.PlaySpace = SourceData.PlaySpace;
+						Params.UserPlaySpaceRot = SourceData.UserDefinedPlaySpace;
 						PreviewCameraShake->AddCameraShake(ShakeClass, Params);
 #endif
 					}
diff --git a/Engine/Source/Runtime/MovieSceneTracks/Private/Evaluation/MovieSceneCameraShakeSourceTriggerTemplate.cpp b/Engine/Source/Runtime/MovieSceneTracks/Private/Evaluation/MovieSceneCameraShakeSourceTriggerTemplate.cpp
--- a/Engine/Source/Runtime/MovieSceneTracks/Private/Evaluation/MovieSceneCameraShakeSourceTriggerTemplate.cpp
+++ b/Engine/Source/Runtime/MovieSceneTracks/Private/Evaluation/MovieSceneCameraShakeSourceTriggerTemplate.cpp
@@ -44,7 +44,7 @@ struct FCameraShakeSourceTriggerExecutionToken : IMovieSceneExecutionToken
 						if (ShakeClass.Get() != nullptr)
 						{
 							// Start playing the shake.
-							ShakeSourceComponent->PlayCameraShake(ShakeClass);
+							ShakeSourceComponent->PlayCameraShake(ShakeClass, Trigger.PlayScale, Trigger.PlaySpace, Trigger.UserDefinedPlaySpace);
 
 #if WITH_EDITOR
 							// Also start playing the shake in our editor preview.
@@ -52,6 +52,9 @@ struct FCameraShakeSourceTriggerExecutionToken : IMovieSceneExecutionToken
 
 							FAddCameraShakeParams Params;
 							Params.SourceComponent = ShakeSourceComponent;
+							Params.Scale = Trigger.PlayScale;
+							Params.PlaySpace = Trigger.PlaySpace;
+							Params.UserPlaySpaceRot = Trigger.UserDefinedPlaySpace;
 							PreviewCameraShake->AddCameraShake(ShakeClass, Params);
 #endif
 						}
diff --git a/Engine/Source/Runtime/MovieSceneTracks/Public/Channels/MovieSceneCameraShakeSourceTriggerChannel.h b/Engine/Source/Runtime/MovieSceneTracks/Public/Channels/MovieSceneCameraShakeSourceTriggerChannel.h
--- a/Engine/Source/Runtime/MovieSceneTracks/Public/Channels/MovieSceneCameraShakeSourceTriggerChannel.h
+++ b/Engine/Source/Runtime/MovieSceneTracks/Public/Channels/MovieSceneCameraShakeSourceTriggerChannel.h
@@ -15,11 +15,33 @@ struct FMovieSceneCameraShakeSourceTrigger
 {
 	GENERATED_BODY()
 
-	FMovieSceneCameraShakeSourceTrigger() {}
-	FMovieSceneCameraShakeSourceTrigger(TSubclassOf<UCameraShake> InShakeClass) : ShakeClass(InShakeClass) {}
+	FMovieSceneCameraShakeSourceTrigger()
+		: ShakeClass(nullptr)
+		, PlayScale(1.f)
+		, PlaySpace(ECameraAnimPlaySpace::CameraLocal)
+		, UserDefinedPlaySpace(ForceInitToZero)
+	{}
 
+	FMovieSceneCameraShakeSourceTrigger(TSubclassOf<UCameraShake> InShakeClass)
+		: ShakeClass(InShakeClass)
+		, PlayScale(1.f)
+		, PlaySpace(ECameraAnimPlaySpace::CameraLocal)
+		, UserDefinedPlaySpace(ForceInitToZero)
+	{}
+
+	/** Class of the camera shake to play */
 	UPROPERTY(EditAnywhere, BlueprintReadWrite, Category="Camera Shake")
 	TSubclassOf<UCameraShake> ShakeClass;
+
+	/** Scalar that affects shake intensity */
+	UPROPERTY(EditAnywhere, BlueprintReadWrite, Category="Camera Shake")
+	float PlayScale;
+
+	UPROPERTY(EditAnywhere, BlueprintReadWrite, Category="Camera Shake")
+	TEnumAsByte<ECameraAnimPlaySpace::Type> PlaySpace;
+
+	UPROPERTY(EditAnywhere, BlueprintReadWrite, Category="Camera Shake")
+	FRotator UserDefinedPlaySpace;
 };
 
 USTRUCT()
