Description: Change the Audio desync check to query the audio clock directly. When we detect a difference between when the audio component was played and where it should be playing at, we've been querying the game thread time previously through UWorld::GetAudioTimeSeconds. The Cyclone/Jerky sequence is driven by the audio render thread clock, which causes a discrepancy between where our sequence is and where we think we should play the audio component.

Risk: Low/Medium (can be disabled via hotfix)
Testing: local, though could not consistently repro the original issue.

diff --git a/Engine/Source/Runtime/MovieSceneTracks/Private/Evaluation/MovieSceneAudioTemplate.cpp b/Engine/Source/Runtime/MovieSceneTracks/Private/Evaluation/MovieSceneAudioTemplate.cpp
--- a/Engine/Source/Runtime/MovieSceneTracks/Private/Evaluation/MovieSceneAudioTemplate.cpp
+++ b/Engine/Source/Runtime/MovieSceneTracks/Private/Evaluation/MovieSceneAudioTemplate.cpp
@@ -29,6 +29,13 @@ FAutoConsoleVariableRef CVarMaxSequenceAudioDesyncTolerance(
 	TEXT("Controls how many seconds an audio track can be out of sync in a Sequence before we attempt a time correction.\n"),
 	ECVF_Default);
 
+static int32 UseAudioClockForSequencerDesyncCVar = 1;
+FAutoConsoleVariableRef CVaUseAudioClockForSequencerDesync(
+	TEXT("Sequencer.Audio.UseAudioClockForAudioDesync"),
+	UseAudioClockForSequencerDesyncCVar,
+	TEXT("When set to 1, we will use the audio render thread directly to query whether audio has went out of sync with the sequence.\n"),
+	ECVF_Default);
+
 
 /** Stop audio from playing */
 struct FStopAudioPreAnimatedToken : IMovieScenePreAnimatedToken
@@ -408,7 +415,17 @@ struct FAudioSectionExecutionToken : IMovieSceneExecutionToken
 
 		if (Player.GetPlaybackContext()->GetWorld())
 		{
-			float CurrentGameTime = Player.GetPlaybackContext()->GetWorld()->GetAudioTimeSeconds();
+			float CurrentGameTime = 0.0f;
+
+			FAudioDevice* AudioDevice = Player.GetPlaybackContext()->GetWorld()->GetAudioDeviceRaw();
+			if (UseAudioClockForSequencerDesyncCVar && AudioDevice)
+			{
+				CurrentGameTime = AudioDevice->GetAudioClock();
+			}
+			else
+			{
+				CurrentGameTime = Player.GetPlaybackContext()->GetWorld()->GetAudioTimeSeconds();
+			}
 
 			// This tells us how much time has passed in the game world (and thus, reasonably, the audio playback)
 			// so if we calculate that we should be playing say, 15s into the section during evaluation, but
@@ -418,7 +435,7 @@ struct FAudioSectionExecutionToken : IMovieSceneExecutionToken
 				float SoundLastPlayedAtTime = TrackData.SoundLastPlayedAtTime[&AudioComponent];
 
 				float GameTimeDelta = CurrentGameTime - SoundLastPlayedAtTime;
-				if (FMath::Abs(GameTimeDelta - AudioTime) > MaxSequenceAudioDesyncToleranceCVar)
+				if (!FMath::IsNearlyZero(MaxSequenceAudioDesyncToleranceCVar) && FMath::Abs(GameTimeDelta - AudioTime) > MaxSequenceAudioDesyncToleranceCVar)
 				{
 					UE_LOG(LogMovieScene, Verbose, TEXT("Audio Component detected a significant mismatch in (assumed) playback time versus the desired time. Time since last play call: %6.2f(s) Desired Time: %6.2f(s). Component: %s sound: %s"), GameTimeDelta, AudioTime, *AudioComponent.GetName(), *GetNameSafe(AudioComponent.Sound));
 					bSoundNeedsTimeSync = true;
@@ -476,7 +493,16 @@ struct FAudioSectionExecutionToken : IMovieSceneExecutionToken
 						TrackData.SoundLastPlayedAtTime.Add(&AudioComponent);
 					}
 
-					TrackData.SoundLastPlayedAtTime[&AudioComponent] = Player.GetPlaybackContext()->GetWorld()->GetAudioTimeSeconds() - AudioTime;
+					FAudioDevice* AudioDevice = Player.GetPlaybackContext()->GetWorld()->GetAudioDeviceRaw();
+					if (UseAudioClockForSequencerDesyncCVar && AudioDevice)
+					{
+						TrackData.SoundLastPlayedAtTime[&AudioComponent] = AudioDevice->GetAudioClock() - AudioTime;
+					}
+					else
+					{
+						TrackData.SoundLastPlayedAtTime[&AudioComponent] = Player.GetPlaybackContext()->GetWorld()->GetAudioTimeSeconds() - AudioTime;
+					}
+					
 				}
 			}
 
