Sequencer: Clamp trajectory drawing bound by view range

diff --git a/Engine/Source/Editor/Sequencer/Private/SequencerEdMode.cpp b/Engine/Source/Editor/Sequencer/Private/SequencerEdMode.cpp
--- a/Engine/Source/Editor/Sequencer/Private/SequencerEdMode.cpp
+++ b/Engine/Source/Editor/Sequencer/Private/SequencerEdMode.cpp
@@ -444,7 +444,9 @@ void FSequencerEdMode::DrawTransformTrack(const TSharedPtr<FSequencer>& Sequence
 	// Draw one line per-track (should only really ever be one)
 	if (FMovieSceneEvaluationTrack * EvalTrack = MovieSceneToolHelpers::GetEvaluationTrack(Sequencer.Get(), TransformTrack->GetSignature()))
 	{
-		TArray<FTrajectoryKey> TrajectoryKeys = TransformTrack->GetTrajectoryData(Sequencer->GetLocalTime().Time.FrameNumber, Sequencer->GetSequencerSettings()->GetTrajectoryPathCap());
+		TRange<FFrameNumber> ViewRange(TickResolution.AsFrameNumber(Sequencer->GetViewRange().GetLowerBoundValue()), TickResolution.AsFrameNumber(Sequencer->GetViewRange().GetUpperBoundValue()));
+		
+		TArray<FTrajectoryKey> TrajectoryKeys = TransformTrack->GetTrajectoryData(Sequencer->GetLocalTime().Time.FrameNumber, Sequencer->GetSequencerSettings()->GetTrajectoryPathCap(), ViewRange);
 		for (TWeakObjectPtr<> WeakBinding : BoundObjects)
 		{
 			UObject* BoundObject = WeakBinding.Get();
diff --git a/Engine/Source/Runtime/MovieSceneTracks/Private/Tracks/MovieScene3DTransformTrack.cpp b/Engine/Source/Runtime/MovieSceneTracks/Private/Tracks/MovieScene3DTransformTrack.cpp
--- a/Engine/Source/Runtime/MovieSceneTracks/Private/Tracks/MovieScene3DTransformTrack.cpp
+++ b/Engine/Source/Runtime/MovieSceneTracks/Private/Tracks/MovieScene3DTransformTrack.cpp
@@ -53,12 +53,12 @@ uint32 GetDistanceFromTo(FFrameNumber TestValue, FFrameNumber TargetValue)
 	}
 }
 
-TArray<FTrajectoryKey> UMovieScene3DTransformTrack::GetTrajectoryData(FFrameNumber Time, int32 MaxNumDataPoints) const
+TArray<FTrajectoryKey> UMovieScene3DTransformTrack::GetTrajectoryData(FFrameNumber Time, int32 MaxNumDataPoints, TRange<FFrameNumber> ViewRange) const
 {
 	struct FCurveKeyIterator
 	{
-		FCurveKeyIterator(UMovieScene3DTransformSection* InSection, FMovieSceneFloatChannel* InChannel, FName InChannelName, FFrameNumber StartTime)
-			: Section(InSection), Channel(InChannel->GetData()), ChannelName(InChannelName), SectionRange(InSection->GetRange()), CurrentIndex(INDEX_NONE)
+		FCurveKeyIterator(UMovieScene3DTransformSection* InSection, FMovieSceneFloatChannel* InChannel, FName InChannelName, FFrameNumber StartTime, TRange<FFrameNumber> ViewRange)
+			: Section(InSection), Channel(InChannel->GetData()), ChannelName(InChannelName), SectionRange(TRange<FFrameNumber>::Intersection(ViewRange, InSection->GetRange())), CurrentIndex(INDEX_NONE)
 		{
 			TArrayView<const FFrameNumber> Times = Channel.GetTimes();
 			CurrentIndex = Algo::LowerBound(Times, StartTime);
@@ -183,33 +183,33 @@ TArray<FTrajectoryKey> UMovieScene3DTransformTrack::GetTrajectoryData(FFrameNumb
 			EMovieSceneTransformChannel Mask = TransformSection->GetMask().GetChannels();
 			if (EnumHasAnyFlags(Mask, EMovieSceneTransformChannel::TranslationX))
 			{
-				ForwardIters.Emplace(TransformSection, FloatChannels[0], MetaData[0].Name, Time);
-				BackwardIters.Emplace(TransformSection, FloatChannels[0], MetaData[0].Name, Time);
+				ForwardIters.Emplace(TransformSection, FloatChannels[0], MetaData[0].Name, Time, ViewRange);
+				BackwardIters.Emplace(TransformSection, FloatChannels[0], MetaData[0].Name, Time, ViewRange);
 			}
 			if (EnumHasAnyFlags(Mask, EMovieSceneTransformChannel::TranslationY))
 			{
-				ForwardIters.Emplace(TransformSection, FloatChannels[1], MetaData[1].Name, Time);
-				BackwardIters.Emplace(TransformSection, FloatChannels[1], MetaData[1].Name, Time);
+				ForwardIters.Emplace(TransformSection, FloatChannels[1], MetaData[1].Name, Time, ViewRange);
+				BackwardIters.Emplace(TransformSection, FloatChannels[1], MetaData[1].Name, Time, ViewRange);
 			}
 			if (EnumHasAnyFlags(Mask, EMovieSceneTransformChannel::TranslationZ))
 			{
-				ForwardIters.Emplace(TransformSection, FloatChannels[2], MetaData[2].Name, Time);
-				BackwardIters.Emplace(TransformSection, FloatChannels[2], MetaData[2].Name, Time);
+				ForwardIters.Emplace(TransformSection, FloatChannels[2], MetaData[2].Name, Time, ViewRange);
+				BackwardIters.Emplace(TransformSection, FloatChannels[2], MetaData[2].Name, Time, ViewRange);
 			}
 			if (EnumHasAnyFlags(Mask, EMovieSceneTransformChannel::RotationX))
 			{
-				ForwardIters.Emplace(TransformSection, FloatChannels[3], MetaData[3].Name, Time);
-				BackwardIters.Emplace(TransformSection, FloatChannels[3], MetaData[3].Name, Time);
+				ForwardIters.Emplace(TransformSection, FloatChannels[3], MetaData[3].Name, Time, ViewRange);
+				BackwardIters.Emplace(TransformSection, FloatChannels[3], MetaData[3].Name, Time, ViewRange);
 			}
 			if (EnumHasAnyFlags(Mask, EMovieSceneTransformChannel::RotationY))
 			{
-				ForwardIters.Emplace(TransformSection, FloatChannels[4], MetaData[4].Name, Time);
-				BackwardIters.Emplace(TransformSection, FloatChannels[4], MetaData[4].Name, Time);
+				ForwardIters.Emplace(TransformSection, FloatChannels[4], MetaData[4].Name, Time, ViewRange);
+				BackwardIters.Emplace(TransformSection, FloatChannels[4], MetaData[4].Name, Time, ViewRange);
 			}
 			if (EnumHasAnyFlags(Mask, EMovieSceneTransformChannel::RotationZ))
 			{
-				ForwardIters.Emplace(TransformSection, FloatChannels[5], MetaData[5].Name, Time);
-				BackwardIters.Emplace(TransformSection, FloatChannels[5], MetaData[5].Name, Time);
+				ForwardIters.Emplace(TransformSection, FloatChannels[5], MetaData[5].Name, Time, ViewRange);
+				BackwardIters.Emplace(TransformSection, FloatChannels[5], MetaData[5].Name, Time, ViewRange);
 			}
 		}
 	}
diff --git a/Engine/Source/Runtime/MovieSceneTracks/Public/Tracks/MovieScene3DTransformTrack.h b/Engine/Source/Runtime/MovieSceneTracks/Public/Tracks/MovieScene3DTransformTrack.h
--- a/Engine/Source/Runtime/MovieSceneTracks/Public/Tracks/MovieScene3DTransformTrack.h
+++ b/Engine/Source/Runtime/MovieSceneTracks/Public/Tracks/MovieScene3DTransformTrack.h
@@ -67,7 +67,7 @@ public:
 
 #if WITH_EDITOR
 
-	MOVIESCENETRACKS_API TArray<FTrajectoryKey> GetTrajectoryData(FFrameNumber Time, int32 MaxNumDataPoints) const;
+	MOVIESCENETRACKS_API TArray<FTrajectoryKey> GetTrajectoryData(FFrameNumber Time, int32 MaxNumDataPoints, TRange<FFrameNumber>) const;
 
 #endif
 };
