Sequencer: NLA Fix for Yaw blending, need to change rotation order so yaw doesn't happen last.

diff --git a/Engine/Source/Runtime/MovieSceneTracks/Private/Tracks/MovieSceneSkeletalAnimationTrack.cpp b/Engine/Source/Runtime/MovieSceneTracks/Private/Tracks/MovieSceneSkeletalAnimationTrack.cpp
--- a/Engine/Source/Runtime/MovieSceneTracks/Private/Tracks/MovieSceneSkeletalAnimationTrack.cpp
+++ b/Engine/Source/Runtime/MovieSceneTracks/Private/Tracks/MovieSceneSkeletalAnimationTrack.cpp
@@ -497,8 +497,6 @@ void UMovieSceneSkeletalAnimationTrack::SetUpRootMotions(bool bForce)
 					PrevSection = AnimSection;
 				}
 			}
-
-
 			if (CurrentWeights.Num() > 0)
 			{
 				float TotalWeight = 0.0f;
@@ -558,6 +556,228 @@ static FTransform GetWorldTransformForBone(UAnimSequence* AnimSequence, USkeleta
 	return WorldTransform;
 }
 
+enum class ESkelAnimRotationOrder 
+{
+	XYZ,
+	XZY,
+	YXZ,
+	YZX,
+	ZXY,
+	ZYX
+};
+
+static FQuat QuatFromEuler(const FVector& XYZAnglesInDegrees, ESkelAnimRotationOrder RotationOrder)
+{
+	float X = FMath::DegreesToRadians(XYZAnglesInDegrees.X);
+	float Y = FMath::DegreesToRadians(XYZAnglesInDegrees.Y);
+	float Z = FMath::DegreesToRadians(XYZAnglesInDegrees.Z);
+
+	float CosX = FMath::Cos(X * 0.5f);
+	float CosY = FMath::Cos(Y * 0.5f);
+	float CosZ = FMath::Cos(Z * 0.5f);
+
+	float SinX = FMath::Sin(X * 0.5f);
+	float SinY = FMath::Sin(Y * 0.5f);
+	float SinZ = FMath::Sin(Z * 0.5f);
+
+	if (RotationOrder == ESkelAnimRotationOrder::XYZ)
+	{
+		return FQuat(SinX * CosY * CosZ - CosX * SinY * SinZ,
+			CosX * SinY * CosZ + SinX * CosY * SinZ,
+			CosX * CosY * SinZ - SinX * SinY * CosZ,
+			CosX * CosY * CosZ + SinX * SinY * SinZ);
+
+	}
+	else if (RotationOrder == ESkelAnimRotationOrder::XZY)
+	{
+		return FQuat(SinX * CosY * CosZ + CosX * SinY * SinZ,
+			CosX * SinY * CosZ + SinX * CosY * SinZ,
+			CosX * CosY * SinZ - SinX * SinY * CosZ,
+			CosX * CosY * CosZ - SinX * SinY * SinZ);
+
+	}
+	else if (RotationOrder == ESkelAnimRotationOrder::YXZ)
+	{
+		return FQuat(SinX * CosY * CosZ - CosX * SinY * SinZ,
+			CosX * SinY * CosZ + SinX * CosY * SinZ,
+			CosX * CosY * SinZ + SinX * SinY * CosZ,
+			CosX * CosY * CosZ - SinX * SinY * SinZ);
+
+	}
+	else if (RotationOrder == ESkelAnimRotationOrder::YZX)
+	{
+		return FQuat(SinX * CosY * CosZ - CosX * SinY * SinZ,
+			CosX * SinY * CosZ - SinX * CosY * SinZ,
+			CosX * CosY * SinZ + SinX * SinY * CosZ,
+			CosX * CosY * CosZ + SinX * SinY * SinZ);
+	}
+	else if (RotationOrder == ESkelAnimRotationOrder::ZXY)
+	{
+		return FQuat(SinX * CosY * CosZ + CosX * SinY * SinZ,
+			CosX * SinY * CosZ - SinX * CosY * SinZ,
+			CosX * CosY * SinZ - SinX * SinY * CosZ,
+			CosX * CosY * CosZ + SinX * SinY * SinZ);
+
+	}
+	else if (RotationOrder == ESkelAnimRotationOrder::ZYX)
+	{
+		return FQuat(SinX * CosY * CosZ + CosX * SinY * SinZ,
+			CosX * SinY * CosZ - SinX * CosY * SinZ,
+			CosX * CosY * SinZ + SinX * SinY * CosZ,
+			CosX * CosY * CosZ - SinX * SinY * SinZ);
+
+	}
+
+	// should not happen
+	return FQuat::Identity;
+}
+
+static FVector EulerFromQuat(const FQuat& Rotation, ESkelAnimRotationOrder RotationOrder)
+{
+	float X = Rotation.X;
+	float Y = Rotation.Y;
+	float Z = Rotation.Z;
+	float W = Rotation.W;
+	float X2 = X * 2.f;
+	float Y2 = Y * 2.f;
+	float Z2 = Z * 2.f;
+	float XX2 = X * X2;
+	float XY2 = X * Y2;
+	float XZ2 = X * Z2;
+	float YX2 = Y * X2;
+	float YY2 = Y * Y2;
+	float YZ2 = Y * Z2;
+	float ZX2 = Z * X2;
+	float ZY2 = Z * Y2;
+	float ZZ2 = Z * Z2;
+	float WX2 = W * X2;
+	float WY2 = W * Y2;
+	float WZ2 = W * Z2;
+
+	FVector AxisX, AxisY, AxisZ;
+	AxisX.X = (1.f - (YY2 + ZZ2));
+	AxisY.X = (XY2 + WZ2);
+	AxisZ.X = (XZ2 - WY2);
+	AxisX.Y = (XY2 - WZ2);
+	AxisY.Y = (1.f - (XX2 + ZZ2));
+	AxisZ.Y = (YZ2 + WX2);
+	AxisX.Z = (XZ2 + WY2);
+	AxisY.Z = (YZ2 - WX2);
+	AxisZ.Z = (1.f - (XX2 + YY2));
+
+	FVector Result = FVector::ZeroVector;
+
+	if (RotationOrder == ESkelAnimRotationOrder::XYZ)
+	{
+		Result.Y = FMath::Asin(-FMath::Clamp<float>(AxisZ.X, -1.f, 1.f));
+
+		if (FMath::Abs(AxisZ.X) < 1.f - SMALL_NUMBER)
+		{
+			Result.X = FMath::Atan2(AxisZ.Y, AxisZ.Z);
+			Result.Z = FMath::Atan2(AxisY.X, AxisX.X);
+		}
+		else
+		{
+			Result.X = 0.f;
+			Result.Z = FMath::Atan2(-AxisX.Y, AxisY.Y);
+		}
+	}
+	else if (RotationOrder == ESkelAnimRotationOrder::XZY)
+	{
+
+		Result.Z = FMath::Asin(FMath::Clamp<float>(AxisY.X, -1.f, 1.f));
+
+		if (FMath::Abs(AxisY.X) < 1.f - SMALL_NUMBER)
+		{
+			Result.X = FMath::Atan2(-AxisY.Z, AxisY.Y);
+			Result.Y = FMath::Atan2(-AxisZ.X, AxisX.X);
+		}
+		else
+		{
+			Result.X = 0.f;
+			Result.Y = FMath::Atan2(AxisX.Z, AxisZ.Z);
+		}
+	}
+	else if (RotationOrder == ESkelAnimRotationOrder::YXZ)
+	{
+		Result.X = FMath::Asin(FMath::Clamp<float>(AxisZ.Y, -1.f, 1.f));
+
+		if (FMath::Abs(AxisZ.Y) < 1.f - SMALL_NUMBER)
+		{
+			Result.Y = FMath::Atan2(-AxisZ.X, AxisZ.Z);
+			Result.Z = FMath::Atan2(-AxisX.Y, AxisY.Y);
+		}
+		else
+		{
+			Result.Y = 0.f;
+			Result.Z = FMath::Atan2(AxisY.X, AxisX.X);
+		}
+	}
+	else if (RotationOrder == ESkelAnimRotationOrder::YZX)
+	{
+		Result.Z = FMath::Asin(-FMath::Clamp<float>(AxisX.Y, -1.f, 1.f));
+
+		if (FMath::Abs(AxisX.Y) < 1.f - SMALL_NUMBER)
+		{
+			Result.X = FMath::Atan2(AxisZ.Y, AxisY.Y);
+			Result.Y = FMath::Atan2(AxisX.Z, AxisX.X);
+		}
+		else
+		{
+			Result.X = FMath::Atan2(-AxisY.Z, AxisZ.Z);
+			Result.Y = 0.f;
+		}
+	}
+	else if (RotationOrder == ESkelAnimRotationOrder::ZXY)
+	{
+		Result.X = FMath::Asin(-FMath::Clamp<float>(AxisY.Z, -1.f, 1.f));
+
+		if (FMath::Abs(AxisY.Z) < 1.f - SMALL_NUMBER)
+		{
+			Result.Y = FMath::Atan2(AxisX.Z, AxisZ.Z);
+			Result.Z = FMath::Atan2(AxisY.X, AxisY.Y);
+		}
+		else
+		{
+			Result.Y = FMath::Atan2(-AxisZ.X, AxisX.X);
+			Result.Z = 0.f;
+		}
+	}
+	else if (RotationOrder == ESkelAnimRotationOrder::ZYX)
+	{
+		Result.Y = FMath::Asin(FMath::Clamp<float>(AxisX.Z, -1.f, 1.f));
+
+		if (FMath::Abs(AxisX.Z) < 1.f - SMALL_NUMBER)
+		{
+			Result.X = FMath::Atan2(-AxisY.Z, AxisZ.Z);
+			Result.Z = FMath::Atan2(-AxisX.Y, AxisX.X);
+		}
+		else
+		{
+			Result.X = FMath::Atan2(AxisZ.Y, AxisY.Y);
+			Result.Z = 0.f;
+		}
+	}
+
+	return Result * 180.f / PI;
+}
+/**
+*  Function to find best rotation order given how we are matching the rotations. 
+*  If it is matched we need to make sure it happens first
+*  Issue is Yaw is most common match but by default FRotation:FQuat conversions it's last, this causes issues.
+*/
+static ESkelAnimRotationOrder FindBestRotationOrder(bool bMatchRoll, bool bMatchPitch, bool bMatchYaw)
+{
+	if (bMatchYaw)
+	{
+		return ESkelAnimRotationOrder::YXZ;
+	}
+	if (bMatchPitch)
+	{
+		return ESkelAnimRotationOrder::YZX;
+	}
+	return ESkelAnimRotationOrder::XYZ;
+}
 
 void UMovieSceneSkeletalAnimationTrack::MatchSectionByBoneTransform(bool bMatchWithPrevious, USkeletalMeshComponent* SkelMeshComp, UMovieSceneSkeletalAnimationSection* CurrentSection, FFrameTime CurrentFrame, FFrameRate FrameRate,
 	const FName& BoneName, FTransform& SecondSectionRootDiff, FVector& TranslationDiff, FQuat& RotationDiff) //add options for z and for rotation.
@@ -597,10 +817,21 @@ void UMovieSceneSkeletalAnimationTrack::MatchSectionByBoneTransform(bool bMatchW
 			FTransform  SecondTransform = GetWorldTransformForBone(SecondAnimSequence, SkelMeshComp, BoneName, SecondSectionTime);
 
 			//Need to match the translations and rotations here 
+			//First need to get the correct rotation order based upon what's matching, otherwise if not all are matched 
+			//and one rotation is set last we will get errors.
+			ESkelAnimRotationOrder RotationOrder = FindBestRotationOrder(CurrentSection->bMatchRotationRoll, CurrentSection->bMatchRotationPitch, CurrentSection->bMatchRotationYaw);
+
 			FVector FirstTransformTranslation = FirstTransform.GetTranslation();
 			FVector SecondTransformTranslation = SecondTransform.GetTranslation();
-			FRotator FirstTransformRotation = FirstTransform.GetRotation().Rotator();
-			FRotator SecondTransformRotation = SecondTransform.GetRotation().Rotator();
+			FQuat FirstTransformQuat = FirstTransform.GetRotation();
+			FQuat SecondTransformQuat = SecondTransform.GetRotation();
+			
+			FirstTransformQuat.EnforceShortestArcWith(SecondTransformQuat);
+
+			FRotator FirstTransformRotation(FRotator::MakeFromEuler(EulerFromQuat(FirstTransformQuat, RotationOrder)));
+			FRotator SecondTransformRotation(FRotator::MakeFromEuler(EulerFromQuat(SecondTransformQuat, RotationOrder)));
+			SecondTransformRotation.SetClosestToMe(FirstTransformRotation);
+
 			if (!CurrentSection->bMatchTranslation)
 			{
 				FirstTransformTranslation.X = SecondTransformTranslation.X;
@@ -625,8 +856,9 @@ void UMovieSceneSkeletalAnimationTrack::MatchSectionByBoneTransform(bool bMatchW
 			{
 				FirstTransformRotation.Roll = SecondTransformRotation.Roll;
 			}
-			FirstTransform.SetRotation(FirstTransformRotation.Quaternion());
-
+			FirstTransformQuat = QuatFromEuler(FirstTransformRotation.Euler(), RotationOrder);
+			FirstTransform.SetRotation(FirstTransformQuat);
+			
 			// Below is the match but we need to use GetRelativeTransformReverse since Inverse doesn't work as expected.
 			//	* GetRelativeTransformReverse returns this(-1)* Other, and parameter is Other.
 			SecondSectionRootDiff = SecondTransform.GetRelativeTransformReverse(FirstTransform);
@@ -649,10 +881,21 @@ void UMovieSceneSkeletalAnimationTrack::MatchSectionByBoneTransform(bool bMatchW
 			FTransform  SecondTransform = GetWorldTransformForBone(SecondAnimSequence, SkelMeshComp, BoneName, SecondSectionTime);
 
 			//Need to match the translations and rotations here 
+			//First need to get the correct rotation order based upon what's matching, otherwise if not all are matched 
+			//and one rotation is set last we will get errors.
+			ESkelAnimRotationOrder RotationOrder = FindBestRotationOrder(CurrentSection->bMatchRotationRoll, CurrentSection->bMatchRotationPitch, CurrentSection->bMatchRotationYaw);
+
 			FVector FirstTransformTranslation = FirstTransform.GetTranslation();
 			FVector SecondTransformTranslation = SecondTransform.GetTranslation();
-			FRotator FirstTransformRotation = FirstTransform.GetRotation().Rotator();
-			FRotator SecondTransformRotation = SecondTransform.GetRotation().Rotator();
+			FQuat FirstTransformQuat = FirstTransform.GetRotation();
+			FQuat SecondTransformQuat = SecondTransform.GetRotation();
+
+			SecondTransformQuat.EnforceShortestArcWith(FirstTransformQuat);
+
+			FRotator FirstTransformRotation(FRotator::MakeFromEuler(EulerFromQuat(FirstTransformQuat, RotationOrder)));
+			FRotator SecondTransformRotation(FRotator::MakeFromEuler(EulerFromQuat(SecondTransformQuat, RotationOrder)));
+			FirstTransformRotation.SetClosestToMe(SecondTransformRotation);
+
 			if (!CurrentSection->bMatchTranslation)
 			{
 				SecondTransformTranslation.X = FirstTransformTranslation.X;
@@ -677,7 +920,8 @@ void UMovieSceneSkeletalAnimationTrack::MatchSectionByBoneTransform(bool bMatchW
 			{
 				SecondTransformRotation.Roll = FirstTransformRotation.Roll;
 			}
-			SecondTransform.SetRotation(SecondTransformRotation.Quaternion());
+			SecondTransformQuat = QuatFromEuler(FirstTransformRotation.Euler(), RotationOrder);
+			SecondTransform.SetRotation(SecondTransformQuat);
 
 			//GetRelativeTransformReverse returns this(-1)* Other, and parameter is Other.
 			SecondSectionRootDiff = FirstTransform.GetRelativeTransformReverse(SecondTransform);
