Sequencer: Fixed default actor reference section values not being serialized

diff --git a/Engine/Source/Runtime/MovieSceneTracks/Public/Sections/MovieSceneActorReferenceSection.h b/Engine/Source/Runtime/MovieSceneTracks/Public/Sections/MovieSceneActorReferenceSection.h
--- a/Engine/Source/Runtime/MovieSceneTracks/Public/Sections/MovieSceneActorReferenceSection.h
+++ b/Engine/Source/Runtime/MovieSceneTracks/Public/Sections/MovieSceneActorReferenceSection.h
@@ -144,6 +144,7 @@ private:
 	TArray<FFrameNumber> KeyTimes;
 
 	/** Default value used when there are no keys */
+	UPROPERTY()
 	FMovieSceneActorReferenceKey DefaultValue;
 
 	/** Array of values that correspond to each key time */
