Add missing include guard

diff --git a/Engine/Source/Runtime/MovieSceneTracks/Public/Systems/MovieSceneInitialValueSystem.h b/Engine/Source/Runtime/MovieSceneTracks/Public/Systems/MovieSceneInitialValueSystem.h
--- a/Engine/Source/Runtime/MovieSceneTracks/Public/Systems/MovieSceneInitialValueSystem.h
+++ b/Engine/Source/Runtime/MovieSceneTracks/Public/Systems/MovieSceneInitialValueSystem.h
@@ -1,5 +1,7 @@
 // Copyright Epic Games, Inc. All Rights Reserved.
 
+#pragma once
+
 #include "EntitySystem/MovieSceneEntityInstantiatorSystem.h"
 #include "MovieSceneInitialValueSystem.generated.h"
 
