Integration of 15281763
Sequencer: Fix crash on PostLoad since it calls SetupRootMotions and SortSections, so null sections don't get a chance to be removed.

diff --git a/Engine/Source/Runtime/MovieSceneTracks/Private/Tracks/MovieSceneSkeletalAnimationTrack.cpp b/Engine/Source/Runtime/MovieSceneTracks/Private/Tracks/MovieSceneSkeletalAnimationTrack.cpp
--- a/Engine/Source/Runtime/MovieSceneTracks/Private/Tracks/MovieSceneSkeletalAnimationTrack.cpp
+++ b/Engine/Source/Runtime/MovieSceneTracks/Private/Tracks/MovieSceneSkeletalAnimationTrack.cpp
@@ -88,6 +88,30 @@ TArray<UMovieSceneSection*> UMovieSceneSkeletalAnimationTrack::GetAnimSectionsAt
 
 void UMovieSceneSkeletalAnimationTrack::PostLoad()
 {
+	// UMovieSceneTrack::PostLoad removes null sections. However, RemoveAtSection requires SetupRootMotions, which accesses AnimationSections, so remove null sections here before anything else 
+	for (int32 SectionIndex = 0; SectionIndex < AnimationSections.Num(); )
+	{
+		UMovieSceneSection* Section = AnimationSections[SectionIndex];
+
+		if (Section == nullptr)
+		{
+#if WITH_EDITOR
+			UE_LOG(LogMovieScene, Warning, TEXT("Removing null section from %s:%s"), *GetPathName(), *GetDisplayName().ToString());
+#endif
+			AnimationSections.RemoveAt(SectionIndex);
+		}
+		else if (Section->GetRange().IsEmpty())
+		{
+#if WITH_EDITOR
+			//UE_LOG(LogMovieScene, Warning, TEXT("Removing section %s:%s with empty range"), *GetPathName(), *GetDisplayName().ToString());
+#endif
+			AnimationSections.RemoveAt(SectionIndex);
+		}
+		else
+		{
+			++SectionIndex;
+		}
+	}
 	Super::PostLoad();
 
 	if (GetLinkerCustomVersion(FMovieSceneEvaluationCustomVersion::GUID) < FMovieSceneEvaluationCustomVersion::AddBlendingSupport)
