Sequencer: Fix broken hierarchical easing.

#[robomerge] 4.27

diff --git a/Engine/Source/Runtime/MovieSceneTracks/Private/Systems/WeightAndEasingEvaluatorSystem.cpp b/Engine/Source/Runtime/MovieSceneTracks/Private/Systems/WeightAndEasingEvaluatorSystem.cpp
--- a/Engine/Source/Runtime/MovieSceneTracks/Private/Systems/WeightAndEasingEvaluatorSystem.cpp
+++ b/Engine/Source/Runtime/MovieSceneTracks/Private/Systems/WeightAndEasingEvaluatorSystem.cpp
@@ -199,8 +199,14 @@ struct FEvaluateEasings
 
 			for (int32 Idx = 0; Idx < Num; ++Idx)
 			{
-				const FSequenceInstance& Instance = InstanceRegistry->GetInstance(OptInstanceHandles[Idx]);
-				EvaluatorSystem->SetSubSequenceEasing(OptInstanceHandles[Idx], Results[Idx]);
+				const FSequenceInstance* RootInstance = &InstanceRegistry->GetInstance(OptInstanceHandles[Idx]);
+				if (!RootInstance->IsRootSequence())
+				{
+					RootInstance = &InstanceRegistry->GetInstance(RootInstance->GetRootInstanceHandle());
+				}
+
+				const FInstanceHandle SubSequenceHandle = RootInstance->FindSubInstance(OptHierarchicalEasingProviders[Idx]);
+				EvaluatorSystem->SetSubSequenceEasing(SubSequenceHandle, Results[Idx]);
 			}
 		}
 	}
