Sequencer: Added the ability to run asynchronous predictions for animated component transforms and future/past times

  - Predictions do not understand changes of attachment that might occur between the current time and the predicted time
  - Separated the functions and data relating to managing interrogation channels and indices out into a separate class called FInterrogationChannels
  - The predictions import interrogation entities into the main ECS and will get processed along with all the other evaluations
  - Spawnable annotations are now available in-game (not just in-editor)

diff --git a/Engine/Source/Runtime/LevelSequence/Private/LevelSequence.cpp b/Engine/Source/Runtime/LevelSequence/Private/LevelSequence.cpp
--- a/Engine/Source/Runtime/LevelSequence/Private/LevelSequence.cpp
+++ b/Engine/Source/Runtime/LevelSequence/Private/LevelSequence.cpp
@@ -399,6 +399,11 @@ void ULevelSequence::LocateBoundObjects(const FGuid& ObjectId, UObject* Context,
 	BindingReferences.ResolveBinding(ObjectId, Context, StreamedLevelAssetPath, OutObjects);
 }
 
+FGuid ULevelSequence::FindBindingFromObject(UObject* InObject, UObject* Context) const
+{
+	return BindingReferences.FindBindingFromObject(InObject, Context);
+}
+
 void ULevelSequence::GatherExpiredObjects(const FMovieSceneObjectCache& InObjectCache, TArray<FGuid>& OutInvalidIDs) const
 {
 	for (const FGuid& ObjectId : BindingReferences.GetBoundAnimInstances())
diff --git a/Engine/Source/Runtime/LevelSequence/Private/LevelSequenceBindingReference.cpp b/Engine/Source/Runtime/LevelSequence/Private/LevelSequenceBindingReference.cpp
--- a/Engine/Source/Runtime/LevelSequence/Private/LevelSequenceBindingReference.cpp
+++ b/Engine/Source/Runtime/LevelSequence/Private/LevelSequenceBindingReference.cpp
@@ -92,6 +92,11 @@ UObject* FLevelSequenceBindingReference::Resolve(UObject* InContext, FName Strea
 	return nullptr;
 }
 
+bool FLevelSequenceBindingReference::operator==(const FLevelSequenceBindingReference& Other) const
+{
+	return ExternalObjectPath == Other.ExternalObjectPath && ObjectPath == Other.ObjectPath;
+}
+
 void FLevelSequenceBindingReference::PostSerialize(const FArchive& Ar)
 {
 	if (Ar.IsLoading() && !PackageName_DEPRECATED.IsEmpty())
@@ -297,6 +302,21 @@ void FLevelSequenceBindingReferences::ResolveBinding(const FGuid& ObjectId, UObj
 	}
 }
 
+FGuid FLevelSequenceBindingReferences::FindBindingFromObject(UObject* InObject, UObject* InContext) const
+{
+	FLevelSequenceBindingReference Predicate(InObject, InContext);
+
+	for (const TPair<FGuid, FLevelSequenceBindingReferenceArray>& Pair : BindingIdToReferences)
+	{
+		if (Pair.Value.References.Contains(Predicate))
+		{
+			return Pair.Key;
+		}
+	}
+
+	return FGuid();
+}
+
 void FLevelSequenceBindingReferences::RemoveInvalidBindings(const TSet<FGuid>& ValidBindingIDs)
 {
 	for (auto It = BindingIdToReferences.CreateIterator(); It; ++It)
diff --git a/Engine/Source/Runtime/LevelSequence/Public/LevelSequence.h b/Engine/Source/Runtime/LevelSequence/Public/LevelSequence.h
--- a/Engine/Source/Runtime/LevelSequence/Public/LevelSequence.h
+++ b/Engine/Source/Runtime/LevelSequence/Public/LevelSequence.h
@@ -46,6 +46,7 @@ public:
 	virtual void BindPossessableObject(const FGuid& ObjectId, UObject& PossessedObject, UObject* Context) override;
 	virtual bool CanPossessObject(UObject& Object, UObject* InPlaybackContext) const override;
 	virtual void LocateBoundObjects(const FGuid& ObjectId, UObject* Context, TArray<UObject*, TInlineAllocator<1>>& OutObjects) const override;
+	virtual FGuid FindBindingFromObject(UObject* InObject, UObject* Context) const override;
 	virtual void GatherExpiredObjects(const FMovieSceneObjectCache& InObjectCache, TArray<FGuid>& OutInvalidIDs) const override;
 	virtual UMovieScene* GetMovieScene() const override;
 	virtual UObject* GetParentObject(UObject* Object) const override;
diff --git a/Engine/Source/Runtime/LevelSequence/Public/LevelSequenceBindingReference.h b/Engine/Source/Runtime/LevelSequence/Public/LevelSequenceBindingReference.h
--- a/Engine/Source/Runtime/LevelSequence/Public/LevelSequenceBindingReference.h
+++ b/Engine/Source/Runtime/LevelSequence/Public/LevelSequenceBindingReference.h
@@ -41,6 +41,11 @@ struct FLevelSequenceBindingReference
 	 */
 	LEVELSEQUENCE_API UObject* Resolve(UObject* InContext, FName StreamedLevelAssetPath) const;
 
+	/**
+	 * Check whether this binding reference is equal to the specified object
+	 */
+	LEVELSEQUENCE_API bool operator==(const FLevelSequenceBindingReference& Other) const;
+
 	/** Handles ExternalObjectPath fixup */
 	void PostSerialize(const FArchive& Ar);
 
@@ -139,6 +144,16 @@ struct FLevelSequenceBindingReferences
 	 */
 	void ResolveBinding(const FGuid& ObjectId, UObject* InContext, FName StreamedLevelAssetPath, TArray<UObject*, TInlineAllocator<1>>& OutObjects) const;
 
+	/**
+	 * Resolve a binding for the specified ID using a given context
+	 *
+	 * @param ObjectId					The ID to associate the object with
+	 * @param InContext					A context in which InObject resides
+	 * @oaram StreamedLevelAssetPath    The path to the streamed level asset that contains the level sequence actor playing back the sequence. 'None' for any non-instance-level setups.
+	 * @param OutObjects				Array to populate with resolved object bindings
+	 */
+	FGuid FindBindingFromObject(UObject* InObject, UObject* InContext) const;
+
 	/**
 	 * Const accessor for the currently bound anim instance IDs
 	 */
diff --git a/Engine/Source/Runtime/MovieScene/Private/EntitySystem/MovieSceneEntityManager.cpp b/Engine/Source/Runtime/MovieScene/Private/EntitySystem/MovieSceneEntityManager.cpp
--- a/Engine/Source/Runtime/MovieScene/Private/EntitySystem/MovieSceneEntityManager.cpp
+++ b/Engine/Source/Runtime/MovieScene/Private/EntitySystem/MovieSceneEntityManager.cpp
@@ -1741,6 +1741,11 @@ void FEntityManager::InitializeMutualComponents(FMovieSceneEntityID EntityID)
 }
 
 void FEntityManager::AddMutualComponents()
+{
+	AddMutualComponents(FEntityComponentFilter());
+}
+
+void FEntityManager::AddMutualComponents(const FEntityComponentFilter& InFilter)
 {
 	CheckCanChangeStructure();
 
@@ -1754,6 +1759,10 @@ void FEntityManager::AddMutualComponents()
 		}
 
 		FComponentMask NewAllocationType = EntityAllocationMasks[AllocationIndex];
+		if (InFilter.IsValid() && !InFilter.Match(NewAllocationType))
+		{
+			continue;
+		}
 
 		const int32 NumNewComponents = ComponentRegistry->Factories.ComputeMutuallyInclusiveComponents(NewAllocationType);
 		if (NumNewComponents != 0)
diff --git a/Engine/Source/Runtime/MovieScene/Private/Evaluation/MovieSceneEvaluationField.cpp b/Engine/Source/Runtime/MovieScene/Private/Evaluation/MovieSceneEvaluationField.cpp
--- a/Engine/Source/Runtime/MovieScene/Private/Evaluation/MovieSceneEvaluationField.cpp
+++ b/Engine/Source/Runtime/MovieScene/Private/Evaluation/MovieSceneEvaluationField.cpp
@@ -142,6 +142,26 @@ void FMovieSceneEntityComponentField::QueryPersistentEntities(FFrameNumber Query
 	}
 }
 
+void FMovieSceneEntityComponentField::QueryPersistentEntities(FFrameNumber QueryTime, TFunctionRef<bool(const FMovieSceneEvaluationFieldEntityQuery&)> QueryCallback, TRange<FFrameNumber>& OutRange) const
+{
+	FMovieSceneEvaluationTreeRangeIterator Iterator = PersistentEntityTree.SerializedData.IterateFromTime(QueryTime);
+	check(Iterator);
+
+	OutRange = Iterator.Range();
+	for (FMovieSceneEvaluationFieldEntityTree::FEntityAndMetaDataIndex Pair : PersistentEntityTree.SerializedData.GetAllData(Iterator.Node()))
+	{
+		FMovieSceneEvaluationFieldEntityQuery Query{
+			GetEntity(Pair.EntityIndex),
+			Pair.MetaDataIndex
+		};
+
+		if (!QueryCallback(Query))
+		{
+			return;
+		}
+	}
+}
+
 bool FMovieSceneEntityComponentField::HasAnyOneShotEntities() const
 {
 	return !OneShotEntityTree.SerializedData.IsEmpty();
diff --git a/Engine/Source/Runtime/MovieScene/Private/MovieSceneSpawnRegister.cpp b/Engine/Source/Runtime/MovieScene/Private/MovieSceneSpawnRegister.cpp
--- a/Engine/Source/Runtime/MovieScene/Private/MovieSceneSpawnRegister.cpp
+++ b/Engine/Source/Runtime/MovieScene/Private/MovieSceneSpawnRegister.cpp
@@ -40,10 +40,8 @@ UObject* FMovieSceneSpawnRegister::SpawnObject(const FGuid& BindingId, UMovieSce
 	
 	if (SpawnedActor)
 	{
-#if WITH_EDITOR
 		UMovieSceneSequence* Sequence = Player.GetEvaluationTemplate().GetSequence(TemplateID);
 		FMovieSceneSpawnableAnnotation::Add(SpawnedActor, BindingId, TemplateID, Sequence);
-#endif
 
 		FMovieSceneSpawnRegisterKey Key(TemplateID, BindingId);
 		Register.Add(Key, FSpawnedObject(BindingId, *SpawnedActor, Spawnable->GetSpawnOwnership()));
diff --git a/Engine/Source/Runtime/MovieScene/Private/MovieSceneSpawnableAnnotation.cpp b/Engine/Source/Runtime/MovieScene/Private/MovieSceneSpawnableAnnotation.cpp
--- a/Engine/Source/Runtime/MovieScene/Private/MovieSceneSpawnableAnnotation.cpp
+++ b/Engine/Source/Runtime/MovieScene/Private/MovieSceneSpawnableAnnotation.cpp
@@ -1,12 +1,8 @@
 // Copyright Epic Games, Inc. All Rights Reserved.
 
 #include "MovieSceneSpawnableAnnotation.h"
-
-#if WITH_EDITOR
-
 #include "UObject/UObjectAnnotation.h"
 
-
 static FUObjectAnnotationSparse<FMovieSceneSpawnableAnnotation,true> SpawnedObjectAnnotation;
 
 void FMovieSceneSpawnableAnnotation::Add(UObject* SpawnedObject, const FGuid& ObjectBindingID, FMovieSceneSequenceID SequenceID, UMovieSceneSequence* InOriginatingSequence)
@@ -33,6 +29,4 @@ TOptional<FMovieSceneSpawnableAnnotation> FMovieSceneSpawnableAnnotation::Find(U
 	}
 
 	return ReturnValue;
-}
-
-#endif // WITH_EDITOR
\ No newline at end of file
+}
\ No newline at end of file
diff --git a/Engine/Source/Runtime/MovieScene/Public/Compilation/MovieSceneCompiledDataManager.h b/Engine/Source/Runtime/MovieScene/Public/Compilation/MovieSceneCompiledDataManager.h
--- a/Engine/Source/Runtime/MovieScene/Public/Compilation/MovieSceneCompiledDataManager.h
+++ b/Engine/Source/Runtime/MovieScene/Public/Compilation/MovieSceneCompiledDataManager.h
@@ -147,7 +147,7 @@ struct FMovieSceneCompiledDataEntry
 {
 	FMovieSceneCompiledDataEntry();
 
-	UMovieSceneSequence* GetSequence() const;
+	MOVIESCENE_API UMovieSceneSequence* GetSequence() const;
 
 	/** 16 Bytes */
 	FGuid CompiledSignature;
diff --git a/Engine/Source/Runtime/MovieScene/Public/EntitySystem/MovieSceneEntityManager.h b/Engine/Source/Runtime/MovieScene/Public/EntitySystem/MovieSceneEntityManager.h
--- a/Engine/Source/Runtime/MovieScene/Public/EntitySystem/MovieSceneEntityManager.h
+++ b/Engine/Source/Runtime/MovieScene/Public/EntitySystem/MovieSceneEntityManager.h
@@ -387,6 +387,11 @@ public:
 	 */
 	void AddMutualComponents();
 
+	/**
+	 * Run through all entities in this entity manager, ensuring that all mutual components exist for any component types that match the specified filter
+	 */
+	void AddMutualComponents(const FEntityComponentFilter& InFilter);
+
 public:
 
 
diff --git a/Engine/Source/Runtime/MovieScene/Public/Evaluation/MovieSceneEvaluationField.h b/Engine/Source/Runtime/MovieScene/Public/Evaluation/MovieSceneEvaluationField.h
--- a/Engine/Source/Runtime/MovieScene/Public/Evaluation/MovieSceneEvaluationField.h
+++ b/Engine/Source/Runtime/MovieScene/Public/Evaluation/MovieSceneEvaluationField.h
@@ -12,6 +12,7 @@
 #include "Evaluation/MovieSceneEvaluationTree.h"
 #include "EntitySystem/MovieSceneEntityIDs.h"
 #include "EntitySystem/MovieSceneEntitySystemTypes.h"
+#include "Templates/Function.h"
 #include "MovieSceneEvaluationField.generated.h"
 
 struct FFrameNumber;
@@ -242,6 +243,16 @@ struct MOVIESCENE_API FMovieSceneEntityComponentField
 	 */
 	void QueryPersistentEntities(FFrameNumber QueryTime, TRange<FFrameNumber>& OutRange, FMovieSceneEvaluationFieldEntitySet& OutEntities) const;
 
+	/**
+	 * Query the persistent entities for any given time within a sequence.
+	 * @note: Persistent entities should remain alive until they are no longer present at the current time.
+	 *
+	 * @param QueryTime     The time at which to query the field (in the TickResolution of the sequence this was generated from)
+	 * @param QueryCallback A handler for dealing with the resulting entities
+	 * @param OutRange      Will receive the hull of the range that was intersected for which the resulting OutEntities remains constant
+	 */
+	void QueryPersistentEntities(FFrameNumber QueryTime, TFunctionRef<bool(const FMovieSceneEvaluationFieldEntityQuery&)> QueryCallback, TRange<FFrameNumber>& OutRange) const;
+
 	/**
 	 * Check whether this field contains any one-shot entities
 	 */
diff --git a/Engine/Source/Runtime/MovieScene/Public/MovieSceneSequence.h b/Engine/Source/Runtime/MovieScene/Public/MovieSceneSequence.h
--- a/Engine/Source/Runtime/MovieScene/Public/MovieSceneSequence.h
+++ b/Engine/Source/Runtime/MovieScene/Public/MovieSceneSequence.h
@@ -94,6 +94,15 @@ public:
 	 */
 	MOVIESCENE_API FGuid FindPossessableObjectId(UObject& Object, UObject* Context) const;
 
+	/**
+	 * Optional method for efficient lookup of an object binding from an actual object in the world
+	 *
+	 * @param ObjectId				The unique identifier of the object.
+	 * @param Context				Optional context to use to find the required object (for instance, a parent spawnable object or its world)
+	 * @return The object's guid, or zero guid if the object is not a valid possessable in the current context
+	 */
+	virtual FGuid FindBindingFromObject(UObject* InObject, UObject* Context) const { return FGuid(); }
+
 	/**
 	 * Called to validate the specified object cache by removing anything that should be deemed out of date
 	 * 
diff --git a/Engine/Source/Runtime/MovieScene/Public/MovieSceneSpawnableAnnotation.h b/Engine/Source/Runtime/MovieScene/Public/MovieSceneSpawnableAnnotation.h
--- a/Engine/Source/Runtime/MovieScene/Public/MovieSceneSpawnableAnnotation.h
+++ b/Engine/Source/Runtime/MovieScene/Public/MovieSceneSpawnableAnnotation.h
@@ -2,8 +2,6 @@
 
 #pragma once
 
-#if WITH_EDITOR
-
 #include "Misc/Optional.h"
 #include "Misc/Guid.h"
 #include "UObject/WeakObjectPtr.h"
@@ -42,6 +40,3 @@ struct FMovieSceneSpawnableAnnotation
 	/** The sequence ID that spawned this object */
 	FMovieSceneSequenceID SequenceID;
 };
-
-
-#endif // WITH_EDITOR
\ No newline at end of file
diff --git a/Engine/Source/Runtime/MovieSceneTracks/Private/EntitySystem/Interrogation/MovieSceneInterrogatedPropertyInstantiator.cpp b/Engine/Source/Runtime/MovieSceneTracks/Private/EntitySystem/Interrogation/MovieSceneInterrogatedPropertyInstantiator.cpp
--- a/Engine/Source/Runtime/MovieSceneTracks/Private/EntitySystem/Interrogation/MovieSceneInterrogatedPropertyInstantiator.cpp
+++ b/Engine/Source/Runtime/MovieSceneTracks/Private/EntitySystem/Interrogation/MovieSceneInterrogatedPropertyInstantiator.cpp
@@ -16,9 +16,6 @@ UMovieSceneInterrogatedPropertyInstantiatorSystem::UMovieSceneInterrogatedProper
 {
 	using namespace UE::MovieScene;
 
-	// This system should never run at runtime
-	SystemExclusionContext |= EEntitySystemContext::Runtime;
-
 	BuiltInComponents = FBuiltInComponentTypes::Get();
 
 	RecomposerImpl.OnGetPropertyInfo = FOnGetPropertyRecomposerPropertyInfo::CreateUObject(
@@ -36,7 +33,9 @@ UMovieSceneInterrogatedPropertyInstantiatorSystem::UMovieSceneInterrogatedProper
 
 bool UMovieSceneInterrogatedPropertyInstantiatorSystem::IsRelevantImpl(UMovieSceneEntitySystemLinker* InLinker) const
 {
-	return true;
+	using namespace UE::MovieScene;
+
+	return InLinker->EntityManager.Contains(FEntityComponentFilter().All({ BuiltInComponents->PropertyBinding, BuiltInComponents->Interrogation.InputKey }));
 }
 
 UE::MovieScene::FPropertyRecomposerPropertyInfo UMovieSceneInterrogatedPropertyInstantiatorSystem::FindPropertyFromSource(FMovieSceneEntityID EntityID, UObject* Object) const
diff --git a/Engine/Source/Runtime/MovieSceneTracks/Private/EntitySystem/Interrogation/MovieSceneInterrogationLinker.cpp b/Engine/Source/Runtime/MovieSceneTracks/Private/EntitySystem/Interrogation/MovieSceneInterrogationLinker.cpp
--- a/Engine/Source/Runtime/MovieSceneTracks/Private/EntitySystem/Interrogation/MovieSceneInterrogationLinker.cpp
+++ b/Engine/Source/Runtime/MovieSceneTracks/Private/EntitySystem/Interrogation/MovieSceneInterrogationLinker.cpp
@@ -72,43 +72,37 @@
 	return ID;
 }
 
-FSystemInterrogator::FSystemInterrogator()
+FInterrogationChannels::FInterrogationChannels()
 {
-	InitialValueCache = FInitialValueCache::GetGlobalInitialValues();
-
-	Linker = NewObject<UMovieSceneEntitySystemLinker>(GetTransientPackage());
-	Linker->SetSystemContext(EEntitySystemContext::Interrogation);
-
-	Linker->AddExtension(IInterrogationExtension::GetExtensionID(), static_cast<IInterrogationExtension*>(this));
-	Linker->AddExtension(InitialValueCache.Get());
-
 	// Always add a bit for the default channel
-	ImportedChannelBits.Add(false);
+	ActiveChannelBits.Add(false);
 }
 
-FSystemInterrogator::~FSystemInterrogator()
+FInterrogationChannels::~FInterrogationChannels()
 {}
 
-void FSystemInterrogator::AddReferencedObjects(FReferenceCollector& Collector)
+void FInterrogationChannels::Reset()
 {
-	Collector.AddReferencedObject(Linker);
-	Collector.AddReferencedObjects(ObjectToChannel);
+	Interrogations.Empty();
+	ActiveChannelBits.Empty();
+
+	// Always add a bit for the default channel
+	ActiveChannelBits.Add(false);
+
+	ObjectToChannel.Empty();
+
+	SparseChannelInfo.Empty();
 }
 
-FString FSystemInterrogator::GetReferencerName() const
+FInterrogationChannel FInterrogationChannels::AllocateChannel(FInterrogationChannel ParentChannel, const FMovieScenePropertyBinding& PropertyBinding)
 {
-	return TEXT("FSystemInterrogator");
-}
-
-FInterrogationChannel FSystemInterrogator::AllocateChannel(FInterrogationChannel ParentChannel, const FMovieScenePropertyBinding& PropertyBinding)
-{
-	if (!ensureMsgf(ImportedChannelBits.Num() < MAX_int32, TEXT("Reached the maximum available number of interrogation channels")))
+	if (!ensureMsgf(ActiveChannelBits.Num() < MAX_int32, TEXT("Reached the maximum available number of interrogation channels")))
 	{
 		return FInterrogationChannel::Invalid();
 	}
 
-	FInterrogationChannel Channel = FInterrogationChannel::FromIndex(ImportedChannelBits.Num());
-	ImportedChannelBits.Add(false);
+	FInterrogationChannel Channel = FInterrogationChannel::FromIndex(ActiveChannelBits.Num());
+	ActiveChannelBits.Add(false);
 
 	if (ParentChannel || !PropertyBinding.PropertyPath.IsNone())
 	{
@@ -119,7 +113,7 @@
 	return Channel;
 }
 
-FInterrogationChannel FSystemInterrogator::AllocateChannel(UObject* Object, FInterrogationChannel ParentChannel, const FMovieScenePropertyBinding& PropertyBinding)
+FInterrogationChannel FInterrogationChannels::AllocateChannel(UObject* Object, FInterrogationChannel ParentChannel, const FMovieScenePropertyBinding& PropertyBinding)
 {
 	FInterrogationChannel NewChannel = AllocateChannel(ParentChannel, PropertyBinding);
 	if (NewChannel)
@@ -130,292 +124,67 @@
 	return NewChannel;
 }
 
-void FSystemInterrogator::Reset()
+FInterrogationChannel FInterrogationChannels::FindChannel(UObject* Object)
 {
-	if (EntityTracker)
-	{
-		EntityTracker->Reset();
-	}
-
-	EntitiesScratch.Empty();
-	EntityComponentField = FMovieSceneEntityComponentField();
-
-	Interrogations.Empty();
-	ImportedChannelBits.Empty();
-	// Always add a bit for the default channel
-	ImportedChannelBits.Add(false);
-
-	ObjectToChannel.Empty();
-
-	SparseChannelInfo.Empty();
-
-	Linker->Reset();
+	return ObjectToChannel.FindRef(Object);
 }
 
-void FSystemInterrogator::ImportTrack(UMovieSceneTrack* Track, FInterrogationChannel InChannel)
+void FInterrogationChannels::ActivateChannel(FInterrogationChannel InChannel)
 {
-	ImportTrack(Track, Track->FindObjectBindingGuid(), InChannel);
+	ActiveChannelBits[InChannel.AsIndex()] = true;
 }
 
-void FSystemInterrogator::ImportTrack(UMovieSceneTrack* Track, const FGuid& ObjectBindingID, FInterrogationChannel InChannel)
+void FInterrogationChannels::DeactivateChannel(FInterrogationChannel InChannel)
 {
-	check(InChannel.AsIndex() < ImportedChannelBits.Num());
-
-	TGuardValue<FEntityManager*> DebugVizGuard(GEntityManagerForDebuggingVisualizers, &Linker->EntityManager);
-
-	FFrameRate TickResolution = Track->GetTypedOuter<UMovieScene>()->GetTickResolution();
-
-	const FMovieSceneTrackEvaluationField& EvaluationField = Track->GetEvaluationField();
-
-	FMovieSceneEntityComponentFieldBuilder FieldBuilder(&EntityComponentField);
-	FieldBuilder.GetSharedMetaData().ObjectBindingID = ObjectBindingID;
-
-	for (const FMovieSceneTrackEvaluationFieldEntry& Entry : EvaluationField.Entries)
-	{
-		IMovieSceneEntityProvider* EntityProvider = Cast<IMovieSceneEntityProvider>(Entry.Section);
-
-		if (!EntityProvider || Entry.Range.IsEmpty())
-		{
-			continue;
-		}
-
-		FMovieSceneEvaluationFieldEntityMetaData MetaData;
-
-		MetaData.ForcedTime = Entry.ForcedTime;
-		MetaData.Flags      = Entry.Flags;
-		MetaData.bEvaluateInSequencePreRoll  = Track->EvalOptions.bEvaluateInPreroll;
-		MetaData.bEvaluateInSequencePostRoll = Track->EvalOptions.bEvaluateInPostroll;
-
-		if (InChannel != FInterrogationChannel::Default())
-		{
-			MetaData.InterrogationChannel = InChannel;
-		}
-
-		if (!EntityProvider->PopulateEvaluationField(Entry.Range, MetaData, &FieldBuilder))
-		{
-			const int32 EntityIndex   = FieldBuilder.FindOrAddEntity(Entry.Section, 0);
-			const int32 MetaDataIndex = FieldBuilder.AddMetaData(MetaData);
-			FieldBuilder.AddPersistentEntity(Entry.Range, EntityIndex, MetaDataIndex);
-		}
-	}
-
-	ImportedChannelBits[InChannel.AsIndex()] = true;
+	ActiveChannelBits[InChannel.AsIndex()] = false;
 }
 
-int32 FSystemInterrogator::AddInterrogation(const FInterrogationParams& Params)
+int32 FInterrogationChannels::AddInterrogation(const FInterrogationParams& Params)
 {
 	if (!ensureMsgf(Interrogations.Num() != MAX_int32, TEXT("Reached the maximum available number of interrogation channels")))
 	{
 		return INDEX_NONE;
 	}
 
-	TGuardValue<FEntityManager*> DebugVizGuard(GEntityManagerForDebuggingVisualizers, &Linker->EntityManager);
-
-	// We always add 1 to the interrogation index here because the default is always at 0
 	const int32 InterrogationIndex = Interrogations.Num();
-
-	TRange<FFrameNumber> UnusedEntityRange;
-
-	// Update the entities that exist at this frame
-	EntitiesScratch.Reset();
-	EntityComponentField.QueryPersistentEntities(Params.Time.FrameNumber, UnusedEntityRange, EntitiesScratch);
-
-	for (const FMovieSceneEvaluationFieldEntityQuery& Query : EntitiesScratch)
-	{
-		InterrogateEntity(InterrogationIndex, Query);
-	}
-
 	Interrogations.Add(Params);
 
 	return InterrogationIndex;
 }
 
-void FSystemInterrogator::InterrogateEntity(int32 InterrogationIndex, const FMovieSceneEvaluationFieldEntityQuery& Query)
-{
-	UObject* EntityOwner = Query.Entity.Key.EntityOwner.Get();
-	IMovieSceneEntityProvider* Provider = Cast<IMovieSceneEntityProvider>(EntityOwner);
-	if (!Provider)
-	{
-		return;
-	}
-
-	FEntityImportParams Params;
-	Params.EntityID = Query.Entity.Key.EntityID;
-	Params.EntityMetaData = EntityComponentField.FindMetaData(Query);
-	Params.SharedMetaData = EntityComponentField.FindSharedMetaData(Query);
-
-	Params.InterrogationKey.Channel = FInterrogationChannel::Default();
-	Params.InterrogationKey.InterrogationIndex = InterrogationIndex;
-
-	if (Params.EntityMetaData && Params.EntityMetaData->InterrogationChannel.IsValid())
-	{
-		Params.InterrogationKey.Channel = Params.EntityMetaData->InterrogationChannel;
-	}
-
-	FImportedEntity ImportedEntity;
-	Provider->InterrogateEntity(Linker, Params, &ImportedEntity);
-
-	if (!ImportedEntity.IsEmpty())
-	{
-		if (UMovieSceneSection* Section = Cast<UMovieSceneSection>(EntityOwner))
-		{
-			Section->BuildDefaultComponents(Linker, Params, &ImportedEntity);
-		}
-
-		const FMovieSceneEntityID NewEntityID = ImportedEntity.Manufacture(Params, &Linker->EntityManager);
-		if (EntityTracker)
-		{
-			EntityTracker->TrackEntity(Params.InterrogationKey, Query.Entity.Key, NewEntityID);
-		}
-	}
-}
-
-void FSystemInterrogator::Update()
-{
-	TGuardValue<FEntityManager*> DebugVizGuard(GEntityManagerForDebuggingVisualizers, &Linker->EntityManager);
-
-	Linker->EntityManager.AddMutualComponents();
-	Linker->LinkRelevantSystems();
-
-	FEntityTaskBuilder()
-	.Read(FBuiltInComponentTypes::Get()->Interrogation.InputKey)
-	.Write(FBuiltInComponentTypes::Get()->EvalTime)
-	.FilterNone({ FBuiltInComponentTypes::Get()->Tags.FixedTime })
-	.Iterate_PerEntity(&Linker->EntityManager, [this](const FInterrogationKey& InterrogationKey, FFrameTime& OutEvalTime) { OutEvalTime = this->Interrogations[InterrogationKey.InterrogationIndex].Time; });
-
-	FMovieSceneEntitySystemRunner Runner;
-	Runner.AttachToLinker(Linker);
-	Runner.Flush();
-
-	Linker->EntityManager.IncrementSystemSerial();
-}
-
-void FSystemInterrogator::TrackImportedEntities(bool bInTrackImportedEntities)
-{
-	if (bInTrackImportedEntities && !EntityTracker)
-	{
-		EntityTracker = MakeUnique<FSystemInterrogatorEntityTracker>();
-	}
-	else if (!bInTrackImportedEntities)
-	{
-		EntityTracker.Reset();
-	}
-}
-
-FMovieSceneEntityID FSystemInterrogator::FindEntityFromOwner(FInterrogationKey InterrogationKey, UObject* Owner, uint32 EntityID) const
-{
-	if (!ensureMsgf(EntityTracker.IsValid(), TEXT("FindEntityFromOwner called on an interrogator that was not tracking entities")))
-	{
-		return FMovieSceneEntityID::Invalid();
-	}
-
-	return EntityTracker->FindTrackedEntity(InterrogationKey, FMovieSceneEvaluationFieldEntityKey { Owner, EntityID });;
-}
-
-FInterrogationChannel FSystemInterrogator::ImportLocalTransforms(USceneComponent* SceneComponent, IMovieScenePlayer* InPlayer, FMovieSceneSequenceID SequenceID)
-{
-	check(SceneComponent);
-
-	UMovieSceneSequence* Sequence = InPlayer->State.FindSequence(SequenceID);
-	if (!Sequence)
-	{
-		return FInterrogationChannel::Invalid();
-	}
-
-	FInterrogationChannel ParentChannel;
-	if (USceneComponent* AttachParent = SceneComponent->GetAttachParent())
-	{
-		ParentChannel = ObjectToChannel.FindRef(AttachParent);
-	}
-
-	FInterrogationChannel Channel = ObjectToChannel.FindRef(SceneComponent);
-	if (!Channel.IsValid())
-	{
-		Channel = AllocateChannel(SceneComponent, ParentChannel, FMovieScenePropertyBinding("Transform", TEXT("Transform")));
-	}
-
-	if (!Channel.IsValid())
-	{
-		return FInterrogationChannel::Invalid();
-	}
-
-	// Find the binding that corresponds to the component directly
-	FGuid ObjectBindingID = InPlayer->State.FindCachedObjectId(*SceneComponent, SequenceID, *InPlayer);
-	if (ObjectBindingID.IsValid())
-	{
-		ImportTransformTracks(*Sequence->GetMovieScene()->FindBinding(ObjectBindingID), Channel);
-	}
-
-	// Also blend in any transforms that exist for this scene component's actor as well (if it is the root)
-	AActor* Owner = SceneComponent->GetOwner();
-	if (SceneComponent == Owner->GetRootComponent())
-	{
-		FGuid OwnerObjectBindingID = InPlayer->State.FindCachedObjectId(*Owner, SequenceID, *InPlayer);
-		if (OwnerObjectBindingID.IsValid())
-		{
-			ImportTransformTracks(*Sequence->GetMovieScene()->FindBinding(OwnerObjectBindingID), Channel);
-		}
-	}
-
-	return Channel;
-}
-
-FInterrogationChannel FSystemInterrogator::ImportTransformHierarchy(USceneComponent* SceneComponent, IMovieScenePlayer* InPlayer, FMovieSceneSequenceID SequenceID)
-{
-	check(SceneComponent);
-
-	if (USceneComponent* AttachParent = SceneComponent->GetAttachParent())
-	{
-		ImportTransformHierarchy(AttachParent, InPlayer, SequenceID);
-	}
-
-	return ImportLocalTransforms(SceneComponent, InPlayer, SequenceID);
-}
-
-void FSystemInterrogator::ImportTransformTracks(const FMovieSceneBinding& Binding, FInterrogationChannel Channel)
-{
-	for (UMovieSceneTrack* Track : Binding.GetTracks())
-	{
-		if (Track->IsA<UMovieScene3DTransformTrack>())
-		{
-			ImportTrack(Track, Binding.GetObjectGuid(), Channel);
-		}
-	}
-}
-
-void FSystemInterrogator::QueryWorldSpaceTransforms(USceneComponent* SceneComponent, TArray<FTransform>& OutTransforms)
+void FInterrogationChannels::QueryWorldSpaceTransforms(UMovieSceneEntitySystemLinker* Linker, USceneComponent* SceneComponent, TArray<FTransform>& OutTransforms) const
 {
 	FInterrogationChannel Channel = ObjectToChannel.FindRef(SceneComponent);
 	if (Channel)
 	{
-		QueryWorldSpaceTransforms(Channel, OutTransforms);
+		QueryWorldSpaceTransforms(Linker, Channel, OutTransforms);
 	}
 }
 
-void FSystemInterrogator::QueryLocalSpaceTransforms(USceneComponent* SceneComponent, TArray<FIntermediate3DTransform>& OutTransforms)
+void FInterrogationChannels::QueryLocalSpaceTransforms(UMovieSceneEntitySystemLinker* Linker, USceneComponent* SceneComponent, TArray<FIntermediate3DTransform>& OutTransforms) const
 {
 	FInterrogationChannel Channel = ObjectToChannel.FindRef(SceneComponent);
 	if (Channel)
 	{
-		QueryLocalSpaceTransforms(Channel, OutTransforms);
+		QueryLocalSpaceTransforms(Linker, Channel, OutTransforms);
 	}
 }
 
-void FSystemInterrogator::QueryWorldSpaceTransforms(FInterrogationChannel Channel, TArray<FTransform>& OutTransforms)
+void FInterrogationChannels::QueryWorldSpaceTransforms(UMovieSceneEntitySystemLinker* Linker, FInterrogationChannel Channel, TArray<FTransform>& OutTransforms) const
 {
 	TBitArray<> ChannelMask(false, Channel.AsIndex()+1);
 	ChannelMask[Channel.AsIndex()] = true;
-	QueryWorldSpaceTransforms(ChannelMask, [&OutTransforms](FInterrogationChannel)-> TArray<FTransform>& { return OutTransforms; });
+	QueryWorldSpaceTransforms(Linker, ChannelMask, [&OutTransforms](FInterrogationChannel)-> TArray<FTransform>& { return OutTransforms; });
 }
 
-void FSystemInterrogator::QueryLocalSpaceTransforms(FInterrogationChannel Channel, TArray<FIntermediate3DTransform>& OutTransforms)
+void FInterrogationChannels::QueryLocalSpaceTransforms(UMovieSceneEntitySystemLinker* Linker, FInterrogationChannel Channel, TArray<FIntermediate3DTransform>& OutTransforms) const
 {
 	TBitArray<> ChannelMask(false, Channel.AsIndex()+1);
 	ChannelMask[Channel.AsIndex()] = true;
-	QueryLocalSpaceTransforms(ChannelMask, [&OutTransforms](FInterrogationChannel) -> TArray<FIntermediate3DTransform>& { return OutTransforms; });
+	QueryLocalSpaceTransforms(Linker, ChannelMask, [&OutTransforms](FInterrogationChannel) -> TArray<FIntermediate3DTransform>& { return OutTransforms; });
 }
 
-void FSystemInterrogator::QueryWorldSpaceTransforms(TSparseArray<TArray<FTransform>>& OutTransformsByChannel)
+void FInterrogationChannels::QueryWorldSpaceTransforms(UMovieSceneEntitySystemLinker* Linker, TSparseArray<TArray<FTransform>>& OutTransformsByChannel) const
 {
 	TBitArray<> AllChannels(true, GetNumChannels());
 
@@ -429,10 +198,10 @@
 		OutTransformsByChannel[Index].Reset(Interrogations.Num());
 	}
 
-	QueryWorldSpaceTransforms(AllChannels, [&OutTransformsByChannel](FInterrogationChannel Channel) -> TArray<FTransform>& { return OutTransformsByChannel[Channel.AsIndex()]; });
+	QueryWorldSpaceTransforms(Linker, AllChannels, [&OutTransformsByChannel](FInterrogationChannel Channel) -> TArray<FTransform>& { return OutTransformsByChannel[Channel.AsIndex()]; });
 }
 
-void FSystemInterrogator::QueryLocalSpaceTransforms(TSparseArray<TArray<FIntermediate3DTransform>>& OutTransformsByChannel)
+void FInterrogationChannels::QueryLocalSpaceTransforms(UMovieSceneEntitySystemLinker* Linker, TSparseArray<TArray<FIntermediate3DTransform>>& OutTransformsByChannel) const
 {
 	TBitArray<> AllChannels(true, GetNumChannels());
 
@@ -446,10 +215,10 @@
 		OutTransformsByChannel[Index].Reset(Interrogations.Num());
 	}
 
-	QueryLocalSpaceTransforms(AllChannels, [&OutTransformsByChannel](FInterrogationChannel Channel) -> TArray<FIntermediate3DTransform>& { return OutTransformsByChannel[Channel.AsIndex()]; });
+	QueryLocalSpaceTransforms(Linker, AllChannels, [&OutTransformsByChannel](FInterrogationChannel Channel) -> TArray<FIntermediate3DTransform>& { return OutTransformsByChannel[Channel.AsIndex()]; });
 }
 
-void FSystemInterrogator::QueryWorldSpaceTransforms(const TBitArray<>& ChannelsToQuery, TSparseArray<TArray<FTransform>>& OutTransformsByChannel)
+void FInterrogationChannels::QueryWorldSpaceTransforms(UMovieSceneEntitySystemLinker* Linker, const TBitArray<>& ChannelsToQuery, TSparseArray<TArray<FTransform>>& OutTransformsByChannel) const
 {
 	for (TConstSetBitIterator<> ChannelBit(ChannelsToQuery); ChannelBit; ++ChannelBit)
 	{
@@ -460,10 +229,10 @@
 		OutTransformsByChannel[ChannelBit.GetIndex()].Reset(Interrogations.Num());
 	}
 
-	QueryWorldSpaceTransforms(ChannelsToQuery, [&OutTransformsByChannel](FInterrogationChannel Channel) -> TArray<FTransform>& { return OutTransformsByChannel[Channel.AsIndex()]; });
+	QueryWorldSpaceTransforms(Linker, ChannelsToQuery, [&OutTransformsByChannel](FInterrogationChannel Channel) -> TArray<FTransform>& { return OutTransformsByChannel[Channel.AsIndex()]; });
 }
 
-void FSystemInterrogator::QueryLocalSpaceTransforms(const TBitArray<>& ChannelsToQuery, TSparseArray<TArray<FIntermediate3DTransform>>& OutTransformsByChannel)
+void FInterrogationChannels::QueryLocalSpaceTransforms(UMovieSceneEntitySystemLinker* Linker, const TBitArray<>& ChannelsToQuery, TSparseArray<TArray<FIntermediate3DTransform>>& OutTransformsByChannel) const
 {
 	for (TConstSetBitIterator<> ChannelBit(ChannelsToQuery); ChannelBit; ++ChannelBit)
 	{
@@ -474,11 +243,11 @@
 		OutTransformsByChannel[ChannelBit.GetIndex()].Reset(Interrogations.Num());
 	}
 
-	QueryLocalSpaceTransforms(ChannelsToQuery, [&OutTransformsByChannel](FInterrogationChannel Channel) -> TArray<FIntermediate3DTransform>& { return OutTransformsByChannel[Channel.AsIndex()]; });
+	QueryLocalSpaceTransforms(Linker, ChannelsToQuery, [&OutTransformsByChannel](FInterrogationChannel Channel) -> TArray<FIntermediate3DTransform>& { return OutTransformsByChannel[Channel.AsIndex()]; });
 }
 
 template<typename GetOutputForChannelType>
-void FSystemInterrogator::QueryWorldSpaceTransforms(const TBitArray<>& ChannelsToQuery, GetOutputForChannelType&& OnGetOutputForChannel)
+void FInterrogationChannels::QueryWorldSpaceTransforms(UMovieSceneEntitySystemLinker* Linker, const TBitArray<>& ChannelsToQuery, GetOutputForChannelType&& OnGetOutputForChannel) const
 {
 	const int32 NumTransforms = Interrogations.Num();
 	const int32 NumChannels   = GetNumChannels();
@@ -498,12 +267,12 @@
 		int32 NumChildren = 0;
 
 		FInterrogationChannel Channel = FInterrogationChannel::FromIndex(ChannelBit.GetIndex());
-		if (ImportedChannelBits[Channel.AsIndex()] == false && SparseChannelInfo.FindObject(Channel) == nullptr)
+		if (ActiveChannelBits[Channel.AsIndex()] == false && SparseChannelInfo.FindObject(Channel) == nullptr)
 		{
 			continue;
 		}
 
-		OnGetOutputForChannel(Channel).SetNumUninitialized(NumTransforms);
+		OnGetOutputForChannel(Channel).SetNum(NumTransforms);
 
 		// Populate the arrays with current values for all objects in the hierarchy
 		while (Channel.IsValid())
@@ -529,7 +298,7 @@
 			}
 
 			// If this channel has variable data, we allocate space for it now
-			if (ImportedChannelBits[ThisChannelIndex] == true && !VariableTransformsByChannel.IsValidIndex(ThisChannelIndex))
+			if (ActiveChannelBits[ThisChannelIndex] == true && !VariableTransformsByChannel.IsValidIndex(ThisChannelIndex))
 			{
 				VariableTransformsByChannel.Insert(ThisChannelIndex, TArray<FTransform>());
 				VariableTransformsByChannel[ThisChannelIndex].SetNum(Interrogations.Num());
@@ -711,7 +480,7 @@
 }
 
 template<typename GetOutputForChannelType>
-void FSystemInterrogator::QueryLocalSpaceTransforms(const TBitArray<>& ChannelsToQuery, GetOutputForChannelType&& OnGetOutputForChannel)
+void FInterrogationChannels::QueryLocalSpaceTransforms(UMovieSceneEntitySystemLinker* Linker, const TBitArray<>& ChannelsToQuery, GetOutputForChannelType&& OnGetOutputForChannel) const
 {
 	const int32 NumTransforms = Interrogations.Num();
 
@@ -725,7 +494,7 @@
 	for (TConstSetBitIterator<> ChannelBit(ChannelsToQuery); ChannelBit; ++ChannelBit)
 	{
 		FInterrogationChannel Channel = FInterrogationChannel::FromIndex(ChannelBit.GetIndex());
-		if (ImportedChannelBits[Channel.AsIndex()] == false && SparseChannelInfo.FindObject(Channel) == nullptr)
+		if (ActiveChannelBits[Channel.AsIndex()] == false && SparseChannelInfo.FindObject(Channel) == nullptr)
 		{
 			continue;
 		}
@@ -742,16 +511,16 @@
 			BaseValues.Insert(ChannelBit.GetIndex(), BaseValue);
 		}
 
-		if (ImportedChannelBits[ChannelBit.GetIndex()] == true)
+		if (ActiveChannelBits[ChannelBit.GetIndex()] == true)
 		{
-			OnGetOutputForChannel(Channel).SetNumUninitialized(NumTransforms);
+			OnGetOutputForChannel(Channel).SetNum(NumTransforms);
 			PredicateBits[ChannelBit.GetIndex()] = true;
 		}
 		else
 		{
 			// If the channel doesn't have any variable data, just copy the base value over
 			TArray<FIntermediate3DTransform>& Array = OnGetOutputForChannel(Channel);
-			Array.SetNumUninitialized(NumTransforms);
+			Array.SetNum(NumTransforms);
 			for (FIntermediate3DTransform& Transform : Array)
 			{
 				Transform = BaseValue;
@@ -842,11 +611,278 @@
 	}
 }
 
-void FSystemInterrogator::FindPropertyOutputEntityIDs(const FPropertyDefinition& PropertyDefinition, FInterrogationChannel Channel, TArray<FMovieSceneEntityID>& OutEntityIDs)
+
+
+FSystemInterrogator::FSystemInterrogator()
 {
+	InitialValueCache = FInitialValueCache::GetGlobalInitialValues();
+
+	Linker = NewObject<UMovieSceneEntitySystemLinker>(GetTransientPackage());
+	Linker->SetSystemContext(EEntitySystemContext::Interrogation);
+
+	Linker->AddExtension(IInterrogationExtension::GetExtensionID(), static_cast<IInterrogationExtension*>(this));
+	Linker->AddExtension(InitialValueCache.Get());
+}
+
+FSystemInterrogator::~FSystemInterrogator()
+{}
+
+void FSystemInterrogator::AddReferencedObjects(FReferenceCollector& Collector)
+{
+	Collector.AddReferencedObject(Linker);
+}
+
+FString FSystemInterrogator::GetReferencerName() const
+{
+	return TEXT("FSystemInterrogator");
+}
+
+void FSystemInterrogator::Reset()
+{
+	if (EntityTracker)
+	{
+		EntityTracker->Reset();
+	}
+
+	EntitiesScratch.Empty();
+	EntityComponentField = FMovieSceneEntityComponentField();
+
+	Channels.Reset();
+	Linker->Reset();
+}
+
+void FSystemInterrogator::ImportTrack(UMovieSceneTrack* Track, FInterrogationChannel InChannel)
+{
+	ImportTrack(Track, Track->FindObjectBindingGuid(), InChannel);
+}
+
+void FSystemInterrogator::ImportTrack(UMovieSceneTrack* Track, const FGuid& ObjectBindingID, FInterrogationChannel InChannel)
+{
+	TGuardValue<FEntityManager*> DebugVizGuard(GEntityManagerForDebuggingVisualizers, &Linker->EntityManager);
+
+	FFrameRate TickResolution = Track->GetTypedOuter<UMovieScene>()->GetTickResolution();
+
+	const FMovieSceneTrackEvaluationField& EvaluationField = Track->GetEvaluationField();
+
+	FMovieSceneEntityComponentFieldBuilder FieldBuilder(&EntityComponentField);
+	FieldBuilder.GetSharedMetaData().ObjectBindingID = ObjectBindingID;
+
+	for (const FMovieSceneTrackEvaluationFieldEntry& Entry : EvaluationField.Entries)
+	{
+		IMovieSceneEntityProvider* EntityProvider = Cast<IMovieSceneEntityProvider>(Entry.Section);
+
+		if (!EntityProvider || Entry.Range.IsEmpty())
+		{
+			continue;
+		}
+
+		FMovieSceneEvaluationFieldEntityMetaData MetaData;
+
+		MetaData.ForcedTime = Entry.ForcedTime;
+		MetaData.Flags      = Entry.Flags;
+		MetaData.bEvaluateInSequencePreRoll  = Track->EvalOptions.bEvaluateInPreroll;
+		MetaData.bEvaluateInSequencePostRoll = Track->EvalOptions.bEvaluateInPostroll;
+
+		if (InChannel != FInterrogationChannel::Default())
+		{
+			MetaData.InterrogationChannel = InChannel;
+		}
+
+		if (!EntityProvider->PopulateEvaluationField(Entry.Range, MetaData, &FieldBuilder))
+		{
+			const int32 EntityIndex   = FieldBuilder.FindOrAddEntity(Entry.Section, 0);
+			const int32 MetaDataIndex = FieldBuilder.AddMetaData(MetaData);
+			FieldBuilder.AddPersistentEntity(Entry.Range, EntityIndex, MetaDataIndex);
+		}
+	}
+
+	Channels.ActivateChannel(InChannel);
+}
+
+int32 FSystemInterrogator::AddInterrogation(const FInterrogationParams& Params)
+{
+	const int32 NewInterrogationIndex = Channels.AddInterrogation(Params);
+	if (NewInterrogationIndex == INDEX_NONE)
+	{
+		return INDEX_NONE;
+	}
+
+	TGuardValue<FEntityManager*> DebugVizGuard(GEntityManagerForDebuggingVisualizers, &Linker->EntityManager);
+
+	TRange<FFrameNumber> UnusedEntityRange;
+
+	// Update the entities that exist at this frame
+	EntitiesScratch.Reset();
+	EntityComponentField.QueryPersistentEntities(Params.Time.FrameNumber, UnusedEntityRange, EntitiesScratch);
+
+	for (const FMovieSceneEvaluationFieldEntityQuery& Query : EntitiesScratch)
+	{
+		InterrogateEntity(NewInterrogationIndex, Query);
+	}
+
+	return NewInterrogationIndex;
+}
+
+void FSystemInterrogator::InterrogateEntity(int32 InterrogationIndex, const FMovieSceneEvaluationFieldEntityQuery& Query)
+{
+	UObject* EntityOwner = Query.Entity.Key.EntityOwner.Get();
+	IMovieSceneEntityProvider* Provider = Cast<IMovieSceneEntityProvider>(EntityOwner);
+	if (!Provider)
+	{
+		return;
+	}
+
+	FEntityImportParams Params;
+	Params.EntityID = Query.Entity.Key.EntityID;
+	Params.EntityMetaData = EntityComponentField.FindMetaData(Query);
+	Params.SharedMetaData = EntityComponentField.FindSharedMetaData(Query);
+
+	Params.InterrogationKey.Channel = FInterrogationChannel::Default();
+	Params.InterrogationKey.InterrogationIndex = InterrogationIndex;
+
+	if (Params.EntityMetaData && Params.EntityMetaData->InterrogationChannel.IsValid())
+	{
+		Params.InterrogationKey.Channel = Params.EntityMetaData->InterrogationChannel;
+	}
+
+	FImportedEntity ImportedEntity;
+	Provider->InterrogateEntity(Linker, Params, &ImportedEntity);
+
+	if (!ImportedEntity.IsEmpty())
+	{
+		if (UMovieSceneSection* Section = Cast<UMovieSceneSection>(EntityOwner))
+		{
+			Section->BuildDefaultComponents(Linker, Params, &ImportedEntity);
+		}
+
+		const FMovieSceneEntityID NewEntityID = ImportedEntity.Manufacture(Params, &Linker->EntityManager);
+		if (EntityTracker)
+		{
+			EntityTracker->TrackEntity(Params.InterrogationKey, Query.Entity.Key, NewEntityID);
+		}
+	}
+}
+
+void FSystemInterrogator::Update()
+{
+	TGuardValue<FEntityManager*> DebugVizGuard(GEntityManagerForDebuggingVisualizers, &Linker->EntityManager);
+
+	Linker->EntityManager.AddMutualComponents();
+	Linker->LinkRelevantSystems();
+
+	TArrayView<const FInterrogationParams> Interrogations = Channels.GetInterrogations();
+	FEntityTaskBuilder()
+	.Read(FBuiltInComponentTypes::Get()->Interrogation.InputKey)
+	.Write(FBuiltInComponentTypes::Get()->EvalTime)
+	.FilterNone({ FBuiltInComponentTypes::Get()->Tags.FixedTime })
+	.Iterate_PerEntity(&Linker->EntityManager, [&Interrogations](const FInterrogationKey& InterrogationKey, FFrameTime& OutEvalTime) { OutEvalTime = Interrogations[InterrogationKey.InterrogationIndex].Time; });
+
+	FMovieSceneEntitySystemRunner Runner;
+	Runner.AttachToLinker(Linker);
+	Runner.Flush();
+
+	Linker->EntityManager.IncrementSystemSerial();
+}
+
+void FSystemInterrogator::TrackImportedEntities(bool bInTrackImportedEntities)
+{
+	if (bInTrackImportedEntities && !EntityTracker)
+	{
+		EntityTracker = MakeUnique<FSystemInterrogatorEntityTracker>();
+	}
+	else if (!bInTrackImportedEntities)
+	{
+		EntityTracker.Reset();
+	}
+}
+
+FMovieSceneEntityID FSystemInterrogator::FindEntityFromOwner(FInterrogationKey InterrogationKey, UObject* Owner, uint32 EntityID) const
+{
+	if (!ensureMsgf(EntityTracker.IsValid(), TEXT("FindEntityFromOwner called on an interrogator that was not tracking entities")))
+	{
+		return FMovieSceneEntityID::Invalid();
+	}
+
+	return EntityTracker->FindTrackedEntity(InterrogationKey, FMovieSceneEvaluationFieldEntityKey { Owner, EntityID });;
+}
+
+FInterrogationChannel FSystemInterrogator::ImportLocalTransforms(USceneComponent* SceneComponent, IMovieScenePlayer* InPlayer, FMovieSceneSequenceID SequenceID)
+{
+	check(SceneComponent);
+
+	UMovieSceneSequence* Sequence = InPlayer->State.FindSequence(SequenceID);
+	if (!Sequence)
+	{
+		return FInterrogationChannel::Invalid();
+	}
+
+	FInterrogationChannel ParentChannel;
+	if (USceneComponent* AttachParent = SceneComponent->GetAttachParent())
+	{
+		ParentChannel = Channels.FindChannel(AttachParent);
+	}
+
+	FInterrogationChannel Channel = Channels.FindChannel(SceneComponent);
+	if (!Channel.IsValid())
+	{
+		Channel = Channels.AllocateChannel(SceneComponent, ParentChannel, FMovieScenePropertyBinding("Transform", TEXT("Transform")));
+	}
+
+	if (!Channel.IsValid())
+	{
+		return FInterrogationChannel::Invalid();
+	}
+
+	// Find the binding that corresponds to the component directly
+	FGuid ObjectBindingID = InPlayer->State.FindCachedObjectId(*SceneComponent, SequenceID, *InPlayer);
+	if (ObjectBindingID.IsValid())
+	{
+		ImportTransformTracks(*Sequence->GetMovieScene()->FindBinding(ObjectBindingID), Channel);
+	}
+
+	// Also blend in any transforms that exist for this scene component's actor as well (if it is the root)
+	AActor* Owner = SceneComponent->GetOwner();
+	if (SceneComponent == Owner->GetRootComponent())
+	{
+		FGuid OwnerObjectBindingID = InPlayer->State.FindCachedObjectId(*Owner, SequenceID, *InPlayer);
+		if (OwnerObjectBindingID.IsValid())
+		{
+			ImportTransformTracks(*Sequence->GetMovieScene()->FindBinding(OwnerObjectBindingID), Channel);
+		}
+	}
+
+	return Channel;
+}
+
+FInterrogationChannel FSystemInterrogator::ImportTransformHierarchy(USceneComponent* SceneComponent, IMovieScenePlayer* InPlayer, FMovieSceneSequenceID SequenceID)
+{
+	check(SceneComponent);
+
+	if (USceneComponent* AttachParent = SceneComponent->GetAttachParent())
+	{
+		ImportTransformHierarchy(AttachParent, InPlayer, SequenceID);
+	}
+
+	return ImportLocalTransforms(SceneComponent, InPlayer, SequenceID);
+}
+
+void FSystemInterrogator::ImportTransformTracks(const FMovieSceneBinding& Binding, FInterrogationChannel Channel)
+{
+	for (UMovieSceneTrack* Track : Binding.GetTracks())
+	{
+		if (Track->IsA<UMovieScene3DTransformTrack>())
+		{
+			ImportTrack(Track, Binding.GetObjectGuid(), Channel);
+		}
+	}
+}
+
+void FSystemInterrogator::FindPropertyOutputEntityIDs(const FPropertyDefinition& PropertyDefinition, FInterrogationChannel Channel, TArray<FMovieSceneEntityID>& OutEntityIDs) const
+{
 	UMovieSceneInterrogatedPropertyInstantiatorSystem* PropertyInstantiator = Linker->LinkSystem<UMovieSceneInterrogatedPropertyInstantiatorSystem>();
 	check(PropertyInstantiator);
 
+	TArrayView<const FInterrogationParams> Interrogations = Channels.GetInterrogations();
 	OutEntityIDs.SetNum(Interrogations.Num());
 
 	for (int32 Index = 0; Index < Interrogations.Num(); ++Index)
diff --git a/Engine/Source/Runtime/MovieSceneTracks/Private/Systems/MovieSceneInitialValueSystem.cpp b/Engine/Source/Runtime/MovieSceneTracks/Private/Systems/MovieSceneInitialValueSystem.cpp
--- a/Engine/Source/Runtime/MovieSceneTracks/Private/Systems/MovieSceneInitialValueSystem.cpp
+++ b/Engine/Source/Runtime/MovieSceneTracks/Private/Systems/MovieSceneInitialValueSystem.cpp
@@ -153,7 +153,7 @@ void UMovieSceneInitialValueSystem::OnRun(FSystemTaskPrerequisites& InPrerequisi
 		// When there is no caching extension, or we are interrogating we simply initialize any initial values directly without going through the cache
 		FEntityComponentFilter Filter;
 		Filter.Any(Mutation.AnyInitialValue);
-		Filter.All({ BuiltInComponents->Tags.NeedsLink });
+		Filter.All({ BuiltInComponents->BoundObject, BuiltInComponents->Tags.NeedsLink });
 
 		for (FEntityAllocationIteratorItem Item : Linker->EntityManager.Iterate(&Filter))
 		{
diff --git a/Engine/Source/Runtime/MovieSceneTracks/Private/Systems/MovieScenePredictionSystem.cpp b/Engine/Source/Runtime/MovieSceneTracks/Private/Systems/MovieScenePredictionSystem.cpp
new file mode 100644
--- /dev/null
+++ b/Engine/Source/Runtime/MovieSceneTracks/Private/Systems/MovieScenePredictionSystem.cpp
@@ -0,0 +1,550 @@
+r// Copyright Epic Games, Inc. All Rights Reserved.
+
+#include "Systems/MovieScenePredictionSystem.h"
+#include "Systems/MovieSceneEventSystems.h"
+#include "Systems/MovieScenePiecewiseFloatBlenderSystem.h"
+#include "Systems/FloatChannelEvaluatorSystem.h"
+
+#include "EntitySystem/BuiltInComponentTypes.h"
+#include "EntitySystem/MovieSceneInitialValueCache.h"
+#include "EntitySystem/Interrogation/MovieSceneInterrogatedPropertyInstantiator.h"
+
+#include "Sections/MovieScene3DTransformSection.h"
+#include "MovieSceneTracksComponentTypes.h"
+#include "MovieSceneSequence.h"
+#include "MovieSceneSequencePlayer.h"
+#include "MovieSceneSpawnableAnnotation.h"
+
+#include "Compilation/MovieSceneCompiledDataManager.h"
+
+#include "Evaluation/MovieSceneEvaluationTemplateInstance.h"
+#include "Evaluation/MovieSceneEvaluationTree.h"
+
+#include "Engine/World.h"
+#include "GameFramework/Actor.h"
+
+DECLARE_CYCLE_STAT(TEXT("Prediction Initialization"), MovieSceneEval_PredictionIntialization, STATGROUP_MovieSceneECS);
+DECLARE_CYCLE_STAT(TEXT("Prediction Finalization"),   MovieSceneEval_PredictionFinalization,  STATGROUP_MovieSceneECS);
+DECLARE_CYCLE_STAT(TEXT("Prediction Report Results"), MovieSceneEval_PredictionReportResults, STATGROUP_MovieSceneECS);
+
+UMovieSceneAsyncAction_SequencePrediction* UMovieSceneAsyncAction_SequencePrediction::PredictWorldTransformAtTime(UMovieSceneSequencePlayer* Player, USceneComponent* TargetComponent, float TimeInSeconds)
+{
+	return MakePredictionImpl(Player, TargetComponent, TimeInSeconds, true);
+}
+
+UMovieSceneAsyncAction_SequencePrediction* UMovieSceneAsyncAction_SequencePrediction::PredictLocalTransformAtTime(UMovieSceneSequencePlayer* Player, USceneComponent* TargetComponent, float TimeInSeconds)
+{
+	return MakePredictionImpl(Player, TargetComponent, TimeInSeconds, false);
+}
+
+UMovieSceneAsyncAction_SequencePrediction* UMovieSceneAsyncAction_SequencePrediction::PredictWorldTransformAtFrame(UMovieSceneSequencePlayer* Player, USceneComponent* TargetComponent, FFrameTime FrameTime)
+{
+	if (!Player || !TargetComponent)
+	{
+		return nullptr;
+	}
+
+	UMovieSceneSequence* RootSequence = Player->GetSequence();
+	if (!RootSequence || !RootSequence->GetMovieScene())
+	{
+		return nullptr;
+	}
+
+	FFrameTime TickResolutionTime = ConvertFrameTime(FrameTime, RootSequence->GetMovieScene()->GetDisplayRate(), RootSequence->GetMovieScene()->GetTickResolution());
+	return MakePredictionImpl(Player, TargetComponent, TickResolutionTime, true);
+}
+
+UMovieSceneAsyncAction_SequencePrediction* UMovieSceneAsyncAction_SequencePrediction::PredictLocalTransformAtFrame(UMovieSceneSequencePlayer* Player, USceneComponent* TargetComponent, FFrameTime FrameTime)
+{
+	if (!Player || !TargetComponent)
+	{
+		return nullptr;
+	}
+
+	UMovieSceneSequence* RootSequence = Player->GetSequence();
+	if (!RootSequence || !RootSequence->GetMovieScene())
+	{
+		return nullptr;
+	}
+
+	FFrameTime TickResolutionTime = ConvertFrameTime(FrameTime, RootSequence->GetMovieScene()->GetDisplayRate(), RootSequence->GetMovieScene()->GetTickResolution());
+	return MakePredictionImpl(Player, TargetComponent, TickResolutionTime, false);
+}
+
+UMovieSceneAsyncAction_SequencePrediction* UMovieSceneAsyncAction_SequencePrediction::MakePredictionImpl(UMovieSceneSequencePlayer* Player, USceneComponent* TargetComponent, float TimeInSeconds, bool bInWorldSpace)
+{
+	if (!Player || !TargetComponent)
+	{
+		return nullptr;
+	}
+
+	UMovieSceneSequence* RootSequence = Player->GetSequence();
+	if (!RootSequence || !RootSequence->GetMovieScene())
+	{
+		return nullptr;
+	}
+
+	FFrameTime TickResolutionTime = TimeInSeconds * RootSequence->GetMovieScene()->GetTickResolution();
+	return MakePredictionImpl(Player, TargetComponent, TickResolutionTime, bInWorldSpace);
+}
+
+UMovieSceneAsyncAction_SequencePrediction* UMovieSceneAsyncAction_SequencePrediction::MakePredictionImpl(UMovieSceneSequencePlayer* Player, USceneComponent* TargetComponent, FFrameTime TickResolutionTime, bool bInWorldSpace)
+{
+	check(Player && TargetComponent);
+
+	UMovieSceneEntitySystemLinker* Linker           = static_cast<IMovieScenePlayer*>(Player)->GetEvaluationTemplate().GetEntitySystemLinker();
+	UMovieScenePredictionSystem*   PredictionSystem = Linker->LinkSystem<UMovieScenePredictionSystem>();
+
+	UMovieSceneAsyncAction_SequencePrediction* NewPrediction = NewObject<UMovieSceneAsyncAction_SequencePrediction>(PredictionSystem);
+
+	NewPrediction->SequencePlayer      = Player;
+	NewPrediction->RootPredictedTime   = TickResolutionTime;
+	NewPrediction->InterrogationIndex  = PredictionSystem->MakeNewInterrogation(NewPrediction->RootPredictedTime);
+	NewPrediction->SceneComponent      = TargetComponent;
+	NewPrediction->bWorldSpace         = bInWorldSpace;
+
+	PredictionSystem->AddPendingPrediction(NewPrediction);
+	return NewPrediction;
+}
+
+void UMovieSceneAsyncAction_SequencePrediction::ReportResult(UE::MovieScene::FInterrogationChannels* Channels, const TSparseArray<TArray<FTransform>>& AllResults)
+{
+	using namespace UE::MovieScene;
+
+	if (InterrogationIndex != INDEX_NONE)
+	{
+		FInterrogationChannel Channel = Channels->FindChannel(SceneComponent);
+		if (Channel && AllResults.IsValidIndex(Channel.AsIndex()))
+		{
+			if (AllResults[Channel.AsIndex()].IsValidIndex(InterrogationIndex))
+			{
+				FTransform Transform = AllResults[Channel.AsIndex()][InterrogationIndex];
+				Result.Broadcast(Transform);
+				return;
+			}
+		}
+	}
+
+	Failure.Broadcast();
+}
+
+void UMovieSceneAsyncAction_SequencePrediction::ReportResult(UE::MovieScene::FInterrogationChannels* Channels, const TSparseArray<TArray<UE::MovieScene::FIntermediate3DTransform>>& AllResults)
+{
+	using namespace UE::MovieScene;
+
+	if (InterrogationIndex != INDEX_NONE)
+	{
+		FInterrogationChannel Channel = Channels->FindChannel(SceneComponent);
+		if (Channel && AllResults.IsValidIndex(Channel.AsIndex()))
+		{
+			if (AllResults[Channel.AsIndex()].IsValidIndex(InterrogationIndex))
+			{
+				FIntermediate3DTransform Transform = AllResults[Channel.AsIndex()][InterrogationIndex];
+				Result.Broadcast(FTransform(Transform.GetRotation(), Transform.GetTranslation(), Transform.GetScale()));
+				return;
+			}
+		}
+	}
+
+	Failure.Broadcast();
+}
+
+void UMovieSceneAsyncAction_SequencePrediction::Reset(UMovieSceneEntitySystemLinker* Linker)
+{
+	using namespace UE::MovieScene;
+
+	FComponentTypeID NeedsUnlink = FBuiltInComponentTypes::Get()->Tags.NeedsUnlink;
+
+	for (FMovieSceneEntityID EntityID : ImportedEntities)
+	{
+		Linker->EntityManager.AddComponent(EntityID, NeedsUnlink, EEntityRecursion::Full);
+	}
+	ImportedEntities.Empty();
+}
+
+void UMovieSceneAsyncAction_SequencePrediction::ImportEntities(UE::MovieScene::FInterrogationChannels* Channels)
+{
+	if (SceneComponent)
+	{
+		ImportTransformHierarchy(Channels, SceneComponent);
+	}
+}
+
+int32 UMovieSceneAsyncAction_SequencePrediction::ImportTransformEntities(UMovieSceneEntitySystemLinker* Linker, const UE::MovieScene::FEntityImportSequenceParams& ImportParams, const FGuid& ObjectGuid, FFrameTime PredictedTime, const FMovieSceneEntityComponentField* ComponentField, const UE::MovieScene::FInterrogationKey& InterrogationKey)
+{
+	using namespace UE::MovieScene;
+
+	int32 NumEntities = 0;
+
+	auto QueryCallback = [this, Linker, ImportParams, ObjectGuid, ComponentField, InterrogationKey, &NumEntities](const FMovieSceneEvaluationFieldEntityQuery& InQuery)
+	{
+		UMovieScene3DTransformSection* TransformSection = Cast<UMovieScene3DTransformSection>(InQuery.Entity.Key.EntityOwner.Get());
+		if (TransformSection == nullptr)
+		{
+			return true;
+		}
+
+		const FMovieSceneEvaluationFieldSharedEntityMetaData* SharedMetaData = ComponentField->FindSharedMetaData(InQuery);
+		if (!SharedMetaData || SharedMetaData->ObjectBindingID != ObjectGuid)
+		{
+			return true;
+		}
+
+		UObject* EntityOwner = InQuery.Entity.Key.EntityOwner.Get();
+		IMovieSceneEntityProvider* Provider = Cast<IMovieSceneEntityProvider>(EntityOwner);
+		if (!Provider)
+		{
+			return true;
+		}
+
+		FEntityImportParams Params;
+		Params.Sequence = ImportParams;
+		Params.EntityID = InQuery.Entity.Key.EntityID;
+		Params.EntityMetaData = ComponentField->FindMetaData(InQuery);
+		Params.SharedMetaData = SharedMetaData;
+		Params.InterrogationKey = InterrogationKey;
+
+		FImportedEntity ImportedEntity;
+		Provider->InterrogateEntity(Linker, Params, &ImportedEntity);
+
+		if (!ImportedEntity.IsEmpty())
+		{
+			TransformSection->BuildDefaultComponents(Linker, Params, &ImportedEntity);
+
+			const FMovieSceneEntityID NewEntityID = ImportedEntity.Manufacture(Params, &Linker->EntityManager);
+			this->ImportedEntities.Add(NewEntityID);
+
+			++NumEntities;
+		}
+
+		return true;
+	};
+
+	// Query any transform track entities at the specified time that relate to the object binding ID
+	TRange<FFrameNumber> Unused;
+	ComponentField->QueryPersistentEntities(PredictedTime.FrameNumber, QueryCallback, Unused);
+
+	return NumEntities;
+}
+
+
+int32 UMovieSceneAsyncAction_SequencePrediction::ImportTransformEntities(UObject* PredicateObject, UObject* ObjectContext, const UE::MovieScene::FInterrogationKey& InterrogationKey)
+{
+	using namespace UE::MovieScene;
+
+	FEntityImportSequenceParams ImportParams;
+
+	int32 TotalNumEntities = 0;
+	const FMovieSceneRootEvaluationTemplateInstance& RootTemplate = static_cast<IMovieScenePlayer*>(SequencePlayer)->GetEvaluationTemplate();
+
+	TOptional<FMovieSceneSpawnableAnnotation> SpawnableAnnotation;
+
+	if (AActor* Actor = Cast<AActor>(PredicateObject))
+	{
+		SpawnableAnnotation = FMovieSceneSpawnableAnnotation::Find(Actor);
+		if (SpawnableAnnotation)
+		{
+			FMovieSceneSpawnRegister& SpawnRegister = static_cast<IMovieScenePlayer*>(SequencePlayer)->GetSpawnRegister();
+
+			// Verify that the spawnable came from this sequence
+			if (SpawnRegister.FindSpawnedObject(SpawnableAnnotation->ObjectBindingID, SpawnableAnnotation->SequenceID).Get() != Actor)
+			{
+				SpawnableAnnotation.Reset();
+			}
+		}
+	}
+
+	UMovieSceneEntitySystemLinker*  Linker              = RootTemplate.GetEntitySystemLinker();
+	UMovieSceneCompiledDataManager* CompiledDataManager = RootTemplate.GetCompiledDataManager();
+	FMovieSceneCompiledDataID       CompiledDataID      = RootTemplate.GetCompiledDataID();
+
+	const FMovieSceneCompiledDataEntry& CompiledEntry = CompiledDataManager->GetEntryRef(CompiledDataID);
+
+	if (UMovieSceneSequence* Sequence = CompiledEntry.GetSequence())
+	{
+		const bool bIsSpawnable = SpawnableAnnotation && SpawnableAnnotation->SequenceID == MovieSceneSequenceID::Root;
+		FGuid FoundID = bIsSpawnable ? SpawnableAnnotation->ObjectBindingID : Sequence->FindBindingFromObject(PredicateObject, ObjectContext);
+
+		if (FoundID.IsValid())
+		{
+			const FMovieSceneEntityComponentField* ComponentField = CompiledDataManager->FindEntityComponentField(CompiledDataID);
+			if (ComponentField)
+			{
+				TotalNumEntities += ImportTransformEntities(Linker, ImportParams, FoundID, RootPredictedTime, ComponentField, InterrogationKey);
+			}
+		}
+	}
+
+	// Also check all the sub-sequences at the predicted time
+	const FMovieSceneSequenceHierarchy* Hierarchy = CompiledDataManager->FindHierarchy(CompiledDataID);
+	if (Hierarchy)
+	{
+		FMovieSceneEvaluationTreeRangeIterator NodeIt = Hierarchy->GetTree().IterateFromTime(RootPredictedTime.FrameNumber);
+		TMovieSceneEvaluationTreeDataIterator<FMovieSceneSubSequenceTreeEntry> SubSequenceIt = Hierarchy->GetTree().GetAllData(NodeIt.Node());
+		for ( ; SubSequenceIt; ++SubSequenceIt)
+		{
+			if (SubSequenceIt->Flags == ESectionEvaluationFlags::None)
+			{
+				const FMovieSceneSubSequenceData* SubData = Hierarchy->FindSubData(SubSequenceIt->SequenceID);
+				if (SubData)
+				{
+					if (UMovieSceneSequence* SubSequence = SubData->GetSequence())
+					{
+						FMovieSceneCompiledDataID SubDataID = CompiledDataManager->FindDataID(SubSequence);
+						if (SubDataID.IsValid())
+						{
+							const bool bIsSpawnable = SpawnableAnnotation && SpawnableAnnotation->SequenceID == SubSequenceIt->SequenceID;
+							FGuid FoundID = bIsSpawnable ? SpawnableAnnotation->ObjectBindingID : SubSequence->FindBindingFromObject(PredicateObject, ObjectContext);
+
+							if (FoundID.IsValid())
+							{
+								const FMovieSceneEntityComponentField* ComponentField = CompiledDataManager->FindEntityComponentField(SubDataID);
+								if (ComponentField)
+								{
+									ImportParams.HierarchicalBias       = SubData->HierarchicalBias;
+									ImportParams.bHasHierarchicalEasing = SubData->bHasHierarchicalEasing;
+
+									FFrameTime SubPredictedTime = RootPredictedTime * SubData->RootToSequenceTransform;
+									TotalNumEntities += ImportTransformEntities(Linker, ImportParams, FoundID, SubPredictedTime, ComponentField, InterrogationKey);
+								}
+							}
+						}
+					}
+				}
+			}
+		}
+	}
+
+	return TotalNumEntities;
+}
+
+void UMovieSceneAsyncAction_SequencePrediction::ImportLocalTransforms(UE::MovieScene::FInterrogationChannels* Channels, USceneComponent* InSceneComponent)
+{
+	using namespace UE::MovieScene;
+
+	check(InSceneComponent);
+
+	FInterrogationChannel ParentChannel;
+	if (USceneComponent* AttachParent = InSceneComponent->GetAttachParent())
+	{
+		ParentChannel = Channels->FindChannel(AttachParent);
+	}
+
+	FInterrogationChannel Channel = Channels->FindChannel(InSceneComponent);
+	if (!Channel.IsValid())
+	{
+		Channel = Channels->AllocateChannel(InSceneComponent, ParentChannel, FMovieScenePropertyBinding("Transform", TEXT("Transform")));
+	}
+
+	if (!Channel.IsValid())
+	{
+		return;
+	}
+
+	FInterrogationKey Key{ Channel, InterrogationIndex };
+
+	// Query for any transform tracks relating to the scene component at the predicted time
+	int32 NumEntities = ImportTransformEntities(InSceneComponent, InSceneComponent->GetOwner(), Key);
+
+	// Also blend in any transforms that exist for this scene component's actor as well (if it is the root)
+	AActor* Owner = InSceneComponent->GetOwner();
+	if (Owner && InSceneComponent == Owner->GetRootComponent())
+	{
+		NumEntities += ImportTransformEntities(Owner, Owner->GetWorld(), Key);
+	}
+
+	if (NumEntities > 0)
+	{
+		Channels->ActivateChannel(Channel);
+	}
+}
+
+void UMovieSceneAsyncAction_SequencePrediction::ImportTransformHierarchy(UE::MovieScene::FInterrogationChannels* Channels, USceneComponent* InSceneComponent)
+{
+	check(InSceneComponent);
+
+	// Ensure the parent is imported first
+	if (USceneComponent* AttachParent = InSceneComponent->GetAttachParent())
+	{
+		ImportTransformHierarchy(Channels, AttachParent);
+	}
+
+	ImportLocalTransforms(Channels, InSceneComponent);
+}
+
+UMovieScenePredictionSystem::UMovieScenePredictionSystem(const FObjectInitializer& ObjInit)
+	: Super(ObjInit)
+{
+	using namespace UE::MovieScene;
+
+	Phase = ESystemPhase::Spawn | ESystemPhase::Instantiation | ESystemPhase::Finalization;
+
+	if (HasAnyFlags(RF_ClassDefaultObject))
+	{
+		DefineImplicitPrerequisite(UMovieScenePiecewiseFloatBlenderSystem::StaticClass(), GetClass());
+		DefineImplicitPrerequisite(UFloatChannelEvaluatorSystem::StaticClass(), GetClass());
+		DefineImplicitPrerequisite(UMovieScenePostEvalEventSystem::StaticClass(), GetClass());
+		DefineImplicitPrerequisite(UMovieSceneInterrogatedPropertyInstantiatorSystem::StaticClass(), GetClass());
+
+		DefineComponentConsumer(GetClass(), UE::MovieScene::FMovieSceneTracksComponentTypes::Get()->ComponentTransform.PropertyTag);
+	}
+}
+
+bool UMovieScenePredictionSystem::IsRelevantImpl(UMovieSceneEntitySystemLinker* InLinker) const
+{
+	// This system is always relevant as long as it has pending or processing predictions in-flight
+	return PendingPredictions.Num() + ProcessingPredictions.Num() != 0;
+}
+
+void UMovieScenePredictionSystem::OnRun(FSystemTaskPrerequisites& InPrerequisites, FSystemSubsequentTasks& Subsequents)
+{
+	using namespace UE::MovieScene;
+
+	FMovieSceneEntitySystemRunner* Runner = Linker->GetActiveRunner();
+	if (!ensure(Runner))
+	{
+		return;
+	}
+
+	FInstanceRegistry* InstanceRegistry = Linker->GetInstanceRegistry();
+
+	ESystemPhase CurrentPhase = Runner->GetCurrentPhase();
+
+	// ------------------------------------------------------------------
+	// Instantiation Phase - import the necessary entities to interrogate
+	// pending predictions
+	if (CurrentPhase == ESystemPhase::Spawn && PendingPredictions.Num())
+	{
+		SCOPE_CYCLE_COUNTER(MovieSceneEval_PredictionIntialization)
+
+		ensure(ProcessingPredictions.Num() == 0);
+
+		Swap(ProcessingPredictions, PendingPredictions);
+		PendingPredictions.Empty();
+
+		for (int32 Index = 0; Index < ProcessingPredictions.Num(); ++Index)
+		{
+			UMovieSceneAsyncAction_SequencePrediction* Pending = ProcessingPredictions[Index];
+			Pending->ImportEntities(&InterrogationChannels);
+		}
+
+		// Add mutual components for any interrogation entities
+		Linker->EntityManager.AddMutualComponents(FEntityComponentFilter().All({ FBuiltInComponentTypes::Get()->Interrogation.InputKey }));
+	}
+	else if (CurrentPhase == ESystemPhase::Instantiation && ProcessingPredictions.Num())
+	{
+		// Initialize evaluation times
+		TArrayView<const FInterrogationParams> Interrogations = InterrogationChannels.GetInterrogations();
+		FEntityTaskBuilder()
+		.Read(FBuiltInComponentTypes::Get()->Interrogation.InputKey)
+		.Write(FBuiltInComponentTypes::Get()->EvalTime)
+		.FilterNone({ FBuiltInComponentTypes::Get()->Tags.FixedTime })
+		.Iterate_PerEntity(&Linker->EntityManager, [&Interrogations](const FInterrogationKey& InterrogationKey, FFrameTime& OutEvalTime) { OutEvalTime = Interrogations[InterrogationKey.InterrogationIndex].Time; });
+
+		// This code is a bit of a hack, but is necessary to ensure that interrogation entities have the correct initial values
+		// for any entities that need them
+		{
+			TMultiMap<FInterrogationChannel, FIntermediate3DTransform*> KeyToInitialValue;
+
+			FEntityTaskBuilder()
+			.Read(FBuiltInComponentTypes::Get()->Interrogation.OutputKey)
+			.Write(FMovieSceneTracksComponentTypes::Get()->ComponentTransform.InitialValue)
+			.Iterate_PerEntity(&Linker->EntityManager, [&KeyToInitialValue](const FInterrogationKey& InterrogationKey, FIntermediate3DTransform& OutInitialValue)
+			{
+				KeyToInitialValue.Add(InterrogationKey.Channel, &OutInitialValue);
+			});
+
+			if (KeyToInitialValue.Num() > 0)
+			{
+				FEntityTaskBuilder()
+				.Read(FBuiltInComponentTypes::Get()->BoundObject)
+				.Read(FMovieSceneTracksComponentTypes::Get()->ComponentTransform.InitialValue)
+				.FilterAll({ FBuiltInComponentTypes::Get()->BlendChannelOutput })
+				.Iterate_PerEntity(&Linker->EntityManager, [&](UObject* BoundObject, const FIntermediate3DTransform& InInitialValue)
+				{
+					FInterrogationChannel Channel = InterrogationChannels.FindChannel(BoundObject);
+					if (Channel)
+					{
+						for (auto It = KeyToInitialValue.CreateKeyIterator(Channel); It; ++It)
+						{
+							*It.Value() = InInitialValue;
+							It.RemoveCurrent();
+						}
+					}
+				});
+
+				const FSparseInterrogationChannelInfo& SparseChannelInfo = InterrogationChannels.GetSparseChannelInfo();
+				for (auto It = KeyToInitialValue.CreateIterator(); It; ++It)
+				{
+					const FInterrogationChannelInfo& ChannelInfo = SparseChannelInfo.Get(It.Key());
+					if (USceneComponent* SceneComponent = Cast<USceneComponent>(ChannelInfo.WeakObject.Get()))
+					{
+						*It.Value() = FIntermediate3DTransform(SceneComponent->GetRelativeLocation(), SceneComponent->GetRelativeRotation(), SceneComponent->GetRelativeScale3D());
+					}
+				}
+			}
+		}
+	}
+	else if (CurrentPhase == ESystemPhase::Finalization && ProcessingPredictions.Num())
+	{
+		bool bNeedsWorldSpace = false, bNeedsLocalSpace = false;
+
+		for (UMovieSceneAsyncAction_SequencePrediction* Prediction : ProcessingPredictions)
+		{
+			bNeedsWorldSpace |= Prediction->bWorldSpace;
+			bNeedsLocalSpace |= !Prediction->bWorldSpace;
+		}
+
+		TSparseArray<TArray<FTransform>> WorldTransforms;
+		TSparseArray<TArray<UE::MovieScene::FIntermediate3DTransform>> LocalTransforms;
+		{
+			SCOPE_CYCLE_COUNTER(MovieSceneEval_PredictionFinalization)
+
+			if (bNeedsLocalSpace)
+			{
+				InterrogationChannels.QueryLocalSpaceTransforms(Linker, LocalTransforms);
+			}
+			if (bNeedsWorldSpace)
+			{
+				InterrogationChannels.QueryWorldSpaceTransforms(Linker, WorldTransforms);
+			}
+		}
+
+		{
+			SCOPE_CYCLE_COUNTER(MovieSceneEval_PredictionReportResults)
+
+			for (UMovieSceneAsyncAction_SequencePrediction* Prediction : ProcessingPredictions)
+			{
+				if (Prediction->bWorldSpace)
+				{
+					Prediction->ReportResult(&InterrogationChannels, WorldTransforms);
+				}
+				else
+				{
+					Prediction->ReportResult(&InterrogationChannels, LocalTransforms);
+				}
+
+				Prediction->Reset(Linker);
+			}
+		}
+
+		Linker->EntityManager.MimicStructureChanged();
+		ProcessingPredictions.Empty();
+
+		InterrogationChannels.Reset();
+	}
+}
+
+void UMovieScenePredictionSystem::AddPendingPrediction(UMovieSceneAsyncAction_SequencePrediction* Prediction)
+{
+	PendingPredictions.Add(Prediction);
+
+	// Mimic structure changed to ensure that the instantiation phase runs
+	Linker->EntityManager.MimicStructureChanged();
+}
+
+int32 UMovieScenePredictionSystem::MakeNewInterrogation(FFrameTime InTime)
+{
+	UE::MovieScene::FInterrogationParams Params{ InTime };
+	return InterrogationChannels.AddInterrogation(Params);
+}
\ No newline at end of file
diff --git a/Engine/Source/Runtime/MovieSceneTracks/Public/EntitySystem/Interrogation/MovieSceneInterrogationLinker.h b/Engine/Source/Runtime/MovieSceneTracks/Public/EntitySystem/Interrogation/MovieSceneInterrogationLinker.h
--- a/Engine/Source/Runtime/MovieSceneTracks/Public/EntitySystem/Interrogation/MovieSceneInterrogationLinker.h
+++ b/Engine/Source/Runtime/MovieSceneTracks/Public/EntitySystem/Interrogation/MovieSceneInterrogationLinker.h
@@ -7,6 +7,7 @@
 #include "Evaluation/MovieSceneEvaluationField.h"
 #include "MovieSceneTracksComponentTypes.h"
 #include "UObject/GCObject.h"
+#include "UObject/ObjectKey.h"
 
 #include "EntitySystem/MovieSceneEntitySystemTypes.h"
 #include "EntitySystem/MovieSceneEntitySystemLinkerExtension.h"
@@ -65,11 +66,19 @@
  *        MySystem->Interrogate(FInterogationKey::Default(), OutData);
  *    }
  */
-class MOVIESCENETRACKS_API FSystemInterrogator : FGCObject, IInterrogationExtension
+class MOVIESCENETRACKS_API FInterrogationChannels
 {
 public:
-	FSystemInterrogator();
-	~FSystemInterrogator();
 
+	FInterrogationChannels();
+	~FInterrogationChannels();
+
+	void Reset();
+
+	const FSparseInterrogationChannelInfo& GetSparseChannelInfo() const
+	{
+		return SparseChannelInfo;
+	}
+
 	/**
 	 * Allocate a new interrogation channel that can be used to uniquely identify groups of tracks that animate the same property or output.
@@ -87,5 +96,179 @@
 
 
 	/**
+	 * Allocate a new channel for a set of transform tracks that isn't bound to any particular object (but can still exist within a hierarchy)
+	 *
+	 * @param ParentChannel     The channel that should be considered this channel's parent, or FInterrogationChannel::Invalid if there is none
+	 * @param CurrentValue      A value to use if this channel has now animated data after the interrogation
+	 * @return A new interrogation channel that can be passed to ImportUnboundTrack
+	 */
+	FInterrogationChannel AllocateUnboundChannel(FInterrogationChannel ParentChannel, const FTransform& CurrentValueLocalSpace);
+
+
+	/**
+	 * Import the entire transform hierarchy for the specified component, including all attached parents and tracks relating to them.
+	 *
+	 * @param SceneComponent    The scene component to import. A new channel will be allocated for this if one does not already exist.
+	 * @param InPlayer          The player interface to use for looking up object binding IDs pertaining to scene components or actors
+	 * @param SequenceID        The current sequence ID for the interrogation
+	 * @return The channel that was either pre-existing or allocated for SceneComponent
+	 */
+	FInterrogationChannel ImportTransformHierarchy(USceneComponent* SceneComponent);
+
+
+	/**
+	 * Find the existing interrogation channel for the specified object
+	 *
+	 * @param Object    The object to find a channel for
+	 * @return The (potentially invalid) interrogation channel that the specified object is on
+	 */
+	FInterrogationChannel FindChannel(UObject* Object);
+
+
+	/**
+	 * Called to activate the specified channel. Inactive channels will only ever use their object's current value.
+	 *
+	 * @param InChannel The channel to activate
+	 */
+	void ActivateChannel(FInterrogationChannel InChannel);
+
+	/**
+	 * Called to deactivate the specified channel. Inactive channels will only ever use their object's current value.
+	 *
+	 * @param InChannel The channel to deactivate
+	 */
+	void DeactivateChannel(FInterrogationChannel InChannel);
+
+
+	/**
+	 * Add a new time to interrogate this linker at, in the time-base of the imported tracks.
+	 *
+	 * @param Params    The desired time to interrogate at
+	 * @return A unique index identifier for the specified time, or INDEX_NONE if the maximum number have been reached
+	 */
+	int32 AddInterrogation(const FInterrogationParams& Params);
+
+public:
+
+
+	/**
+	 * Retrieve the number of channels allocated
+	 */
+	int32 GetNumChannels() const
+	{
+		return ActiveChannelBits.Num();
+	}
+
+
+	/**
+	 * Retrieve the current interrogations
+	 */
+	TArrayView<const FInterrogationParams> GetInterrogations() const
+	{
+		return Interrogations;
+	}
+
+public:
+
+	/**
+	 * Query local space transforms
+	 *
+	 * @param SceneComponent    The scene component to query
+	 * @param OutTransforms     Array to output transforms into, one per Interrogation
+	 */
+	void QueryLocalSpaceTransforms(UMovieSceneEntitySystemLinker* Linker, USceneComponent* SceneComponent, TArray<FIntermediate3DTransform>& OutTransforms) const;
+
+	/**
+	 * Query local space transforms
+	 *
+	 * @param InChannel         The channel to query
+	 * @param OutTransforms     Array to output transforms into, one per Interrogation
+	 */
+	void QueryLocalSpaceTransforms(UMovieSceneEntitySystemLinker* Linker, FInterrogationChannel InChannel, TArray<FIntermediate3DTransform>& OutTransforms) const;
+
+	/**
+	 * Query all local space transforms, even including channels that do not have any variable track data
+	 *
+	 * @param OutTransformsByChannel  Sparse array to receive transforms allocated by their Channel index
+	 */
+	void QueryLocalSpaceTransforms(UMovieSceneEntitySystemLinker* Linker, TSparseArray<TArray<FIntermediate3DTransform>>& OutTransformsByChannel) const;
+
+	/**
+	 * Query a specific set of channels for their local space transforms as defined by set bits within ChannelsToQuery
+	 *
+	 * @param ChannelsToQuery         Bit array containing set bits for each channel to query
+	 * @param OutTransformsByChannel  Sparse array to receive transforms allocated by their Channel index
+	 */
+	void QueryLocalSpaceTransforms(UMovieSceneEntitySystemLinker* Linker, const TBitArray<>& ChannelsToQuery, TSparseArray<TArray<FIntermediate3DTransform>>& OutTransformsByChannel) const;
+
+public:
+
+	/**
+	 * Query world space transforms for a component
+	 *
+	 * @param SceneComponent    The scene component to query
+	 * @param OutTransforms     Array to output transforms into, one per Interrogation
+	 */
+	void QueryWorldSpaceTransforms(UMovieSceneEntitySystemLinker* Linker, USceneComponent* SceneComponent, TArray<FTransform>& OutTransforms) const;
+
+	/**
+	 * Query world space transforms for a channel
+	 *
+	 * @param InChannel         The channel to query
+	 * @param OutTransforms     Array to output transforms into, one per Interrogation
+	 */
+	void QueryWorldSpaceTransforms(UMovieSceneEntitySystemLinker* Linker, FInterrogationChannel InChannel, TArray<FTransform>& OutTransforms) const;
+
+	/**
+	 * Query all world space transforms, even including channels that do not have any variable track data
+	 *
+	 * @param OutTransformsByChannel  Sparse array to receive transforms allocated by their Channel index
+	 */
+	void QueryWorldSpaceTransforms(UMovieSceneEntitySystemLinker* Linker, TSparseArray<TArray<FTransform>>& OutTransformsByChannel) const;
+
+	/**
+	 * Query a specific set of channels for their world space transforms as defined by set bits within ChannelsToQuery
+	 *
+	 * @param ChannelsToQuery         Bit array containing set bits for each channel to query
+	 * @param OutTransformsByChannel  Sparse array to receive transforms allocated by their Channel index
+	 */
+	void QueryWorldSpaceTransforms(UMovieSceneEntitySystemLinker* Linker, const TBitArray<>& ChannelsToQuery, TSparseArray<TArray<FTransform>>& OutTransformsByChannel) const;
+
+private:
+
+	/** Query implementation functions */
+	template<typename GetOutputForChannelType>
+	void QueryWorldSpaceTransforms(UMovieSceneEntitySystemLinker* Linker, const TBitArray<>& ChannelsToQuery, GetOutputForChannelType&& OnGetOutputForChannel) const;
+
+	template<typename GetOutputForChannelType>
+	void QueryLocalSpaceTransforms(UMovieSceneEntitySystemLinker* Linker, const TBitArray<>& ChannelsToQuery, GetOutputForChannelType&& OnGetOutputForChannel) const;
+
+protected:
+
+	/** Map from an object to its interrogation channel */
+	TMap<FObjectKey, FInterrogationChannel> ObjectToChannel;
+
+	/** Array of information pertaining to a given channel */
+	FSparseInterrogationChannelInfo SparseChannelInfo;
+
+	/** BitArray containing set bits for any channel that has data associated with it. The number of bits (0 or 1) in this array defines how many channels are allocated */
+	TBitArray<> ActiveChannelBits;
+
+	/** An array of interrogation times */
+	TArray<FInterrogationParams> Interrogations;
+};
+
+
+
+class MOVIESCENETRACKS_API FSystemInterrogator
+	: FGCObject
+	, IInterrogationExtension
+{
+public:
+	FSystemInterrogator();
+	~FSystemInterrogator();
+
+
+	/**
 	 * Import a track into this linker. This will add the track to the linker's evaluation field and
 	 * cause entities to be created for it at each interrogation channel (if it is relevant at such times)
@@ -110,15 +293,5 @@
 
 
 	/**
-	 * Allocate a new channel for a set of transform tracks that isn't bound to any particular object (but can still exist within a hierarchy)
-	 *
-	 * @param ParentChannel     The channel that should be considered this channel's parent, or FInterrogationChannel::Invalid if there is none
-	 * @param CurrentValue      A value to use if this channel has now animated data after the interrogation
-	 * @return A new interrogation channel that can be passed to ImportUnboundTrack
-	 */
-	FInterrogationChannel AllocateUnboundChannel(FInterrogationChannel ParentChannel, const FTransform& CurrentValueLocalSpace);
-
-
-	/**
 	 * Import the entire transform hierarchy for the specified component, including all attached parents and tracks relating to them.
 	 *
@@ -203,7 +376,7 @@
 	 */
 	int32 GetNumChannels() const
 	{
-		return ImportedChannelBits.Num();
+		return Channels.GetNumChannels();
 	}
 
 
@@ -212,55 +385,12 @@
 	 */
 	TArrayView<const FInterrogationParams> GetInterrogations() const
 	{
-		return Interrogations;
+		return Channels.GetInterrogations();
 	}
 
-public:
-
-	/**
-	 * Query the computed value of an animated property.
-	 *
-	 * See the other QueryPropertyValues method description.
-	 *
-	 * @param InPropertyComponent		The type of property being animated on the default channel.
-	 * @param OutValues					The animated values, one for each interrogation time.
-	 */
-	template<typename PropertyTraits>
-	void QueryPropertyValues(const TPropertyComponents<PropertyTraits>& InPropertyComponents, TArray<typename PropertyTraits::StorageType>& OutValues)
-	{
-		return QueryPropertyValues(InPropertyComponents, FInterrogationChannel::Default(), OutValues);
-	}
-
-	/**
-	 * Query the computed value of an animated property.
-	 *
-	 * All the tracks imported on the given channel are expected to be animating a property of the type described
-	 * by the InPropertyComponents parameter.
-	 *
-	 * @param InPropertyComponent		The type of property being animated on the given channel.
-	 * @param InChannel					The channel on which the property is being animated.
-	 * @param OutValues					The animated values, one for each interrogation time.
-	 */
-	template<typename PropertyTraits>
-	void QueryPropertyValues(const TPropertyComponents<PropertyTraits>& InPropertyComponents, FInterrogationChannel InChannel, TArray<typename PropertyTraits::StorageType>& OutValues)
-	{
-		FBuiltInComponentTypes* Components = FBuiltInComponentTypes::Get();
-		const FPropertyDefinition& PropertyDefinition = Components->PropertyRegistry.GetDefinition(InPropertyComponents.CompositeID);
-		TArrayView<const FPropertyCompositeDefinition> PropertyComposites = Components->PropertyRegistry.GetComposites(PropertyDefinition);
-
-		TArray<FMovieSceneEntityID> ValueEntityIDs;
-		FindPropertyOutputEntityIDs(PropertyDefinition, InChannel, ValueEntityIDs);
-
-		const int32 NumPropertyValues = Interrogations.Num();
-		check(ValueEntityIDs.Num() == NumPropertyValues);
-		OutValues.SetNum(NumPropertyValues);
-
-		PropertyDefinition.Handler->RebuildOperational(PropertyDefinition, PropertyComposites, ValueEntityIDs, Linker, OutValues);
-	}
-
 private:
 
-	void FindPropertyOutputEntityIDs(const FPropertyDefinition& PropertyDefinition, FInterrogationChannel Channel, TArray<FMovieSceneEntityID>& OutEntityIDs);
+	void FindPropertyOutputEntityIDs(const FPropertyDefinition& PropertyDefinition, FInterrogationChannel Channel, TArray<FMovieSceneEntityID>& OutEntityIDs) const;
 
 public:
 
@@ -270,7 +400,10 @@
 	 * @param SceneComponent    The scene component to query
 	 * @param OutTransforms     Array to output transforms into, one per Interrogation
 	 */
-	void QueryLocalSpaceTransforms(USceneComponent* SceneComponent, TArray<FIntermediate3DTransform>& OutTransforms);
+	void QueryLocalSpaceTransforms(USceneComponent* SceneComponent, TArray<FIntermediate3DTransform>& OutTransforms) const
+	{
+		Channels.QueryLocalSpaceTransforms(Linker, SceneComponent, OutTransforms);
+	}
 
 	/**
 	 * Query local space transforms
@@ -278,14 +411,20 @@
 	 * @param InChannel         The channel to query
 	 * @param OutTransforms     Array to output transforms into, one per Interrogation
 	 */
-	void QueryLocalSpaceTransforms(FInterrogationChannel InChannel, TArray<FIntermediate3DTransform>& OutTransforms);
+	void QueryLocalSpaceTransforms(FInterrogationChannel InChannel, TArray<FIntermediate3DTransform>& OutTransforms) const
+	{
+		Channels.QueryLocalSpaceTransforms(Linker, InChannel, OutTransforms);
+	}
 
 	/**
 	 * Query all local space transforms, even including channels that do not have any variable track data
 	 *
 	 * @param OutTransformsByChannel  Sparse array to receive transforms allocated by their Channel index
 	 */
-	void QueryLocalSpaceTransforms(TSparseArray<TArray<FIntermediate3DTransform>>& OutTransformsByChannel);
+	void QueryLocalSpaceTransforms(TSparseArray<TArray<FIntermediate3DTransform>>& OutTransformsByChannel) const
+	{
+		Channels.QueryLocalSpaceTransforms(Linker, OutTransformsByChannel);
+	}
 
 	/**
 	 * Query a specific set of channels for their local space transforms as defined by set bits within ChannelsToQuery
@@ -293,7 +432,10 @@
 	 * @param ChannelsToQuery         Bit array containing set bits for each channel to query
 	 * @param OutTransformsByChannel  Sparse array to receive transforms allocated by their Channel index
 	 */
-	void QueryLocalSpaceTransforms(const TBitArray<>& ChannelsToQuery, TSparseArray<TArray<FIntermediate3DTransform>>& OutTransformsByChannel);
+	void QueryLocalSpaceTransforms(const TBitArray<>& ChannelsToQuery, TSparseArray<TArray<FIntermediate3DTransform>>& OutTransformsByChannel) const
+	{
+		Channels.QueryLocalSpaceTransforms(Linker, ChannelsToQuery, OutTransformsByChannel);
+	}
 
 public:
 
@@ -303,7 +445,10 @@
 	 * @param SceneComponent    The scene component to query
 	 * @param OutTransforms     Array to output transforms into, one per Interrogation
 	 */
-	void QueryWorldSpaceTransforms(USceneComponent* SceneComponent, TArray<FTransform>& OutTransforms);
+	void QueryWorldSpaceTransforms(USceneComponent* SceneComponent, TArray<FTransform>& OutTransforms) const
+	{
+		Channels.QueryWorldSpaceTransforms(Linker, SceneComponent, OutTransforms);
+	}
 
 	/**
 	 * Query world space transforms for a channel
@@ -311,14 +456,20 @@
 	 * @param InChannel         The channel to query
 	 * @param OutTransforms     Array to output transforms into, one per Interrogation
 	 */
-	void QueryWorldSpaceTransforms(FInterrogationChannel InChannel, TArray<FTransform>& OutTransforms);
+	void QueryWorldSpaceTransforms(FInterrogationChannel InChannel, TArray<FTransform>& OutTransforms) const
+	{
+		Channels.QueryWorldSpaceTransforms(Linker, InChannel, OutTransforms);
+	}
 
 	/**
 	 * Query all world space transforms, even including channels that do not have any variable track data
 	 *
 	 * @param OutTransformsByChannel  Sparse array to receive transforms allocated by their Channel index
 	 */
-	void QueryWorldSpaceTransforms(TSparseArray<TArray<FTransform>>& OutTransformsByChannel);
+	void QueryWorldSpaceTransforms(TSparseArray<TArray<FTransform>>& OutTransformsByChannel) const
+	{
+		Channels.QueryWorldSpaceTransforms(Linker, OutTransformsByChannel);
+	}
 
 	/**
 	 * Query a specific set of channels for their world space transforms as defined by set bits within ChannelsToQuery
@@ -326,13 +477,57 @@
 	 * @param ChannelsToQuery         Bit array containing set bits for each channel to query
 	 * @param OutTransformsByChannel  Sparse array to receive transforms allocated by their Channel index
 	 */
-	void QueryWorldSpaceTransforms(const TBitArray<>& ChannelsToQuery, TSparseArray<TArray<FTransform>>& OutTransformsByChannel);
+	void QueryWorldSpaceTransforms(const TBitArray<>& ChannelsToQuery, TSparseArray<TArray<FTransform>>& OutTransformsByChannel) const
+	{
+		Channels.QueryWorldSpaceTransforms(Linker, ChannelsToQuery, OutTransformsByChannel);
+	}
 
+	/**
+	 * Query the computed value of an animated property.
+	 *
+	 * See the other QueryPropertyValues method description.
+	 *
+	 * @param InPropertyComponent		The type of property being animated on the default channel.
+	 * @param OutValues					The animated values, one for each interrogation time.
+	 */
+	template<typename PropertyTraits>
+	void QueryPropertyValues(const TPropertyComponents<PropertyTraits>& InPropertyComponents, TArray<typename PropertyTraits::StorageType>& OutValues) const
+	{
+		return QueryPropertyValues(InPropertyComponents, FInterrogationChannel::Default(), OutValues);
+	}
+
+	/**
+	 * Query the computed value of an animated property.
+	 *
+	 * All the tracks imported on the given channel are expected to be animating a property of the type described
+	 * by the InPropertyComponents parameter.
+	 *
+	 * @param InPropertyComponent		The type of property being animated on the given channel.
+	 * @param InChannel					The channel on which the property is being animated.
+	 * @param OutValues					The animated values, one for each interrogation time.
+	 */
+	template<typename PropertyTraits>
+	void QueryPropertyValues(const TPropertyComponents<PropertyTraits>& InPropertyComponents, FInterrogationChannel InChannel, TArray<typename PropertyTraits::StorageType>& OutValues) const
+	{
+		FBuiltInComponentTypes* Components = FBuiltInComponentTypes::Get();
+		const FPropertyDefinition& PropertyDefinition = Components->PropertyRegistry.GetDefinition(InPropertyComponents.CompositeID);
+		TArrayView<const FPropertyCompositeDefinition> PropertyComposites = Components->PropertyRegistry.GetComposites(PropertyDefinition);
+
+		TArray<FMovieSceneEntityID> ValueEntityIDs;
+		FindPropertyOutputEntityIDs(PropertyDefinition, InChannel, ValueEntityIDs);
+
+		const int32 NumPropertyValues = Channels.GetInterrogations().Num();
+		check(ValueEntityIDs.Num() == NumPropertyValues);
+		OutValues.SetNum(NumPropertyValues);
+
+		PropertyDefinition.Handler->RebuildOperational(PropertyDefinition, PropertyComposites, ValueEntityIDs, Linker, OutValues);
+	}
+
 public:
 
 	const FSparseInterrogationChannelInfo& GetSparseChannelInfo() const override
 	{
-		return SparseChannelInfo;
+		return Channels.GetSparseChannelInfo();
 	}
 
 private:
@@ -347,13 +542,6 @@
 	 */
 	void InterrogateEntity(int32 InterrogationIndex, const FMovieSceneEvaluationFieldEntityQuery& Query);
 
-	/** Query implementation functions */
-	template<typename GetOutputForChannelType>
-	void QueryWorldSpaceTransforms(const TBitArray<>& ChannelsToQuery, GetOutputForChannelType&& OnGetOutputForChannel);
 
-	template<typename GetOutputForChannelType>
-	void QueryLocalSpaceTransforms(const TBitArray<>& ChannelsToQuery, GetOutputForChannelType&& OnGetOutputForChannel);
-
-
 private:
 
@@ -362,24 +550,15 @@
 
 protected:
 
+	/** Object <-> Channel tracker */
+	FInterrogationChannels Channels;
+
 	/** Scratch buffer used for generating entities for interrogation times */
 	FMovieSceneEvaluationFieldEntitySet EntitiesScratch;
 
 	/** Entity component field containing all the entity owners relevant at specific times */
 	FMovieSceneEntityComponentField EntityComponentField;
 
-	/** Map from an object to its interrogation channel */
-	TMap<UObject*, FInterrogationChannel> ObjectToChannel;
-
-	/** Array of information pertaining to a given channel */
-	FSparseInterrogationChannelInfo SparseChannelInfo;
-
-	/** BitArray containing set bits for any channel that has data associated with it. The number of bits (0 or 1) in this array defines how many channels are allocated */
-	TBitArray<> ImportedChannelBits;
-
-	/** An array of interrogation times */
-	TArray<FInterrogationParams> Interrogations;
-
 	/** Tracker class that is used for keeping track of imported entities */
 	TUniquePtr<FSystemInterrogatorEntityTracker> EntityTracker;
 
diff --git a/Engine/Source/Runtime/MovieSceneTracks/Public/Systems/MovieScenePredictionSystem.h b/Engine/Source/Runtime/MovieSceneTracks/Public/Systems/MovieScenePredictionSystem.h
new file mode 100644
--- /dev/null
+++ b/Engine/Source/Runtime/MovieSceneTracks/Public/Systems/MovieScenePredictionSystem.h
@@ -0,0 +1,202 @@
+// Copyright Epic Games, Inc. All Rights Reserved.
+
+#pragma once
+
+#include "Containers/SparseArray.h"
+#include "Containers/Array.h"
+#include "Delegates/DelegateCombinations.h"
+#include "Kismet/BlueprintAsyncActionBase.h"
+#include "Misc/FrameTime.h"
+
+#include "EntitySystem/MovieSceneEntityIDs.h"
+#include "EntitySystem/MovieSceneEntitySystem.h"
+#include "EntitySystem/Interrogation/MovieSceneInterrogationLinker.h"
+
+#include "MovieScenePredictionSystem.generated.h"
+
+struct FGuid;
+struct FMovieSceneEntityComponentField;
+
+class USceneComponent;
+class UMovieSceneSequencePlayer;
+class UMovieSceneEntitySystemLinker;
+
+namespace UE
+{
+namespace MovieScene
+{
+
+struct FEntityImportSequenceParams;
+
+} // namespace MovieScene
+} // namespace UE
+
+
+DECLARE_DYNAMIC_MULTICAST_DELEGATE_OneParam(FMovieSceneActorPredictionResult, FTransform, PredictedTransform);
+DECLARE_DYNAMIC_MULTICAST_DELEGATE(FMovieSceneActorPredictionFailure);
+
+/**
+ * Async BP action that represents a pending prediction that is dispatched on a playing sequence.
+ */
+UCLASS(BlueprintType, meta=(ExposedAsyncProxy = "AsyncTask", HasDedicatedAsyncNode))
+class UMovieSceneAsyncAction_SequencePrediction : public UBlueprintAsyncActionBase
+{
+	GENERATED_BODY()
+
+public:
+
+	/**
+	 * Initiate an asynchronous prediction for the specified component's world transform at a specific time in a sequence
+	 * Changes in attachment between the sequence's current time, and the predicted time are not accounted for
+	 * Calling this function on a stopped sequence player is undefined.
+	 *
+	 * @param Player          An active, currently playing sequence player to use for predicting the transform
+	 * @param TargetComponent The component to predict a world transform for
+	 * @param TimeInSeconds   The time within the sequence to predict the transform at
+	 * @return An asynchronous prediction object that contains Result and Failure delegates
+	 */
+	UFUNCTION(BlueprintCallable, Category=Cinematics)
+	static UMovieSceneAsyncAction_SequencePrediction* PredictWorldTransformAtTime(UMovieSceneSequencePlayer* Player, USceneComponent* TargetComponent, float TimeInSeconds);
+
+
+	/**
+	 * Initiate an asynchronous prediction for the specified component's world transform at a specific time in a sequence
+	 * Changes in attachment between the sequence's current time, and the predicted time are not accounted for
+	 * Calling this function on a stopped sequence player is undefined.
+	 *
+	 * @param Player          An active, currently playing sequence player to use for predicting the transform
+	 * @param TargetComponent The component to predict a world transform for
+	 * @param FrameTime       The frame time to predict at in the sequence's display rate
+	 * @return An asynchronous prediction object that contains Result and Failure delegates
+	 */
+	UFUNCTION(BlueprintCallable, Category=Cinematics)
+	static UMovieSceneAsyncAction_SequencePrediction* PredictWorldTransformAtFrame(UMovieSceneSequencePlayer* Player, USceneComponent* TargetComponent, FFrameTime FrameTime);
+
+
+	/**
+	 * Initiate an asynchronous prediction for the specified component's local transform at a specific time in a sequence
+	 * Changes in attachment between the sequence's current time, and the predicted time are not accounted for
+	 * Calling this function on a stopped sequence player is undefined.
+	 *
+	 * @param Player          An active, currently playing sequence player to use for predicting the transform
+	 * @param TargetComponent The component to predict a world transform for
+	 * @param TimeInSeconds   The time within the sequence to predict the transform at
+	 * @return An asynchronous prediction object that contains Result and Failure delegates
+	 */
+	UFUNCTION(BlueprintCallable, Category=Cinematics)
+	static UMovieSceneAsyncAction_SequencePrediction* PredictLocalTransformAtTime(UMovieSceneSequencePlayer* Player, USceneComponent* TargetComponent, float TimeInSeconds);
+
+
+	/**
+	 * Initiate an asynchronous prediction for the specified component's local transform at a specific time in a sequence
+	 * Changes in attachment between the sequence's current time, and the predicted time are not accounted for
+	 * Calling this function on a stopped sequence player is undefined.
+	 *
+	 * @param Player          An active, currently playing sequence player to use for predicting the transform
+	 * @param TargetComponent The component to predict a world transform for
+	 * @param FrameTime       The frame time to predict at in the sequence's display rate
+	 * @return An asynchronous prediction object that contains Result and Failure delegates
+	 */
+	UFUNCTION(BlueprintCallable, Category=Cinematics)
+	static UMovieSceneAsyncAction_SequencePrediction* PredictLocalTransformAtFrame(UMovieSceneSequencePlayer* Player, USceneComponent* TargetComponent, FFrameTime FrameTime);
+
+	/**
+	 * Called during the instantiation phase to query the sequence for any transform track entities in the field
+	 * and import the necessary entities into the Entity Manager.
+	 *
+	 * @param Channels Interrogation channel map for tracking channels for objects
+	 */
+	void ImportEntities(UE::MovieScene::FInterrogationChannels* Channels);
+
+
+	/**
+	 * Reset this prediction by marking all its entities for unlink
+	 */
+	void Reset(UMovieSceneEntitySystemLinker* Linker);
+
+
+	/**
+	 * Report the result for this prediction to any observers
+	 *
+	 * @param Channels    Interrogation channel map for retrieving channels for objects
+	 * @param AllResults  The interrogation results for all predictions, in the space requested by this prediction
+	 */
+	void ReportResult(UE::MovieScene::FInterrogationChannels* Channels, const TSparseArray<TArray<FTransform>>& AllResults);
+	void ReportResult(UE::MovieScene::FInterrogationChannels* Channels, const TSparseArray<TArray<UE::MovieScene::FIntermediate3DTransform>>& AllResults);
+
+private:
+
+	static UMovieSceneAsyncAction_SequencePrediction* MakePredictionImpl(UMovieSceneSequencePlayer* Player, USceneComponent* TargetComponent, float TimeInSeconds, bool bInWorldSpace);
+	static UMovieSceneAsyncAction_SequencePrediction* MakePredictionImpl(UMovieSceneSequencePlayer* Player, USceneComponent* TargetComponent, FFrameTime TickResolutionTime, bool bInWorldSpace);
+
+	int32 ImportTransformEntities(UMovieSceneEntitySystemLinker* Linker, const UE::MovieScene::FEntityImportSequenceParams& ImportParams, const FGuid& ObjectGuid, FFrameTime PredictedTime, const FMovieSceneEntityComponentField* ComponentField, const UE::MovieScene::FInterrogationKey& InterrogationKey);
+	int32 ImportTransformEntities(UObject* PredicateObject, UObject* ObjectContext, const UE::MovieScene::FInterrogationKey& InterrogationKey);
+
+	void ImportLocalTransforms(UE::MovieScene::FInterrogationChannels* Channels, USceneComponent* InSceneComponent);
+	void ImportTransformHierarchy(UE::MovieScene::FInterrogationChannels* Channels, USceneComponent* InSceneComponent);
+
+public:
+
+	/** Called when a message is broadcast on the specified channel. Use GetPayload() to request the message payload. */
+	UPROPERTY(BlueprintAssignable)
+	FMovieSceneActorPredictionResult Result;
+
+	/** Called when a message is broadcast on the specified channel. Use GetPayload() to request the message payload. */
+	UPROPERTY(BlueprintAssignable)
+	FMovieSceneActorPredictionFailure Failure;
+
+private:
+
+	friend class UMovieScenePredictionSystem;
+
+	/** Cached array of all the entities created by this prediction */
+	TArray<UE::MovieScene::FMovieSceneEntityID> ImportedEntities;
+
+	/** The sequence player we're interrogating */
+	UPROPERTY()
+	UMovieSceneSequencePlayer* SequencePlayer;
+
+	/** The target scene component we're interrogating */
+	UPROPERTY()
+	USceneComponent* SceneComponent;
+
+	/** The time at which the interrogation should run */
+	FFrameTime RootPredictedTime;
+
+	/** The interrogation index used for any FInterrogationKeys that this prediction creates. Unique to this prediction */
+	int32 InterrogationIndex;
+
+	/** Whether we're capturing world space or local space transforms */
+	bool bWorldSpace;
+};
+
+
+/**
+ * System responsible for managing and reporting on pending UMovieSceneAsyncAction_SequencePrediction tasks
+ */
+UCLASS(MinimalAPI)
+class UMovieScenePredictionSystem : public UMovieSceneEntitySystem
+{
+public:
+
+	GENERATED_BODY()
+
+	UMovieScenePredictionSystem(const FObjectInitializer& ObjInit);
+
+	void AddPendingPrediction(UMovieSceneAsyncAction_SequencePrediction* Prediction);
+
+	int32 MakeNewInterrogation(FFrameTime InTime);
+
+private:
+
+	bool IsRelevantImpl(UMovieSceneEntitySystemLinker* InLinker) const override;
+	void OnRun(FSystemTaskPrerequisites& InPrerequisites, FSystemSubsequentTasks& Subsequents) override;
+
+	UE::MovieScene::FInterrogationChannels InterrogationChannels;
+
+	UPROPERTY()
+	TArray<UMovieSceneAsyncAction_SequencePrediction*> PendingPredictions;
+
+	UPROPERTY()
+	TArray<UMovieSceneAsyncAction_SequencePrediction*> ProcessingPredictions;
+};
