Removed errant character

diff --git a/Engine/Source/Runtime/MovieSceneTracks/Private/Systems/MovieScenePredictionSystem.cpp b/Engine/Source/Runtime/MovieSceneTracks/Private/Systems/MovieScenePredictionSystem.cpp
--- a/Engine/Source/Runtime/MovieSceneTracks/Private/Systems/MovieScenePredictionSystem.cpp
+++ b/Engine/Source/Runtime/MovieSceneTracks/Private/Systems/MovieScenePredictionSystem.cpp
@@ -1,4 +1,4 @@
-r// Copyright Epic Games, Inc. All Rights Reserved.
+// Copyright Epic Games, Inc. All Rights Reserved.
 
 #include "Systems/MovieScenePredictionSystem.h"
 #include "Systems/MovieSceneEventSystems.h"
