Sequencer: Fixed hierarchical easing provider not being able to remove easings correctly for sub sequences

diff --git a/Engine/Source/Runtime/MovieSceneTracks/Private/Systems/WeightAndEasingEvaluatorSystem.cpp b/Engine/Source/Runtime/MovieSceneTracks/Private/Systems/WeightAndEasingEvaluatorSystem.cpp
--- a/Engine/Source/Runtime/MovieSceneTracks/Private/Systems/WeightAndEasingEvaluatorSystem.cpp
+++ b/Engine/Source/Runtime/MovieSceneTracks/Private/Systems/WeightAndEasingEvaluatorSystem.cpp
@@ -111,9 +111,14 @@ void UMovieSceneHierarchicalEasingInstantiatorSystem::OnRun(FSystemTaskPrerequis
 	{
 		for (int32 Index = 0; Index < Allocation->Num(); ++Index)
 		{
-			const FSequenceInstance& Instance = InstanceRegistry->GetInstance(InstanceHandles[Index]);
+			const FSequenceInstance* RootInstance = &InstanceRegistry->GetInstance(InstanceHandles[Index]);
+			if (!RootInstance->IsRootSequence())
+			{
+				RootInstance = &InstanceRegistry->GetInstance(RootInstance->GetRootInstanceHandle());
+			}
+
 			const FMovieSceneSequenceID SubSequenceID = HierarchicalEasingProviders[Index];
-			const FInstanceHandle SubSequenceHandle = Instance.FindSubInstance(SubSequenceID);
+			const FInstanceHandle SubSequenceHandle = RootInstance->FindSubInstance(SubSequenceID);
 
 			uint16 OutEasingChannel;
 			const bool bRemoved = InstanceHandleToEasingChannel.RemoveAndCopyValue(SubSequenceHandle, OutEasingChannel);
