Added missing changes that were erroneously excluded by robomerge

diff --git a/Engine/Source/Runtime/Core/Public/UObject/ReleaseObjectVersion.h b/Engine/Source/Runtime/Core/Public/UObject/ReleaseObjectVersion.h
--- a/Engine/Source/Runtime/Core/Public/UObject/ReleaseObjectVersion.h
+++ b/Engine/Source/Runtime/Core/Public/UObject/ReleaseObjectVersion.h
@@ -127,6 +127,12 @@ struct CORE_API FReleaseObjectVersion
 		// Add structure data to Chaos Convex
 		StructureDataAddedToConvex,
 
+		// Changed axis UI for LiveLink AxisSwitch Pre Processor
+		AddedFrontRightUpAxesToLiveLinkPreProcessor,
+
+		// Some sequencer event sections that were copy-pasted left broken links to the director BP
+		FixupCopiedEventSections,
+
 		// -----<new versions can be added above this line>-------------------------------------------------
 		VersionPlusOne,
 		LatestVersion = VersionPlusOne - 1
diff --git a/Engine/Source/Runtime/MovieSceneTracks/Private/Sections/MovieSceneEventSectionBase.cpp b/Engine/Source/Runtime/MovieSceneTracks/Private/Sections/MovieSceneEventSectionBase.cpp
--- a/Engine/Source/Runtime/MovieSceneTracks/Private/Sections/MovieSceneEventSectionBase.cpp
+++ b/Engine/Source/Runtime/MovieSceneTracks/Private/Sections/MovieSceneEventSectionBase.cpp
@@ -3,6 +3,7 @@
 #include "Sections/MovieSceneEventSectionBase.h"
 #include "Modules/ModuleManager.h"
 #include "Evaluation/MovieSceneEvaluationCustomVersion.h"
+#include "UObject/ReleaseObjectVersion.h"
 
 #if WITH_EDITOR
 
@@ -63,6 +64,16 @@ void UMovieSceneEventSectionBase::PostDuplicate(bool bDuplicateForPIE)
 	PostDuplicateSectionEvent.Execute(this);
 }
 
+void UMovieSceneEventSectionBase::PostRename(UObject* OldOuter, const FName OldName)
+{
+	if (OldOuter != GetOuter())
+	{
+		Super::PostRename(OldOuter, OldName);
+
+		PostDuplicateSectionEvent.Execute(this);
+	}
+}
+
 void UMovieSceneEventSectionBase::AttemptUpgrade()
 {
 	if (!bDataUpgradeRequired)
@@ -91,6 +102,7 @@ UMovieSceneEventSectionBase::UMovieSceneEventSectionBase(const FObjectInitialize
 void UMovieSceneEventSectionBase::Serialize(FArchive& Ar)
 {
 	Ar.UsingCustomVersion(FMovieSceneEvaluationCustomVersion::GUID);
+	Ar.UsingCustomVersion(FReleaseObjectVersion::GUID);
 
 	Super::Serialize(Ar);
 
@@ -98,7 +110,8 @@ void UMovieSceneEventSectionBase::Serialize(FArchive& Ar)
 
 	if (Ar.IsLoading())
 	{
-		if (Ar.CustomVer(FMovieSceneEvaluationCustomVersion::GUID) < FMovieSceneEvaluationCustomVersion::DeprecateEventGUIDs)
+		if (Ar.CustomVer(FMovieSceneEvaluationCustomVersion::GUID) < FMovieSceneEvaluationCustomVersion::DeprecateEventGUIDs
+			|| Ar.CustomVer(FReleaseObjectVersion::GUID) < FReleaseObjectVersion::FixupCopiedEventSections)
 		{
 			AttemptUpgrade();
 		}
diff --git a/Engine/Source/Runtime/MovieSceneTracks/Private/Tracks/MovieSceneEventTrack.cpp b/Engine/Source/Runtime/MovieSceneTracks/Private/Tracks/MovieSceneEventTrack.cpp
--- a/Engine/Source/Runtime/MovieSceneTracks/Private/Tracks/MovieSceneEventTrack.cpp
+++ b/Engine/Source/Runtime/MovieSceneTracks/Private/Tracks/MovieSceneEventTrack.cpp
@@ -3,6 +3,7 @@
 #include "Tracks/MovieSceneEventTrack.h"
 #include "MovieSceneCommonHelpers.h"
 #include "Sections/MovieSceneEventSection.h"
+#include "Sections/MovieSceneEventSectionBase.h"
 #include "Sections/MovieSceneEventTriggerSection.h"
 #include "Sections/MovieSceneEventRepeaterSection.h"
 #include "Evaluation/MovieSceneEventTemplate.h"
@@ -33,6 +34,22 @@ void UMovieSceneEventTrack::Serialize(FArchive& Ar)
 
 	Super::Serialize(Ar);
 }
+
+void UMovieSceneEventTrack::PostRename(UObject* OldOuter, const FName OldName)
+{
+	if (OldOuter != GetOuter())
+	{
+		Super::PostRename(OldOuter, OldName);
+
+		for (UMovieSceneSection* Section : Sections)
+		{
+			if (UMovieSceneEventSectionBase* EventSection = Cast<UMovieSceneEventSectionBase>(Section))
+			{
+				EventSection->PostDuplicateSectionEvent.Execute(EventSection);
+			}
+		}
+	}
+}
 #endif// WITH_EDITOR
 
 void UMovieSceneEventTrack::AddSection(UMovieSceneSection& Section)
diff --git a/Engine/Source/Runtime/MovieSceneTracks/Public/Sections/MovieSceneEventSectionBase.h b/Engine/Source/Runtime/MovieSceneTracks/Public/Sections/MovieSceneEventSectionBase.h
--- a/Engine/Source/Runtime/MovieSceneTracks/Public/Sections/MovieSceneEventSectionBase.h
+++ b/Engine/Source/Runtime/MovieSceneTracks/Public/Sections/MovieSceneEventSectionBase.h
@@ -26,6 +26,7 @@ public:
 
 #if WITH_EDITOR
 
+	virtual void PostRename(UObject* OldOuter, const FName OldName) override;
 	virtual void PostDuplicate(bool bDuplicateForPIE) override;
 
 	MOVIESCENETRACKS_API void AttemptUpgrade();
diff --git a/Engine/Source/Runtime/MovieSceneTracks/Public/Tracks/MovieSceneEventTrack.h b/Engine/Source/Runtime/MovieSceneTracks/Public/Tracks/MovieSceneEventTrack.h
--- a/Engine/Source/Runtime/MovieSceneTracks/Public/Tracks/MovieSceneEventTrack.h
+++ b/Engine/Source/Runtime/MovieSceneTracks/Public/Tracks/MovieSceneEventTrack.h
@@ -67,6 +67,7 @@ public:
 
 #if WITH_EDITORONLY_DATA
 	virtual void Serialize(FArchive& Ar) override;
+	virtual void PostRename(UObject* OldOuter, const FName OldName) override;
 #endif
 
 #if WITH_EDITORONLY_DATA
