Sequencer: Add a track for controlling data layers

  - Data Layer sections can be assigned any number of data layers for World Partition worlds, and can be instructed to Load, Unload or Activate layers.
  - Pre-rolled data layer sections that are set to Activate will automatically preload data layers without making them visible

diff --git a/Engine/Source/Editor/EditorStyle/Private/StarshipStyle.cpp b/Engine/Source/Editor/EditorStyle/Private/StarshipStyle.cpp
--- a/Engine/Source/Editor/EditorStyle/Private/StarshipStyle.cpp
+++ b/Engine/Source/Editor/EditorStyle/Private/StarshipStyle.cpp
@@ -2173,6 +2173,7 @@ void FStarshipEditorStyle::FStyle::SetupSequencerStyles()
 		Set("Sequencer.Tracks.Slomo", new IMAGE_BRUSH("Sequencer/Dropdown_Icons/Icon_Play_Rate_Track_16x", Icon16x16));
 		Set("Sequencer.Tracks.Sub", new IMAGE_BRUSH("Sequencer/Dropdown_Icons/Icon_Sub_Track_16x", Icon16x16));
 		Set("Sequencer.Tracks.LevelVisibility", new IMAGE_BRUSH("Sequencer/Dropdown_Icons/Icon_Level_Visibility_Track_16x", Icon16x16));
+		Set("Sequencer.Tracks.DataLayer", new IMAGE_BRUSH_SVG("Starship/Common/DataLayers", Icon16x16));
 
 		Set("Sequencer.CursorDecorator_MarqueeAdd", new IMAGE_BRUSH("Sequencer/CursorDecorator_MarqueeAdd", Icon16x16));
 		Set("Sequencer.CursorDecorator_MarqueeSubtract", new IMAGE_BRUSH("Sequencer/CursorDecorator_MarqueeSubtract", Icon16x16));
diff --git a/Engine/Source/Editor/MovieSceneTools/MovieSceneTools.Build.cs b/Engine/Source/Editor/MovieSceneTools/MovieSceneTools.Build.cs
--- a/Engine/Source/Editor/MovieSceneTools/MovieSceneTools.Build.cs
+++ b/Engine/Source/Editor/MovieSceneTools/MovieSceneTools.Build.cs
@@ -38,6 +38,7 @@ public class MovieSceneTools : ModuleRules
                 "ActorPickerMode",
 				"AppFramework",
 				"CinematicCamera",
+				"DataLayerEditor",
                 "CurveEditor",
                 "DesktopPlatform",
                 "Json",
diff --git a/Engine/Source/Editor/MovieSceneTools/Private/MovieSceneToolsModule.cpp b/Engine/Source/Editor/MovieSceneTools/Private/MovieSceneToolsModule.cpp
--- a/Engine/Source/Editor/MovieSceneTools/Private/MovieSceneToolsModule.cpp
+++ b/Engine/Source/Editor/MovieSceneTools/Private/MovieSceneToolsModule.cpp
@@ -43,6 +43,7 @@
 #include "TrackEditors/FadeTrackEditor.h"
 #include "TrackEditors/SpawnTrackEditor.h"
 #include "TrackEditors/LevelVisibilityTrackEditor.h"
+#include "TrackEditors/DataLayerTrackEditor.h"
 #include "TrackEditors/CameraAnimTrackEditor.h"
 #include "TrackEditors/CameraShakeTrackEditor.h"
 #include "TrackEditors/MaterialParameterCollectionTrackEditor.h"
@@ -130,6 +131,7 @@ void FMovieSceneToolsModule::StartupModule()
 		FadeTrackCreateEditorHandle = SequencerModule.RegisterTrackEditor( FOnCreateTrackEditor::CreateStatic( &FFadeTrackEditor::CreateTrackEditor ) );
 		SpawnTrackCreateEditorHandle = SequencerModule.RegisterTrackEditor( FOnCreateTrackEditor::CreateStatic( &FSpawnTrackEditor::CreateTrackEditor ) );
 		LevelVisibilityTrackCreateEditorHandle = SequencerModule.RegisterTrackEditor( FOnCreateTrackEditor::CreateStatic( &FLevelVisibilityTrackEditor::CreateTrackEditor ) );
+		DataLayerTrackCreateEditorHandle = SequencerModule.RegisterTrackEditor( FOnCreateTrackEditor::CreateStatic( &FDataLayerTrackEditor::CreateTrackEditor ) );
 		CameraAnimTrackCreateEditorHandle = SequencerModule.RegisterTrackEditor(FOnCreateTrackEditor::CreateStatic(&FCameraAnimTrackEditor::CreateTrackEditor));
 		CameraShakeTrackCreateEditorHandle = SequencerModule.RegisterTrackEditor(FOnCreateTrackEditor::CreateStatic(&FCameraShakeTrackEditor::CreateTrackEditor));
 		MPCTrackCreateEditorHandle = SequencerModule.RegisterTrackEditor(FOnCreateTrackEditor::CreateStatic(&FMaterialParameterCollectionTrackEditor::CreateTrackEditor));
@@ -255,6 +257,7 @@ void FMovieSceneToolsModule::ShutdownModule()
 	SequencerModule.UnRegisterTrackEditor( FadeTrackCreateEditorHandle );
 	SequencerModule.UnRegisterTrackEditor( SpawnTrackCreateEditorHandle );
 	SequencerModule.UnRegisterTrackEditor( LevelVisibilityTrackCreateEditorHandle );
+	SequencerModule.UnRegisterTrackEditor( DataLayerTrackCreateEditorHandle );
 	SequencerModule.UnRegisterTrackEditor( CameraAnimTrackCreateEditorHandle );
 	SequencerModule.UnRegisterTrackEditor( CameraShakeTrackCreateEditorHandle );
 	SequencerModule.UnRegisterTrackEditor( MPCTrackCreateEditorHandle );
diff --git a/Engine/Source/Editor/MovieSceneTools/Private/TrackEditors/DataLayerTrackEditor.cpp b/Engine/Source/Editor/MovieSceneTools/Private/TrackEditors/DataLayerTrackEditor.cpp
new file mode 100644
--- /dev/null
+++ b/Engine/Source/Editor/MovieSceneTools/Private/TrackEditors/DataLayerTrackEditor.cpp
@@ -0,0 +1,269 @@
+// Copyright Epic Games, Inc. All Rights Reserved.
+
+#include "TrackEditors/DataLayerTrackEditor.h"
+#include "Framework/MultiBox/MultiBoxBuilder.h"
+#include "Widgets/Text/STextBlock.h"
+#include "Widgets/Layout/SBox.h"
+#include "Widgets/SBoxPanel.h"
+#include "Styling/StyleColors.h"
+#include "EditorStyleSet.h"
+#include "Sections/MovieSceneDataLayerSection.h"
+#include "Tracks/MovieSceneDataLayerTrack.h"
+#include "SequencerUtilities.h"
+#include "LevelUtils.h"
+#include "MovieSceneTimeHelpers.h"
+#include "ISequencerSection.h"
+#include "SequencerSectionPainter.h"
+#include "DataLayer/DataLayerEditorSubsystem.h"
+
+#define LOCTEXT_NAMESPACE "DataLayerTrackEditor"
+
+struct FDataLayerSection
+	: public ISequencerSection
+	, public TSharedFromThis<FDataLayerSection>
+{
+	FDataLayerSection(UMovieSceneDataLayerSection* InSection)
+		: WeakSection(InSection)
+	{}
+
+	/*~ ISequencerSection */
+	virtual UMovieSceneSection* GetSectionObject() override
+	{
+		return WeakSection.Get();
+	}
+	virtual int32 OnPaintSection(FSequencerSectionPainter& InPainter) const override
+	{
+		return InPainter.PaintSectionBackground();
+	}
+	virtual float GetSectionHeight() const override
+	{
+		return 30.f;
+	}
+	virtual TSharedRef<SWidget> GenerateSectionWidget() override
+	{
+		return SNew(SBox)
+		.Padding(FMargin(4.f))
+		[
+			SNew(SVerticalBox)
+
+			+ SVerticalBox::Slot()
+			.AutoHeight()
+			[
+				SNew(STextBlock)
+				.Text(this, &FDataLayerSection::GetVisibilityText)
+				.ColorAndOpacity(this, &FDataLayerSection::GetTextColor)
+				.TextStyle(FAppStyle::Get(), "NormalText.Important")
+			]
+
+			+ SVerticalBox::Slot()
+			[
+				SNew(STextBlock)
+				.Text(this, &FDataLayerSection::GetLayerBarText)
+				.AutoWrapText(true)
+			]
+		];
+	}
+
+	FText GetVisibilityText() const
+	{
+		UMovieSceneDataLayerSection* Section = WeakSection.Get();
+		if (Section)
+		{
+			switch (Section->GetDesiredState())
+			{
+			case EDataLayerState::Unloaded:  return LOCTEXT("VisibilityText_Unloaded", "Unload");
+			case EDataLayerState::Loaded:    return LOCTEXT("VisibilityText_Loaded", "Load");
+			default: break;
+			}
+		}
+
+		return LOCTEXT("VisibilityText_Activated", "Activate");
+	}
+
+	FText GetLayerBarText() const
+	{
+		UMovieSceneDataLayerSection* Section   = WeakSection.Get();
+		UDataLayerEditorSubsystem*   SubSystem = UDataLayerEditorSubsystem::Get();
+
+		if (SubSystem && Section)
+		{
+			FString LayerName;
+
+			const TArray<FActorDataLayer>& DataLayers = Section->GetDataLayers();
+			for (int32 Index = 0; Index < DataLayers.Num(); ++Index)
+			{
+				UDataLayer* DataLayer = SubSystem->GetDataLayer(DataLayers[Index]);
+				if (DataLayer)
+				{
+					LayerName += DataLayer->GetDataLayerLabel().ToString();
+				}
+				else
+				{
+					LayerName += FText::Format(LOCTEXT("UnknownDataLayer", "**invalid: {0}**"), FText::FromString(DataLayers[Index].Name.ToString())).ToString();
+				}
+
+				if (Index < DataLayers.Num()-1)
+				{
+					LayerName += TEXT(", ");
+				}
+			}
+
+			return FText::FromString(LayerName);
+		}
+
+		return FText();
+	}
+
+	FSlateColor GetTextColor() const
+	{
+		UMovieSceneDataLayerSection* Section = WeakSection.Get();
+		if (Section)
+		{
+			switch (Section->GetDesiredState())
+			{
+			case EDataLayerState::Unloaded:  return FStyleColors::AccentRed;
+			case EDataLayerState::Loaded:    return FStyleColors::AccentBlue;
+			case EDataLayerState::Activated: return FStyleColors::AccentGreen;
+			}
+		}
+		return FStyleColors::Foreground;
+	}
+private:
+
+	TWeakObjectPtr<UMovieSceneDataLayerSection> WeakSection;
+};
+
+FDataLayerTrackEditor::FDataLayerTrackEditor(TSharedRef<ISequencer> InSequencer)
+	: FMovieSceneTrackEditor(InSequencer)
+{}
+
+TSharedRef<ISequencerTrackEditor> FDataLayerTrackEditor::CreateTrackEditor(TSharedRef<ISequencer> InSequencer)
+{
+	return MakeShared<FDataLayerTrackEditor>(InSequencer);
+}
+
+bool FDataLayerTrackEditor::SupportsSequence(UMovieSceneSequence* InSequence) const
+{
+	ETrackSupport TrackSupported = InSequence ? InSequence->IsTrackSupported(UMovieSceneDataLayerTrack::StaticClass()) : ETrackSupport::NotSupported;
+	return TrackSupported == ETrackSupport::Supported;
+}
+
+bool FDataLayerTrackEditor::SupportsType(TSubclassOf<UMovieSceneTrack> Type) const
+{
+	return Type == UMovieSceneDataLayerTrack::StaticClass();
+}
+
+const FSlateBrush* FDataLayerTrackEditor::GetIconBrush() const
+{
+	return FAppStyle::Get().GetBrush("Sequencer.Tracks.DataLayer");
+}
+
+TSharedRef<ISequencerSection> FDataLayerTrackEditor::MakeSectionInterface(UMovieSceneSection& SectionObject, UMovieSceneTrack& Track, FGuid ObjectBinding)
+{
+	UMovieSceneDataLayerSection* DataLayerSection = Cast<UMovieSceneDataLayerSection>(&SectionObject);
+	check(SupportsType(SectionObject.GetOuter()->GetClass()) && DataLayerSection != nullptr);
+
+	return MakeShared<FDataLayerSection>(DataLayerSection);
+}
+
+void FDataLayerTrackEditor::BuildAddTrackMenu(FMenuBuilder& MenuBuilder)
+{
+	MenuBuilder.AddMenuEntry(
+		LOCTEXT("AddTrack", "Data Layer"),
+		LOCTEXT("AddTrackToolTip", "Adds a new track that can load, activate or unload Data Layers in a World Partition world."),
+		FSlateIcon(FAppStyle::GetAppStyleSetName(), "Sequencer.Tracks.DataLayer"),
+		FUIAction(FExecuteAction::CreateRaw(this, &FDataLayerTrackEditor::HandleAddTrack)));
+}
+
+TSharedPtr<SWidget> FDataLayerTrackEditor::BuildOutlinerEditWidget(const FGuid& ObjectBinding, UMovieSceneTrack* Track, const FBuildEditWidgetParams& Params)
+{
+	return FSequencerUtilities::MakeAddButton(
+		LOCTEXT("AddDataLayer_ButtonLabel", "Data Layer"),
+		FOnGetContent::CreateSP(this, &FDataLayerTrackEditor::BuildAddDataLayerMenu, Track),
+		Params.NodeIsHovered, GetSequencer());
+}
+
+UMovieSceneDataLayerSection* FDataLayerTrackEditor::AddNewSection(UMovieScene* MovieScene, UMovieSceneTrack* DataLayerTrack, EDataLayerState DesiredState)
+{
+	using namespace UE::MovieScene;
+
+	const FScopedTransaction Transaction(LOCTEXT("AddDataLayerSection_Transaction", "Add Data Layer"));
+	DataLayerTrack->Modify();
+
+	UMovieSceneDataLayerSection* DataLayerSection = CastChecked<UMovieSceneDataLayerSection>(DataLayerTrack->CreateNewSection());
+	DataLayerSection->SetDesiredState(DesiredState);
+
+	TRange<FFrameNumber> SectionRange = MovieScene->GetPlaybackRange();
+	DataLayerSection->InitialPlacement(DataLayerTrack->GetAllSections(), MovieScene->GetPlaybackRange().GetLowerBoundValue(), DiscreteSize(MovieScene->GetPlaybackRange()), true);
+	DataLayerTrack->AddSection(*DataLayerSection);
+
+	DataLayerSection->SetPreRollFrames((2.f * MovieScene->GetTickResolution()).RoundToFrame().Value);
+
+	return DataLayerSection;
+}
+
+void FDataLayerTrackEditor::HandleAddTrack()
+{
+	UMovieScene* FocusedMovieScene = GetFocusedMovieScene();
+	if (FocusedMovieScene == nullptr || FocusedMovieScene->IsReadOnly())
+	{
+		return;
+	}
+
+	const FScopedTransaction Transaction(LOCTEXT("AddDataLayerTrack_Transaction", "Add Data Layer Track"));
+	FocusedMovieScene->Modify();
+
+	UMovieSceneDataLayerTrack* NewTrack = FocusedMovieScene->AddMasterTrack<UMovieSceneDataLayerTrack>();
+	checkf(NewTrack, TEXT("Failed to create new data layer track."));
+
+	UMovieSceneDataLayerSection* NewSection = AddNewSection(FocusedMovieScene, NewTrack, EDataLayerState::Activated);
+	if (GetSequencer().IsValid())
+	{
+		GetSequencer()->OnAddTrack(NewTrack, FGuid());
+	}
+}
+
+TSharedRef<SWidget> FDataLayerTrackEditor::BuildAddDataLayerMenu(UMovieSceneTrack* DataLayerTrack)
+{
+	FMenuBuilder MenuBuilder(true, nullptr);
+
+	MenuBuilder.AddMenuEntry(
+		LOCTEXT("AddActivatedDataLayer", "Activated"),
+		LOCTEXT("AddActivatedDataLayer_Tip", "Instruct a data layer to be loaded and active."),
+		FSlateIcon(),
+		FUIAction(FExecuteAction::CreateSP(
+			this, &FDataLayerTrackEditor::HandleAddNewSection, DataLayerTrack, EDataLayerState::Activated)));
+
+	MenuBuilder.AddMenuEntry(
+		LOCTEXT("AddLoadedDataLayer", "Loaded"),
+		LOCTEXT("AddLoadedDataLayer_Tip", "Instruct a data layer to be loaded (but not active)."),
+		FSlateIcon(),
+		FUIAction(FExecuteAction::CreateSP(
+			this, &FDataLayerTrackEditor::HandleAddNewSection, DataLayerTrack, EDataLayerState::Loaded)));
+
+	MenuBuilder.AddMenuEntry(
+		LOCTEXT("AddUnloadedDataLayer", "Unloaded"),
+		LOCTEXT("AddUnloadedDataLayer_Tip", "Instruct a data layer to be unloaded for a duration."),
+		FSlateIcon(),
+		FUIAction(FExecuteAction::CreateSP(
+			this, &FDataLayerTrackEditor::HandleAddNewSection, DataLayerTrack, EDataLayerState::Loaded)));
+
+	return MenuBuilder.MakeWidget();
+}
+
+
+void FDataLayerTrackEditor::HandleAddNewSection(UMovieSceneTrack* DataLayerTrack, EDataLayerState DesiredState)
+{
+	UMovieScene* FocusedMovieScene = GetFocusedMovieScene();
+	if (FocusedMovieScene)
+	{
+		UMovieSceneDataLayerSection* NewSection = AddNewSection(FocusedMovieScene, DataLayerTrack, DesiredState);
+
+		GetSequencer()->NotifyMovieSceneDataChanged(EMovieSceneDataChangeType::MovieSceneStructureItemAdded);
+		GetSequencer()->EmptySelection();
+		GetSequencer()->SelectSection(NewSection);
+		GetSequencer()->ThrobSectionSelection();
+	}
+}
+
+#undef LOCTEXT_NAMESPACE
diff --git a/Engine/Source/Editor/MovieSceneTools/Private/TrackEditors/DataLayerTrackEditor.h b/Engine/Source/Editor/MovieSceneTools/Private/TrackEditors/DataLayerTrackEditor.h
new file mode 100644
--- /dev/null
+++ b/Engine/Source/Editor/MovieSceneTools/Private/TrackEditors/DataLayerTrackEditor.h
@@ -0,0 +1,48 @@
+// Copyright Epic Games, Inc. All Rights Reserved.
+
+#pragma once
+
+#include "CoreTypes.h"
+#include "Misc/Guid.h"
+#include "MovieSceneTrackEditor.h"
+#include "WorldPartition/DataLayer/DataLayer.h"
+
+class FMenuBuilder;
+class UMovieSceneDataLayerSection;
+
+/**
+ * A sequencer track editor for data layer tracks.
+ */
+class FDataLayerTrackEditor
+	: public FMovieSceneTrackEditor
+{
+public:
+
+	FDataLayerTrackEditor(TSharedRef<ISequencer> InSequencer);
+
+	static TSharedRef<ISequencerTrackEditor> CreateTrackEditor(TSharedRef<ISequencer> OwningSequencer);
+
+public:
+
+	// ISequencerTrackEditor interface
+	virtual TSharedRef<ISequencerSection> MakeSectionInterface(UMovieSceneSection& SectionObject, UMovieSceneTrack& Track, FGuid ObjectBinding) override;
+	virtual bool SupportsSequence(UMovieSceneSequence* InSequence) const override;
+	virtual bool SupportsType(TSubclassOf<UMovieSceneTrack> Type) const override;
+	virtual const FSlateBrush* GetIconBrush() const override;
+	virtual void BuildAddTrackMenu(FMenuBuilder& MenuBuilder) override;
+	TSharedPtr<SWidget> BuildOutlinerEditWidget(const FGuid& ObjectBinding, UMovieSceneTrack* Track, const FBuildEditWidgetParams& Params) override;
+	//virtual void BuildTrackContextMenu(FMenuBuilder& MenuBuilder, UMovieSceneTrack* Track) override;
+
+private:
+	/** Adds a new section which spans the length of the owning movie scene with the specified desired state. */
+	UMovieSceneDataLayerSection* AddNewSection(UMovieScene* MovieScene, UMovieSceneTrack* DataLayerTrack, EDataLayerState DesiredState);
+
+	/** Handles when the add track menu item is activated. */
+	void HandleAddTrack();
+
+	/** Builds the add data layer menu which is displayed on the track. */
+	TSharedRef<SWidget> BuildAddDataLayerMenu(UMovieSceneTrack* DataLayerTrack);
+
+	/** Called to add a new section with a desired state. */
+	void HandleAddNewSection(UMovieSceneTrack* DataLayerTrack, EDataLayerState DesiredState);
+};
diff --git a/Engine/Source/Editor/MovieSceneTools/Public/MovieSceneToolsModule.h b/Engine/Source/Editor/MovieSceneTools/Public/MovieSceneToolsModule.h
--- a/Engine/Source/Editor/MovieSceneTools/Public/MovieSceneToolsModule.h
+++ b/Engine/Source/Editor/MovieSceneTools/Public/MovieSceneToolsModule.h
@@ -110,6 +110,7 @@ private:
 	FDelegateHandle FadeTrackCreateEditorHandle;
 	FDelegateHandle SpawnTrackCreateEditorHandle;
 	FDelegateHandle LevelVisibilityTrackCreateEditorHandle;
+	FDelegateHandle DataLayerTrackCreateEditorHandle;
 	FDelegateHandle CameraAnimTrackCreateEditorHandle;
 	FDelegateHandle CameraShakeTrackCreateEditorHandle;
 	FDelegateHandle MPCTrackCreateEditorHandle;
diff --git a/Engine/Source/Runtime/LevelSequence/Private/LevelSequence.cpp b/Engine/Source/Runtime/LevelSequence/Private/LevelSequence.cpp
--- a/Engine/Source/Runtime/LevelSequence/Private/LevelSequence.cpp
+++ b/Engine/Source/Runtime/LevelSequence/Private/LevelSequence.cpp
@@ -21,6 +21,7 @@
 #include "Tracks/MovieSceneEventTrack.h"
 #include "Tracks/MovieSceneFadeTrack.h"
 #include "Tracks/MovieSceneLevelVisibilityTrack.h"
+#include "Tracks/MovieSceneDataLayerTrack.h"
 #include "Tracks/MovieSceneMaterialParameterCollectionTrack.h"
 #include "Tracks/MovieSceneSlomoTrack.h"
 #include "Tracks/MovieSceneSpawnTrack.h"
@@ -133,6 +134,7 @@ ETrackSupport ULevelSequence::IsTrackSupported(TSubclassOf<class UMovieSceneTrac
 		InTrackClass == UMovieSceneEventTrack::StaticClass() ||
 		InTrackClass == UMovieSceneFadeTrack::StaticClass() ||
 		InTrackClass == UMovieSceneLevelVisibilityTrack::StaticClass() ||
+		InTrackClass == UMovieSceneDataLayerTrack::StaticClass() ||
 		InTrackClass == UMovieSceneMaterialParameterCollectionTrack::StaticClass() ||
 		InTrackClass == UMovieSceneSlomoTrack::StaticClass() ||
 		InTrackClass == UMovieSceneSpawnTrack::StaticClass() ||
diff --git a/Engine/Source/Runtime/MovieSceneTracks/MovieSceneTracks.Build.cs b/Engine/Source/Runtime/MovieSceneTracks/MovieSceneTracks.Build.cs
--- a/Engine/Source/Runtime/MovieSceneTracks/MovieSceneTracks.Build.cs
+++ b/Engine/Source/Runtime/MovieSceneTracks/MovieSceneTracks.Build.cs
@@ -38,6 +38,7 @@ public class MovieSceneTracks : ModuleRules
 				});
 			PrivateDependencyModuleNames.AddRange(new string[]
 				{
+					"DataLayerEditor",
 					"EditorFramework",
 					"UnrealEd"
 				});
diff --git a/Engine/Source/Runtime/MovieSceneTracks/Private/MovieSceneTracksComponentTypes.cpp b/Engine/Source/Runtime/MovieSceneTracks/Private/MovieSceneTracksComponentTypes.cpp
--- a/Engine/Source/Runtime/MovieSceneTracks/Private/MovieSceneTracksComponentTypes.cpp
+++ b/Engine/Source/Runtime/MovieSceneTracks/Private/MovieSceneTracksComponentTypes.cpp
@@ -444,6 +444,7 @@ FMovieSceneTracksComponentTypes::FMovieSceneTracksComponentTypes()
 	ComponentRegistry->NewComponentType(&AttachParentBinding, TEXT("Attach Parent Binding"));
 
 	ComponentRegistry->NewComponentType(&LevelVisibility, TEXT("Level Visibility"));
+	ComponentRegistry->NewComponentType(&DataLayer, TEXT("Data Layer"));
 
 	FBuiltInComponentTypes* BuiltInComponents = FBuiltInComponentTypes::Get();
 
diff --git a/Engine/Source/Runtime/MovieSceneTracks/Private/Sections/MovieSceneDataLayerSection.cpp b/Engine/Source/Runtime/MovieSceneTracks/Private/Sections/MovieSceneDataLayerSection.cpp
new file mode 100644
--- /dev/null
+++ b/Engine/Source/Runtime/MovieSceneTracks/Private/Sections/MovieSceneDataLayerSection.cpp
@@ -0,0 +1,34 @@
+// Copyright Epic Games, Inc. All Rights Reserved.
+
+#include "Sections/MovieSceneDataLayerSection.h"
+#include "MovieSceneTracksComponentTypes.h"
+#include "WorldPartition/DataLayer/DataLayer.h"
+
+UMovieSceneDataLayerSection::UMovieSceneDataLayerSection(const FObjectInitializer& ObjInit)
+	: Super(ObjInit)
+{
+	DesiredState = EDataLayerState::Loaded;
+	EvalOptions.EnableAndSetCompletionMode(EMovieSceneCompletionMode::RestoreState);
+}
+
+EDataLayerState UMovieSceneDataLayerSection::GetDesiredState() const
+{
+	return DesiredState;
+}
+
+void UMovieSceneDataLayerSection::SetDesiredState(EDataLayerState InDesiredState)
+{
+	DesiredState = InDesiredState;
+}
+
+void UMovieSceneDataLayerSection::ImportEntityImpl(UMovieSceneEntitySystemLinker* EntityLinker, const FEntityImportParams& Params, FImportedEntity* OutImportedEntity)
+{
+	using namespace UE::MovieScene;
+
+	FMovieSceneDataLayerComponentData ComponentData{decltype(FMovieSceneDataLayerComponentData::Section)(this) };
+
+	OutImportedEntity->AddBuilder(
+		FEntityBuilder()
+		.Add(FMovieSceneTracksComponentTypes::Get()->DataLayer, ComponentData)
+	);
+}
diff --git a/Engine/Source/Runtime/MovieSceneTracks/Private/Systems/MovieSceneDataLayerSystem.cpp b/Engine/Source/Runtime/MovieSceneTracks/Private/Systems/MovieSceneDataLayerSystem.cpp
new file mode 100644
--- /dev/null
+++ b/Engine/Source/Runtime/MovieSceneTracks/Private/Systems/MovieSceneDataLayerSystem.cpp
@@ -0,0 +1,544 @@
+// Copyright Epic Games, Inc. All Rights Reserved.
+
+#include "Systems/MovieSceneDataLayerSystem.h"
+#include "EntitySystem/MovieSceneEntitySystemTask.h"
+#include "EntitySystem/MovieSceneSpawnablesSystem.h"
+#include "EntitySystem/MovieSceneEntitySystemLinker.h"
+#include "EntitySystem/BuiltInComponentTypes.h"
+#include "Evaluation/PreAnimatedState/MovieScenePreAnimatedEntityCaptureSource.h"
+#include "Evaluation/PreAnimatedState/MovieScenePreAnimatedStateStorage.h"
+#include "Evaluation/PreAnimatedState/MovieScenePreAnimatedStorageID.inl"
+#include "MovieSceneTracksComponentTypes.h"
+#include "Sections/MovieSceneDataLayerSection.h"
+#include "WorldPartition/DataLayer/DataLayerSubsystem.h"
+#include "Engine/World.h"
+
+#include "Kismet/KismetSystemLibrary.h"
+
+#if WITH_EDITOR
+	#include "DataLayer/DataLayerEditorSubsystem.h"
+#endif
+
+namespace UE
+{
+namespace MovieScene
+{
+
+/** Traits class governing how pre-animated state is (re)stored for data layers */
+struct FPreAnimatedDataLayerStorageTraits
+{
+	using KeyType = TObjectKey<UDataLayer>;
+	using StorageType = EDataLayerState;
+
+	/** Called when a previously animated data layer needs to be restored */
+	static void RestorePreAnimatedValue(const TObjectKey<UDataLayer>& InKey, EDataLayerState PreviousState, const FRestoreStateParams& Params);
+};
+
+/** Container class for all pre-animated data layer state */
+struct FPreAnimatedDataLayerStorage
+	: TPreAnimatedStateStorage<FPreAnimatedDataLayerStorageTraits>
+	, IPreAnimatedStateGroupManager
+{
+	static TAutoRegisterPreAnimatedStorageID<FPreAnimatedDataLayerStorage> StorageID;
+	static TAutoRegisterPreAnimatedStorageID<FPreAnimatedDataLayerStorage> GroupManagerID;
+
+	/*~ IPreAnimatedStateGroupManager */
+	void InitializeGroupManager(FPreAnimatedStateExtension* Extension) override;
+	void OnGroupDestroyed(FPreAnimatedStorageGroupHandle Group) override;
+
+	/** Make an entry for the specified data layer */
+	FPreAnimatedStateEntry MakeEntry(UDataLayer* InDataLayer);
+
+	/** Save the value of a data layer. Should only be used for runtime / PIE worlds */
+	void SavePreAnimatedState(UDataLayer* DataLayer, UDataLayerSubsystem* SubSystem);
+
+#if WITH_EDITOR
+	/** Save the value of a data layer. Should only be used for editor worlds */
+	void SavePreAnimatedStateInEditor(UDataLayer* DataLayer);
+#endif
+
+private:
+
+	FPreAnimatedStorageGroupHandle GroupHandle;
+};
+
+
+struct FDataLayerState
+{
+	void Reset();
+	bool IsEmpty() const;
+	void AddRequest(int16 InBias, EDataLayerState RequestedState);
+	TOptional<EDataLayerState> ComputeDesiredState() const;
+
+private:
+
+	int16 HierarchicalBias = 0;
+	int32 UnloadedCount    = 0;
+	int32 LoadedCount      = 0;
+	int32 ActivatedCount   = 0;
+};
+
+
+struct FDesiredLayerStates
+{
+	bool IsEmpty() const;
+	void Reset();
+	void Apply(FPreAnimatedDataLayerStorage* PreAnimatedStorage, UDataLayerSubsystem* SubSystem);
+#if WITH_EDITOR
+	void ApplyInEditor(FPreAnimatedDataLayerStorage* PreAnimatedStorage, UDataLayerEditorSubsystem* EditorSubSystem);
+#endif
+	void ApplyNewState(const FName& InDataLayerName, int16 HierarchicalBias, EDataLayerState DesiredState);
+
+	TMap<FName, FDataLayerState> StatesByLayer;
+};
+
+// ---------------------------------------------------------------------
+// FPreAnimatedDataLayerStorageTraits definitions
+void FPreAnimatedDataLayerStorageTraits::RestorePreAnimatedValue(const TObjectKey<UDataLayer>& InKey, EDataLayerState PreviousState, const FRestoreStateParams& Params)
+{
+	UDataLayer* DataLayer = InKey.ResolveObjectPtr();
+	if (!DataLayer)
+	{
+		return;
+	}
+
+	UWorld* World = DataLayer->GetWorld();
+	if (World)
+	{
+#if WITH_EDITOR
+		if (World->WorldType == EWorldType::Editor)
+		{
+			UDataLayerEditorSubsystem* SubSystem = UDataLayerEditorSubsystem::Get();
+			SubSystem->SetDataLayerVisibility(DataLayer, PreviousState == EDataLayerState::Activated);
+		}
+		else
+#endif
+		if (UDataLayerSubsystem* SubSystem = World->GetSubsystem<UDataLayerSubsystem>())
+		{
+			SubSystem->SetDataLayerState(DataLayer, PreviousState);
+		}
+	}
+}
+
+// ---------------------------------------------------------------------
+// FPreAnimatedDataLayerStorage definitions
+TAutoRegisterPreAnimatedStorageID<FPreAnimatedDataLayerStorage> FPreAnimatedDataLayerStorage::StorageID;
+TAutoRegisterPreAnimatedStorageID<FPreAnimatedDataLayerStorage> FPreAnimatedDataLayerStorage::GroupManagerID;
+
+void FPreAnimatedDataLayerStorage::InitializeGroupManager(FPreAnimatedStateExtension* Extension)
+{}
+
+void FPreAnimatedDataLayerStorage::OnGroupDestroyed(FPreAnimatedStorageGroupHandle Group)
+{
+	ensure(Group == GroupHandle);
+	GroupHandle = FPreAnimatedStorageGroupHandle();
+}
+
+FPreAnimatedStateEntry FPreAnimatedDataLayerStorage::MakeEntry(UDataLayer* InDataLayer)
+{
+	if (!GroupHandle)
+	{
+		GroupHandle = ParentExtension->AllocateGroup(SharedThis(this));
+	}
+	FPreAnimatedStorageIndex StorageIndex = GetOrCreateStorageIndex(InDataLayer);
+	return FPreAnimatedStateEntry{ GroupHandle, FPreAnimatedStateCachedValueHandle{ StorageID, StorageIndex } };
+}
+
+void FPreAnimatedDataLayerStorage::SavePreAnimatedState(UDataLayer* DataLayer, UDataLayerSubsystem* SubSystem)
+{
+	FPreAnimatedStateEntry         Entry              = MakeEntry(DataLayer);
+	EPreAnimatedStorageRequirement StorageRequirement = this->ParentExtension->GetStorageRequirement(Entry);
+
+	if (!IsStorageRequirementSatisfied(Entry.ValueHandle.StorageIndex, StorageRequirement))
+	{
+		// @todo: If a data layer is loading when Sequencer attempts to activate it,
+		// should it return to ::Loading when sequencer is done?
+		EDataLayerState ExistingState = SubSystem->GetDataLayerState(DataLayer);
+
+		AssignPreAnimatedValue(Entry.ValueHandle.StorageIndex, StorageRequirement, CopyTemp(ExistingState));
+	}
+}
+
+#if WITH_EDITOR
+void FPreAnimatedDataLayerStorage::SavePreAnimatedStateInEditor(UDataLayer* DataLayer)
+{
+	FPreAnimatedStateEntry         Entry              = MakeEntry(DataLayer);
+	EPreAnimatedStorageRequirement StorageRequirement = this->ParentExtension->GetStorageRequirement(Entry);
+
+	if (!IsStorageRequirementSatisfied(Entry.ValueHandle.StorageIndex, StorageRequirement))
+	{
+		// We never unload data-layers in editor, so feign currently unloaded layers as loaded
+		EDataLayerState ExistingState = DataLayer->IsVisible() ? EDataLayerState::Activated : EDataLayerState::Loaded;
+
+		AssignPreAnimatedValue(Entry.ValueHandle.StorageIndex, StorageRequirement, CopyTemp(ExistingState));
+	}
+}
+#endif
+
+// ---------------------------------------------------------------------
+// FDataLayerState definitions
+void FDataLayerState::Reset()
+{
+	HierarchicalBias = 0;
+	UnloadedCount    = 0;
+	LoadedCount      = 0;
+	ActivatedCount   = 0;
+}
+
+void FDataLayerState::AddRequest(int16 InBias, EDataLayerState RequestedState)
+{
+	if (InBias > HierarchicalBias)
+	{
+		Reset();
+		HierarchicalBias = InBias;
+	}
+
+	if (InBias == HierarchicalBias)
+	{
+		switch (RequestedState)
+		{
+		case EDataLayerState::Unloaded:  ++UnloadedCount;  break;
+		case EDataLayerState::Loaded:    ++LoadedCount;    break;
+		case EDataLayerState::Activated: ++ActivatedCount; break;
+		}
+	}
+}
+
+bool FDataLayerState::IsEmpty() const
+{
+	return (UnloadedCount + LoadedCount + ActivatedCount) == 0;
+}
+
+TOptional<EDataLayerState> FDataLayerState::ComputeDesiredState() const
+{
+	// If we have any requests to keep a layer loaded, always keep it loaded (even if things ask for it to be hidden)
+	EDataLayerState FallbackState = LoadedCount != 0 ? EDataLayerState::Loaded : EDataLayerState::Unloaded;
+
+	if (ActivatedCount == UnloadedCount)
+	{
+		// Equal number of requests for active and unloaded - just leave the data layer alone
+		if (LoadedCount != 0)
+		{
+			return EDataLayerState::Loaded;
+		}
+		return TOptional<EDataLayerState>();
+	}
+	
+	if (ActivatedCount > UnloadedCount)
+	{
+		return EDataLayerState::Activated;
+	}
+
+	return FallbackState;
+}
+
+// ---------------------------------------------------------------------
+// FDataLayerStates definitions
+bool FDesiredLayerStates::IsEmpty() const
+{
+	return StatesByLayer.Num() == 0;
+}
+
+void FDesiredLayerStates::Reset()
+{
+	for (TTuple<FName, FDataLayerState>& Pair : StatesByLayer)
+	{
+		Pair.Value.Reset();
+	}
+}
+
+void FDesiredLayerStates::Apply(FPreAnimatedDataLayerStorage* PreAnimatedStorage, UDataLayerSubsystem* SubSystem)
+{
+	for (auto It = StatesByLayer.CreateIterator(); It; ++It)
+	{
+		const FDataLayerState& StateValue = It.Value();
+		if (StateValue.IsEmpty())
+		{
+			It.RemoveCurrent();
+			continue;
+		}
+
+		if (TOptional<EDataLayerState> DesiredState = StateValue.ComputeDesiredState())
+		{
+			UDataLayer* DataLayer = SubSystem->GetDataLayerFromName(It.Key());
+			if (DataLayer)
+			{
+				if (PreAnimatedStorage)
+				{
+					PreAnimatedStorage->SavePreAnimatedState(DataLayer, SubSystem);
+				}
+
+				SubSystem->SetDataLayerState(DataLayer, DesiredState.GetValue());
+			}
+			else
+			{
+				UE_LOG(LogMovieScene, Warning, TEXT("Unable to find data layer with name '%s'"), *It.Key().ToString());
+			}
+		}
+	}
+}
+
+#if WITH_EDITOR
+void FDesiredLayerStates::ApplyInEditor(FPreAnimatedDataLayerStorage* PreAnimatedStorage, UDataLayerEditorSubsystem* SubSystem)
+{
+	TArray<UDataLayer*> DatalayersNeedingLoad;
+	TArray<UDataLayer*> DatalayersNeedingShow;
+	TArray<UDataLayer*> DatalayersNeedingHide;
+
+	for (auto It = StatesByLayer.CreateIterator(); It; ++It)
+	{
+		const FDataLayerState& StateValue = It.Value();
+		if (StateValue.IsEmpty())
+		{
+			It.RemoveCurrent();
+			continue;
+		}
+
+		if (TOptional<EDataLayerState> DesiredState = StateValue.ComputeDesiredState())
+		{
+			UDataLayer* DataLayer = SubSystem->GetDataLayerFromName(It.Key());
+			if (DataLayer)
+			{
+				if (PreAnimatedStorage)
+				{
+					PreAnimatedStorage->SavePreAnimatedStateInEditor(DataLayer);
+				}
+
+				// In-editor we only ever hide data layers, we never unload them
+				switch (DesiredState.GetValue())
+				{
+				case EDataLayerState::Unloaded:
+					DatalayersNeedingHide.Add(DataLayer);
+					break;
+				case EDataLayerState::Loaded:
+					DatalayersNeedingLoad.Add(DataLayer);
+					DatalayersNeedingHide.Add(DataLayer);
+					break;
+				case EDataLayerState::Activated:
+					DatalayersNeedingLoad.Add(DataLayer);
+					DatalayersNeedingShow.Add(DataLayer);
+					break;
+				default:
+					break;
+				}
+			}
+			else
+			{
+				UE_LOG(LogMovieScene, Warning, TEXT("Unable to find data layer with name '%s'"), *It.Key().ToString());
+			}
+		}
+	}
+
+	if (DatalayersNeedingLoad.Num() > 0)
+	{
+		// This blocks while we load data layers
+		SubSystem->SetDataLayersIsDynamicallyLoadedInEditor(DatalayersNeedingLoad, true);
+	}
+	if (DatalayersNeedingShow.Num() > 0)
+	{
+		SubSystem->SetDataLayersVisibility(DatalayersNeedingShow, true);
+	}
+	if (DatalayersNeedingHide.Num() > 0)
+	{
+		SubSystem->SetDataLayersVisibility(DatalayersNeedingHide, false);
+	}
+}
+#endif
+
+void FDesiredLayerStates::ApplyNewState(const FName& InDataLayerName, int16 HierarchicalBias, EDataLayerState DesiredState)
+{
+	using namespace UE::MovieScene;
+
+	FDataLayerState* LayerState = StatesByLayer.Find(InDataLayerName);
+	if (!LayerState)
+	{
+		LayerState = &StatesByLayer.Add(InDataLayerName, FDataLayerState());
+	}
+
+	LayerState->AddRequest(HierarchicalBias, DesiredState);
+}
+
+} // namespace MovieScene
+} // namespace UE
+
+UMovieSceneDataLayerSystem::UMovieSceneDataLayerSystem(const FObjectInitializer& ObjInit)
+	: Super(ObjInit)
+{
+	using namespace UE::MovieScene;
+
+	FBuiltInComponentTypes* BuiltInComponents = FBuiltInComponentTypes::Get();
+	FMovieSceneTracksComponentTypes* TracksComponents = FMovieSceneTracksComponentTypes::Get();
+
+	Phase = UE::MovieScene::ESystemPhase::Spawn;
+	RelevantComponent = TracksComponents->DataLayer;
+
+	if (HasAnyFlags(RF_ClassDefaultObject))
+	{
+		DefineImplicitPrerequisite(GetClass(), UMovieSceneSpawnablesSystem::StaticClass());
+	}
+	else
+	{
+		DesiredLayerStates = MakeShared<FDesiredLayerStates>();
+
+		// We only need to run if there are data layer components that need (un)linking
+		ApplicableFilter.Filter.All({ TracksComponents->DataLayer });
+		ApplicableFilter.Filter.Any({ BuiltInComponents->Tags.NeedsLink, BuiltInComponents->Tags.NeedsUnlink });
+	}
+}
+
+void UMovieSceneDataLayerSystem::OnLink()
+{
+}
+
+bool UMovieSceneDataLayerSystem::IsRelevantImpl(UMovieSceneEntitySystemLinker* InLinker) const
+{
+	return InLinker->EntityManager.ContainsComponent(RelevantComponent) || (DesiredLayerStates && !DesiredLayerStates->IsEmpty());
+}
+
+void UMovieSceneDataLayerSystem::OnRun(FSystemTaskPrerequisites& InPrerequisites, FSystemSubsequentTasks& Subsequents)
+{
+	// Only run if we must
+	UWorld* World = GetWorld();
+
+	if (!World || !ApplicableFilter.Matches(Linker->EntityManager))
+	{
+		return;
+	}
+
+	// Update the desired states of all data layers from the entity manager
+	UpdateDesiredStates();
+
+	// In-editor we apply desired states through the editor sub-system
+#if WITH_EDITOR
+	if (World->WorldType == EWorldType::Editor)
+	{
+		UDataLayerEditorSubsystem* DataLayerEditorSubsystem = UDataLayerEditorSubsystem::Get();
+		if (ensureMsgf(DataLayerEditorSubsystem, TEXT("Unable to retrieve data layer editor subsystem - data layer tracks will not function correctly")))
+		{
+			DesiredLayerStates->ApplyInEditor(WeakPreAnimatedStorage.Pin().Get(), DataLayerEditorSubsystem);
+		}
+	}
+	else
+#endif
+
+	// Outside of editor, or in PIE, we use the runtime data layer sub-system
+	{
+		UDataLayerSubsystem* DataLayerSubsystem = World->GetSubsystem<UDataLayerSubsystem>();
+		if (ensureMsgf(DataLayerSubsystem, TEXT("Unable to retrieve data layer subsystem - data layer tracks will not function correctly")))
+		{
+			DesiredLayerStates->Apply(WeakPreAnimatedStorage.Pin().Get(), DataLayerSubsystem);
+		}
+	}
+}
+
+void UMovieSceneDataLayerSystem::UpdateDesiredStates()
+{
+	using namespace UE::MovieScene;
+
+	FBuiltInComponentTypes*          BuiltInComponents = FBuiltInComponentTypes::Get();
+	FMovieSceneTracksComponentTypes* TracksComponents  = FMovieSceneTracksComponentTypes::Get();
+
+	BeginTrackingEntities();
+
+	// Reset the container and re-harvest all active states from the entity manager
+	DesiredLayerStates->Reset();
+
+	auto GatherDataLayers = [this, BuiltInComponents](FEntityAllocationIteratorItem Item, const FMovieSceneDataLayerComponentData* ComponentData, const int16* OptHBiases)
+	{
+		const bool bPreroll = Item.GetAllocationType().Contains(BuiltInComponents->Tags.PreRoll);
+		for (int32 Index = 0; Index < Item.GetAllocation()->Num(); ++Index)
+		{
+			const UMovieSceneDataLayerSection* Section = ComponentData[Index].Section.Get();
+			if (!ensure(Section))
+			{
+				continue;
+			}
+
+			EDataLayerState DesiredState = Section->GetDesiredState();
+			if (bPreroll)
+			{
+				// Ignore unloaded data layers in preroll
+				if (DesiredState == EDataLayerState::Unloaded)
+				{
+					continue;
+				}
+
+				// Preroll always means keep the data layer in memory, but do not activate until the section is active
+				DesiredState = EDataLayerState::Loaded;
+			}
+
+			for (const FActorDataLayer& ActorDataLayer : Section->GetDataLayers())
+			{
+				this->DesiredLayerStates->ApplyNewState(ActorDataLayer.Name, OptHBiases ? OptHBiases[Index] : 0, DesiredState);
+			}
+		}
+	};
+
+	FEntityTaskBuilder()
+	.Read(TracksComponents->DataLayer)
+	.ReadOptional(BuiltInComponents->HierarchicalBias)
+	.FilterNone({ BuiltInComponents->Tags.NeedsUnlink })  // Do not iterate things that are being destroyed
+	.Iterate_PerAllocation(&Linker->EntityManager, GatherDataLayers);
+}
+
+void UMovieSceneDataLayerSystem::BeginTrackingEntities()
+{
+	using namespace UE::MovieScene;
+
+	UWorld*                     World              = GetWorld();
+	UDataLayerSubsystem*        DataLayerSubsystem = World ? World->GetSubsystem<UDataLayerSubsystem>() : nullptr;
+	FPreAnimatedStateExtension* PreAnimatedState   = Linker->FindExtension<FPreAnimatedStateExtension>();
+
+	if (!DataLayerSubsystem || !PreAnimatedState)
+	{
+		return;
+	}
+
+	FBuiltInComponentTypes*          BuiltInComponents = FBuiltInComponentTypes::Get();
+	FMovieSceneTracksComponentTypes* TracksComponents  = FMovieSceneTracksComponentTypes::Get();
+	FPreAnimatedEntityCaptureSource* EntityMetaData    = PreAnimatedState->GetOrCreateEntityMetaData();
+
+	// Cache the preanimated storage
+	TSharedPtr<FPreAnimatedDataLayerStorage> PreAnimatedStorage = PreAnimatedState->GetOrCreateStorage<FPreAnimatedDataLayerStorage>();
+	WeakPreAnimatedStorage = PreAnimatedStorage;
+
+	// ---------------------------------------------------------------------------------
+	// Only gather entity meta-data during SavePreAnimatedState - the actual values will be cached
+	// inside FDataLayerState::Apply
+	auto GatherDataLayers = [this, EntityMetaData, DataLayerSubsystem, PreAnimatedStorage](
+		FEntityAllocationIteratorItem Item,
+		TRead<FMovieSceneEntityID> EntityIDs,
+		TRead<FInstanceHandle> RootInstanceHandles,
+		TRead<FMovieSceneDataLayerComponentData> ComponentData)
+	{
+		const bool bRestoreState = Item.GetAllocationType().Contains(FBuiltInComponentTypes::Get()->Tags.RestoreState);
+		for (int32 Index = 0; Index < Item.GetAllocation()->Num(); ++Index)
+		{
+			const UMovieSceneDataLayerSection* Section = ComponentData[Index].Section.Get();
+			if (Section)
+			{
+				FMovieSceneEntityID EntityID     = EntityIDs[Index];
+				FInstanceHandle     RootInstance = RootInstanceHandles[Index];
+
+				for (const FActorDataLayer& ActorDataLayer : Section->GetDataLayers())
+				{
+					UDataLayer* DataLayer = DataLayerSubsystem->GetDataLayer(ActorDataLayer);
+					if (DataLayer)
+					{
+						FPreAnimatedStateEntry Entry = PreAnimatedStorage->MakeEntry(DataLayer);
+						EntityMetaData->BeginTrackingEntity(Entry, EntityID, RootInstance, bRestoreState);
+					}
+				}
+			}
+		}
+	};
+
+	// Iterate any data layer components that need link
+	FEntityTaskBuilder()
+	.ReadEntityIDs()
+	.Read(BuiltInComponents->RootInstanceHandle)
+	.Read(TracksComponents->DataLayer)
+	.FilterAll({ BuiltInComponents->Tags.NeedsLink })
+	.Iterate_PerAllocation(&Linker->EntityManager, GatherDataLayers);
+}
+
diff --git a/Engine/Source/Runtime/MovieSceneTracks/Private/Systems/MovieSceneDataLayerSystem.h b/Engine/Source/Runtime/MovieSceneTracks/Private/Systems/MovieSceneDataLayerSystem.h
new file mode 100644
--- /dev/null
+++ b/Engine/Source/Runtime/MovieSceneTracks/Private/Systems/MovieSceneDataLayerSystem.h
@@ -0,0 +1,57 @@
+// Copyright Epic Games, Inc. All Rights Reserved.
+
+#pragma once
+
+#include "Templates/SharedPointer.h"
+#include "UObject/ObjectKey.h"
+#include "Containers/Map.h"
+#include "WorldPartition/DataLayer/DataLayer.h"
+#include "EntitySystem/MovieSceneEntitySystem.h"
+#include "EntitySystem/MovieSceneCachedEntityFilterResult.h"
+#include "MovieSceneDataLayerSystem.generated.h"
+
+
+
+namespace UE
+{
+namespace MovieScene
+{
+
+struct FDesiredLayerStates;
+struct FPreAnimatedDataLayerStorage;
+
+} // namespace MovieScene
+} // namespace UE
+
+
+/**
+ * System that applies all data layer changes to the world
+ */
+UCLASS(MinimalAPI)
+class UMovieSceneDataLayerSystem
+	: public UMovieSceneEntitySystem
+{
+public:
+
+	GENERATED_BODY()
+
+	UMovieSceneDataLayerSystem(const FObjectInitializer& ObjInit);
+
+private:
+
+	virtual bool IsRelevantImpl(UMovieSceneEntitySystemLinker* InLinker) const override;
+	virtual void OnLink() override;
+	virtual void OnRun(FSystemTaskPrerequisites& InPrerequisites, FSystemSubsequentTasks& Subsequents) override final;
+
+	void UpdateDesiredStates();
+	void BeginTrackingEntities();
+
+private:
+
+	/** Cached filter that tells us whether we need to run this frame */
+	UE::MovieScene::FCachedEntityFilterResult_Match ApplicableFilter;
+	/** Impl class that stores the desired layer states for the current (and last) frame */
+	TSharedPtr<UE::MovieScene::FDesiredLayerStates> DesiredLayerStates;
+	/** Weak ptr to the pre-animated storage if we need to cache previous states for layers */
+	TWeakPtr<UE::MovieScene::FPreAnimatedDataLayerStorage> WeakPreAnimatedStorage;
+};
diff --git a/Engine/Source/Runtime/MovieSceneTracks/Private/Tracks/MovieSceneDataLayerTrack.cpp b/Engine/Source/Runtime/MovieSceneTracks/Private/Tracks/MovieSceneDataLayerTrack.cpp
new file mode 100644
--- /dev/null
+++ b/Engine/Source/Runtime/MovieSceneTracks/Private/Tracks/MovieSceneDataLayerTrack.cpp
@@ -0,0 +1,61 @@
+// Copyright Epic Games, Inc. All Rights Reserved.
+
+#include "Tracks/MovieSceneDataLayerTrack.h"
+#include "Sections/MovieSceneDataLayerSection.h"
+
+#define LOCTEXT_NAMESPACE "MovieSceneDataLayerTrack"
+
+UMovieSceneDataLayerTrack::UMovieSceneDataLayerTrack(const FObjectInitializer& ObjectInitializer)
+	: Super(ObjectInitializer)
+{
+}
+
+bool UMovieSceneDataLayerTrack::IsEmpty() const
+{
+	return Sections.Num() == 0;
+}
+
+void UMovieSceneDataLayerTrack::AddSection(UMovieSceneSection& Section)
+{
+	Sections.Add(&Section);
+}
+
+void UMovieSceneDataLayerTrack::RemoveSection(UMovieSceneSection& Section)
+{
+	Sections.Remove(&Section);
+}
+
+void UMovieSceneDataLayerTrack::RemoveSectionAt(int32 SectionIndex)
+{
+	Sections.RemoveAt(SectionIndex);
+}
+
+bool UMovieSceneDataLayerTrack::SupportsType(TSubclassOf<UMovieSceneSection> SectionClass) const
+{
+	return SectionClass == UMovieSceneDataLayerSection::StaticClass();
+}
+
+UMovieSceneSection* UMovieSceneDataLayerTrack::CreateNewSection()
+{
+	return NewObject<UMovieSceneDataLayerSection>(this, NAME_None, RF_Transactional);
+}
+
+const TArray<UMovieSceneSection*>& UMovieSceneDataLayerTrack::GetAllSections() const
+{
+	return Sections;
+}
+
+bool UMovieSceneDataLayerTrack::HasSection(const UMovieSceneSection& Section) const
+{
+	return Sections.Contains(&Section);
+}
+
+
+#if WITH_EDITORONLY_DATA
+FText UMovieSceneDataLayerTrack::GetDefaultDisplayName() const
+{
+	return LOCTEXT("DisplayName", "Data Layer");
+}
+#endif
+
+#undef LOCTEXT_NAMESPACE
diff --git a/Engine/Source/Runtime/MovieSceneTracks/Public/MovieSceneTracksComponentTypes.h b/Engine/Source/Runtime/MovieSceneTracks/Public/MovieSceneTracksComponentTypes.h
--- a/Engine/Source/Runtime/MovieSceneTracks/Public/MovieSceneTracksComponentTypes.h
+++ b/Engine/Source/Runtime/MovieSceneTracks/Public/MovieSceneTracksComponentTypes.h
@@ -15,6 +15,7 @@
 #include "Styling/SlateColor.h"
 #include "MovieSceneTracksComponentTypes.generated.h"
 
+class UMovieSceneDataLayerSection;
 class UMovieSceneLevelVisibilitySection;
 struct FMovieSceneObjectBindingID;
 
@@ -28,6 +29,16 @@ struct FLevelVisibilityComponentData
 	TObjectPtr<const UMovieSceneLevelVisibilitySection> Section = nullptr;
 };
 
+/** Component data for the data layer system */
+USTRUCT()
+struct FMovieSceneDataLayerComponentData
+{
+	GENERATED_BODY()
+
+	UPROPERTY()
+	TObjectPtr<const UMovieSceneDataLayerSection> Section = nullptr;
+};
+
 namespace UE
 {
 namespace MovieScene
@@ -255,6 +266,7 @@ struct MOVIESCENETRACKS_API FMovieSceneTracksComponentTypes
 	} Accessors;
 
 	TComponentTypeID<FLevelVisibilityComponentData> LevelVisibility;
+	TComponentTypeID<FMovieSceneDataLayerComponentData> DataLayer;
 
 	static void Destroy();
 
diff --git a/Engine/Source/Runtime/MovieSceneTracks/Public/Sections/MovieSceneDataLayerSection.h b/Engine/Source/Runtime/MovieSceneTracks/Public/Sections/MovieSceneDataLayerSection.h
new file mode 100644
--- /dev/null
+++ b/Engine/Source/Runtime/MovieSceneTracks/Public/Sections/MovieSceneDataLayerSection.h
@@ -0,0 +1,48 @@
+// Copyright Epic Games, Inc. All Rights Reserved.
+
+#pragma once
+
+#include "CoreTypes.h"
+#include "UObject/ObjectMacros.h"
+#include "WorldPartition/DataLayer/ActorDataLayer.h"
+#include "EntitySystem/IMovieSceneEntityProvider.h"
+#include "MovieSceneSection.h"
+#include "MovieSceneDataLayerSection.generated.h"
+
+UCLASS(MinimalAPI)
+class UMovieSceneDataLayerSection
+	: public UMovieSceneSection
+	, public IMovieSceneEntityProvider
+{
+public:
+
+	GENERATED_BODY()
+
+	UMovieSceneDataLayerSection(const FObjectInitializer& ObjInit);
+
+	UFUNCTION(BlueprintPure, Category = "Sequencer|Section")
+	MOVIESCENETRACKS_API EDataLayerState GetDesiredState() const;
+
+	UFUNCTION(BlueprintCallable, Category = "Sequencer|Section")
+	MOVIESCENETRACKS_API void SetDesiredState(EDataLayerState InDesiredState);
+
+	UFUNCTION(BlueprintPure, Category = "Sequencer|Section")
+	const TArray<FActorDataLayer>& GetDataLayers() const { return DataLayers; }
+
+	UFUNCTION(BlueprintCallable, Category = "Sequencer|Section")
+	void SetDataLayers(const TArray<FActorDataLayer>& InDataLayers) { DataLayers = InDataLayers; }
+
+private:
+
+	virtual void ImportEntityImpl(UMovieSceneEntitySystemLinker* EntityLinker, const FEntityImportParams& Params, FImportedEntity* OutImportedEntity) override;
+
+private:
+
+	/** A list of data layers that should be loaded or unloaded by this section */
+	UPROPERTY(EditAnywhere, Category=DataLayer)
+	TArray<FActorDataLayer> DataLayers;
+
+	/** Whether or not the levels in this section should be visible or hidden. */
+	UPROPERTY(EditAnywhere, Category=DataLayer)
+	EDataLayerState DesiredState;
+};
diff --git a/Engine/Source/Runtime/MovieSceneTracks/Public/Tracks/MovieSceneDataLayerTrack.h b/Engine/Source/Runtime/MovieSceneTracks/Public/Tracks/MovieSceneDataLayerTrack.h
new file mode 100644
--- /dev/null
+++ b/Engine/Source/Runtime/MovieSceneTracks/Public/Tracks/MovieSceneDataLayerTrack.h
@@ -0,0 +1,39 @@
+// Copyright Epic Games, Inc. All Rights Reserved.
+
+#pragma once
+
+#include "CoreTypes.h"
+#include "MovieSceneNameableTrack.h"
+#include "MovieSceneDataLayerTrack.generated.h"
+
+/**
+ * A track that controls loading, unloading and visibility of data layers
+ */
+UCLASS( MinimalAPI )
+class UMovieSceneDataLayerTrack : public UMovieSceneNameableTrack
+{
+public:
+	GENERATED_BODY()
+
+	UMovieSceneDataLayerTrack(const FObjectInitializer& ObjectInitializer);
+
+	/** UMovieSceneTrack interface */
+	virtual bool IsEmpty() const override;
+	virtual void AddSection(UMovieSceneSection& Section) override;
+	virtual void RemoveSection( UMovieSceneSection& Section ) override;
+	virtual void RemoveSectionAt(int32 SectionIndex) override;
+	virtual bool SupportsType(TSubclassOf<UMovieSceneSection> SectionClass) const override;
+	virtual UMovieSceneSection* CreateNewSection() override;
+	virtual const TArray<UMovieSceneSection*>& GetAllSections() const override;
+	virtual bool HasSection(const UMovieSceneSection& Section) const override;
+	virtual bool SupportsMultipleRows() const override { return true; }
+
+#if WITH_EDITORONLY_DATA
+	virtual FText GetDefaultDisplayName() const override;
+#endif
+
+private:
+
+	UPROPERTY()
+	TArray<TObjectPtr<UMovieSceneSection>> Sections;
+};
