Sequencer: Always evaluate data layers in preroll

diff --git a/Engine/Source/Runtime/MovieSceneTracks/Private/Tracks/MovieSceneDataLayerTrack.cpp b/Engine/Source/Runtime/MovieSceneTracks/Private/Tracks/MovieSceneDataLayerTrack.cpp
--- a/Engine/Source/Runtime/MovieSceneTracks/Private/Tracks/MovieSceneDataLayerTrack.cpp
+++ b/Engine/Source/Runtime/MovieSceneTracks/Private/Tracks/MovieSceneDataLayerTrack.cpp
@@ -8,6 +8,7 @@
 UMovieSceneDataLayerTrack::UMovieSceneDataLayerTrack(const FObjectInitializer& ObjectInitializer)
 	: Super(ObjectInitializer)
 {
+	EvalOptions.bEvaluateInPreroll = EvalOptions.bEvaluateInPostroll = true;
 }
 
 bool UMovieSceneDataLayerTrack::IsEmpty() const
