Change FPreAnimatedDataLayerStorage::SavePreAnimatedStateInEditor to take loaded state in consideration when saving existing state

diff --git a/Engine/Source/Runtime/MovieSceneTracks/Private/Systems/MovieSceneDataLayerSystem.cpp b/Engine/Source/Runtime/MovieSceneTracks/Private/Systems/MovieSceneDataLayerSystem.cpp
--- a/Engine/Source/Runtime/MovieSceneTracks/Private/Systems/MovieSceneDataLayerSystem.cpp
+++ b/Engine/Source/Runtime/MovieSceneTracks/Private/Systems/MovieSceneDataLayerSystem.cpp
@@ -168,7 +168,7 @@ void FPreAnimatedDataLayerStorage::SavePreAnimatedStateInEditor(UDataLayer* Data
 	if (!IsStorageRequirementSatisfied(Entry.ValueHandle.StorageIndex, StorageRequirement))
 	{
 		// We never unload data-layers in editor, so feign currently unloaded layers as loaded
-		EDataLayerState ExistingState = DataLayer->IsVisible() ? EDataLayerState::Activated : EDataLayerState::Loaded;
+		EDataLayerState ExistingState = (DataLayer->IsVisible() && DataLayer->IsDynamicallyLoadedInEditor()) ? EDataLayerState::Activated : EDataLayerState::Loaded;
 
 		AssignPreAnimatedValue(Entry.ValueHandle.StorageIndex, StorageRequirement, CopyTemp(ExistingState));
 	}
