Sequencer: Block until level streaming has been flushed after making an Activation or Unload request from a Data Layer track

diff --git a/Engine/Source/Runtime/MovieSceneTracks/Private/Systems/MovieSceneDataLayerSystem.cpp b/Engine/Source/Runtime/MovieSceneTracks/Private/Systems/MovieSceneDataLayerSystem.cpp
--- a/Engine/Source/Runtime/MovieSceneTracks/Private/Systems/MovieSceneDataLayerSystem.cpp
+++ b/Engine/Source/Runtime/MovieSceneTracks/Private/Systems/MovieSceneDataLayerSystem.cpp
@@ -12,6 +12,7 @@
 #include "Sections/MovieSceneDataLayerSection.h"
 #include "WorldPartition/DataLayer/DataLayerSubsystem.h"
 #include "Engine/World.h"
+#include "Misc/EnumClassFlags.h"
 
 #include "Kismet/KismetSystemLibrary.h"
 
@@ -24,6 +25,13 @@ namespace UE
 namespace MovieScene
 {
 
+enum class EDataLayerUpdateFlags : uint8
+{
+	None           = 0,
+	FlushStreaming = 1,
+};
+ENUM_CLASS_FLAGS(EDataLayerUpdateFlags)
+
 /** Traits class governing how pre-animated state is (re)stored for data layers */
 struct FPreAnimatedDataLayerStorageTraits
 {
@@ -83,7 +91,7 @@ struct FDesiredLayerStates
 {
 	bool IsEmpty() const;
 	void Reset();
-	void Apply(FPreAnimatedDataLayerStorage* PreAnimatedStorage, UDataLayerSubsystem* SubSystem);
+	EDataLayerUpdateFlags Apply(FPreAnimatedDataLayerStorage* PreAnimatedStorage, UDataLayerSubsystem* SubSystem);
 #if WITH_EDITOR
 	void ApplyInEditor(FPreAnimatedDataLayerStorage* PreAnimatedStorage, UDataLayerEditorSubsystem* EditorSubSystem);
 #endif
@@ -247,8 +255,10 @@ void FDesiredLayerStates::Reset()
 	}
 }
 
-void FDesiredLayerStates::Apply(FPreAnimatedDataLayerStorage* PreAnimatedStorage, UDataLayerSubsystem* SubSystem)
+EDataLayerUpdateFlags FDesiredLayerStates::Apply(FPreAnimatedDataLayerStorage* PreAnimatedStorage, UDataLayerSubsystem* SubSystem)
 {
+	EDataLayerUpdateFlags Flags = EDataLayerUpdateFlags::None;
+
 	for (auto It = StatesByLayer.CreateIterator(); It; ++It)
 	{
 		const FDataLayerState& StateValue = It.Value();
@@ -269,6 +279,11 @@ void FDesiredLayerStates::Apply(FPreAnimatedDataLayerStorage* PreAnimatedStorage
 				}
 
 				SubSystem->SetDataLayerState(DataLayer, DesiredState.GetValue());
+
+				if (DesiredState.GetValue() == EDataLayerState::Activated || DesiredState.GetValue() == EDataLayerState::Unloaded)
+				{
+					Flags |= EDataLayerUpdateFlags::FlushStreaming;
+				}
 			}
 			else
 			{
@@ -276,6 +291,8 @@ void FDesiredLayerStates::Apply(FPreAnimatedDataLayerStorage* PreAnimatedStorage
 			}
 		}
 	}
+
+	return Flags;
 }
 
 #if WITH_EDITOR
@@ -397,6 +414,8 @@ bool UMovieSceneDataLayerSystem::IsRelevantImpl(UMovieSceneEntitySystemLinker* I
 
 void UMovieSceneDataLayerSystem::OnRun(FSystemTaskPrerequisites& InPrerequisites, FSystemSubsequentTasks& Subsequents)
 {
+	using namespace UE::MovieScene;
+
 	// Only run if we must
 	UWorld* World = GetWorld();
 
@@ -426,7 +445,12 @@ void UMovieSceneDataLayerSystem::OnRun(FSystemTaskPrerequisites& InPrerequisites
 		UDataLayerSubsystem* DataLayerSubsystem = World->GetSubsystem<UDataLayerSubsystem>();
 		if (ensureMsgf(DataLayerSubsystem, TEXT("Unable to retrieve data layer subsystem - data layer tracks will not function correctly")))
 		{
-			DesiredLayerStates->Apply(WeakPreAnimatedStorage.Pin().Get(), DataLayerSubsystem);
+			EDataLayerUpdateFlags UpdateFlags = DesiredLayerStates->Apply(WeakPreAnimatedStorage.Pin().Get(), DataLayerSubsystem);
+
+			if (EnumHasAnyFlags(UpdateFlags, EDataLayerUpdateFlags::FlushStreaming))
+			{
+				World->BlockTillLevelStreamingCompleted();
+			}
 		}
 	}
 }
