Sequencer: Added user-defined PrerollState to data layer tracks

diff --git a/Engine/Source/Editor/MovieSceneTools/Private/TrackEditors/DataLayerTrackEditor.cpp b/Engine/Source/Editor/MovieSceneTools/Private/TrackEditors/DataLayerTrackEditor.cpp
--- a/Engine/Source/Editor/MovieSceneTools/Private/TrackEditors/DataLayerTrackEditor.cpp
+++ b/Engine/Source/Editor/MovieSceneTools/Private/TrackEditors/DataLayerTrackEditor.cpp
@@ -49,10 +49,23 @@ struct FDataLayerSection
 			+ SVerticalBox::Slot()
 			.AutoHeight()
 			[
-				SNew(STextBlock)
-				.Text(this, &FDataLayerSection::GetVisibilityText)
-				.ColorAndOpacity(this, &FDataLayerSection::GetTextColor)
-				.TextStyle(FAppStyle::Get(), "NormalText.Important")
+				SNew(SHorizontalBox)
+
+				+ SHorizontalBox::Slot()
+				.AutoWidth()
+				[
+					SNew(STextBlock)
+					.Text(this, &FDataLayerSection::GetVisibilityText)
+					.ColorAndOpacity(this, &FDataLayerSection::GetTextColor)
+					.TextStyle(FAppStyle::Get(), "NormalText.Important")
+				]
+
+				+ SHorizontalBox::Slot()
+				.AutoWidth()
+				[
+					SNew(STextBlock)
+					.Text(this, &FDataLayerSection::GetPrerollText)
+				]
 			]
 
 			+ SVerticalBox::Slot()
@@ -80,6 +93,22 @@ struct FDataLayerSection
 		return LOCTEXT("VisibilityText_Activated", "Activate");
 	}
 
+	FText GetPrerollText() const
+	{
+		UMovieSceneDataLayerSection* Section = WeakSection.Get();
+		if (Section)
+		{
+			switch (Section->GetPrerollState())
+			{
+			case EDataLayerState::Unloaded:  return LOCTEXT("PrerollText_Unloaded", "(Unloaded over time in pre/post roll)");
+			case EDataLayerState::Loaded:    return LOCTEXT("PrerollText_Loaded", "(Loaded over time in preroll)");
+			default: break;
+			}
+		}
+
+		return LOCTEXT("PrerollText_Activated", "(Activated over time in preroll)");
+	}
+
 	FText GetLayerBarText() const
 	{
 		UMovieSceneDataLayerSection* Section   = WeakSection.Get();
@@ -193,11 +222,25 @@ UMovieSceneDataLayerSection* FDataLayerTrackEditor::AddNewSection(UMovieScene* M
 	UMovieSceneDataLayerSection* DataLayerSection = CastChecked<UMovieSceneDataLayerSection>(DataLayerTrack->CreateNewSection());
 	DataLayerSection->SetDesiredState(DesiredState);
 
+	// By default, activated states will preroll to loaded
+	if (DesiredState == EDataLayerState::Activated)
+	{
+		DataLayerSection->SetPrerollState(EDataLayerState::Loaded);
+	}
+	else
+	{
+		DataLayerSection->SetPrerollState(DesiredState);
+	}
+
 	TRange<FFrameNumber> SectionRange = MovieScene->GetPlaybackRange();
 	DataLayerSection->InitialPlacement(DataLayerTrack->GetAllSections(), MovieScene->GetPlaybackRange().GetLowerBoundValue(), DiscreteSize(MovieScene->GetPlaybackRange()), true);
 	DataLayerTrack->AddSection(*DataLayerSection);
 
-	DataLayerSection->SetPreRollFrames((2.f * MovieScene->GetTickResolution()).RoundToFrame().Value);
+	// Set some default preroll for activated or loaded data layers
+	if (DesiredState != EDataLayerState::Unloaded)
+	{
+		DataLayerSection->SetPreRollFrames((2.f * MovieScene->GetTickResolution()).RoundToFrame().Value);
+	}
 
 	return DataLayerSection;
 }
diff --git a/Engine/Source/Runtime/MovieSceneTracks/Private/Sections/MovieSceneDataLayerSection.cpp b/Engine/Source/Runtime/MovieSceneTracks/Private/Sections/MovieSceneDataLayerSection.cpp
--- a/Engine/Source/Runtime/MovieSceneTracks/Private/Sections/MovieSceneDataLayerSection.cpp
+++ b/Engine/Source/Runtime/MovieSceneTracks/Private/Sections/MovieSceneDataLayerSection.cpp
@@ -8,6 +8,7 @@ UMovieSceneDataLayerSection::UMovieSceneDataLayerSection(const FObjectInitialize
 	: Super(ObjInit)
 {
 	DesiredState = EDataLayerState::Loaded;
+	PrerollState = EDataLayerState::Loaded;
 	EvalOptions.EnableAndSetCompletionMode(EMovieSceneCompletionMode::RestoreState);
 }
 
@@ -21,6 +22,16 @@ void UMovieSceneDataLayerSection::SetDesiredState(EDataLayerState InDesiredState
 	DesiredState = InDesiredState;
 }
 
+EDataLayerState UMovieSceneDataLayerSection::GetPrerollState() const
+{
+	return PrerollState;
+}
+
+void UMovieSceneDataLayerSection::SetPrerollState(EDataLayerState InPrerollState)
+{
+	PrerollState = InPrerollState;
+}
+
 void UMovieSceneDataLayerSection::ImportEntityImpl(UMovieSceneEntitySystemLinker* EntityLinker, const FEntityImportParams& Params, FImportedEntity* OutImportedEntity)
 {
 	using namespace UE::MovieScene;
diff --git a/Engine/Source/Runtime/MovieSceneTracks/Private/Systems/MovieSceneDataLayerSystem.cpp b/Engine/Source/Runtime/MovieSceneTracks/Private/Systems/MovieSceneDataLayerSystem.cpp
--- a/Engine/Source/Runtime/MovieSceneTracks/Private/Systems/MovieSceneDataLayerSystem.cpp
+++ b/Engine/Source/Runtime/MovieSceneTracks/Private/Systems/MovieSceneDataLayerSystem.cpp
@@ -75,8 +75,9 @@ struct FDataLayerState
 {
 	void Reset();
 	bool IsEmpty() const;
-	void AddRequest(int16 InBias, EDataLayerState RequestedState);
+	void AddRequest(int16 InBias, EDataLayerState RequestedState, bool bRequiresStreamingFlush);
 	TOptional<EDataLayerState> ComputeDesiredState() const;
+	bool ShouldFlushStreaming(EDataLayerState ComputedState) const;
 
 private:
 
@@ -84,6 +85,8 @@ private:
 	int32 UnloadedCount    = 0;
 	int32 LoadedCount      = 0;
 	int32 ActivatedCount   = 0;
+	bool  bFlushUnloaded   = false;
+	bool  bFlushActivated  = false;
 };
 
 
@@ -95,7 +98,7 @@ struct FDesiredLayerStates
 #if WITH_EDITOR
 	void ApplyInEditor(FPreAnimatedDataLayerStorage* PreAnimatedStorage, UDataLayerEditorSubsystem* EditorSubSystem);
 #endif
-	void ApplyNewState(const FName& InDataLayerName, int16 HierarchicalBias, EDataLayerState DesiredState);
+	void ApplyNewState(const FName& InDataLayerName, int16 HierarchicalBias, EDataLayerState DesiredState, bool bRequiresStreamingFlush);
 
 	TMap<FName, FDataLayerState> StatesByLayer;
 };
@@ -191,9 +194,11 @@ void FDataLayerState::Reset()
 	UnloadedCount    = 0;
 	LoadedCount      = 0;
 	ActivatedCount   = 0;
+	bFlushUnloaded   = false;
+	bFlushActivated  = false;
 }
 
-void FDataLayerState::AddRequest(int16 InBias, EDataLayerState RequestedState)
+void FDataLayerState::AddRequest(int16 InBias, EDataLayerState RequestedState, bool bRequiresStreamingFlush)
 {
 	if (InBias > HierarchicalBias)
 	{
@@ -205,9 +210,9 @@ void FDataLayerState::AddRequest(int16 InBias, EDataLayerState RequestedState)
 	{
 		switch (RequestedState)
 		{
-		case EDataLayerState::Unloaded:  ++UnloadedCount;  break;
+		case EDataLayerState::Unloaded:  ++UnloadedCount;  bFlushUnloaded |= bRequiresStreamingFlush; break;
 		case EDataLayerState::Loaded:    ++LoadedCount;    break;
-		case EDataLayerState::Activated: ++ActivatedCount; break;
+		case EDataLayerState::Activated: ++ActivatedCount; bFlushActivated |= bRequiresStreamingFlush; break;
 		}
 	}
 }
@@ -240,6 +245,17 @@ TOptional<EDataLayerState> FDataLayerState::ComputeDesiredState() const
 	return FallbackState;
 }
 
+bool FDataLayerState::ShouldFlushStreaming(EDataLayerState ComputedState) const
+{
+	switch (ComputedState)
+	{
+	case EDataLayerState::Unloaded:  return bFlushUnloaded;
+	case EDataLayerState::Activated: return bFlushActivated;
+	}
+
+	return false;
+}
+
 // ---------------------------------------------------------------------
 // FDataLayerStates definitions
 bool FDesiredLayerStates::IsEmpty() const
@@ -280,7 +296,7 @@ EDataLayerUpdateFlags FDesiredLayerStates::Apply(FPreAnimatedDataLayerStorage* P
 
 				SubSystem->SetDataLayerState(DataLayer, DesiredState.GetValue());
 
-				if (DesiredState.GetValue() == EDataLayerState::Activated || DesiredState.GetValue() == EDataLayerState::Unloaded)
+				if (StateValue.ShouldFlushStreaming(DesiredState.GetValue()))
 				{
 					Flags |= EDataLayerUpdateFlags::FlushStreaming;
 				}
@@ -362,7 +378,7 @@ void FDesiredLayerStates::ApplyInEditor(FPreAnimatedDataLayerStorage* PreAnimate
 }
 #endif
 
-void FDesiredLayerStates::ApplyNewState(const FName& InDataLayerName, int16 HierarchicalBias, EDataLayerState DesiredState)
+void FDesiredLayerStates::ApplyNewState(const FName& InDataLayerName, int16 HierarchicalBias, EDataLayerState DesiredState, bool bRequiresStreamingFlush)
 {
 	using namespace UE::MovieScene;
 
@@ -372,7 +388,7 @@ void FDesiredLayerStates::ApplyNewState(const FName& InDataLayerName, int16 Hier
 		LayerState = &StatesByLayer.Add(InDataLayerName, FDataLayerState());
 	}
 
-	LayerState->AddRequest(HierarchicalBias, DesiredState);
+	LayerState->AddRequest(HierarchicalBias, DesiredState, bRequiresStreamingFlush);
 }
 
 } // namespace MovieScene
@@ -478,22 +494,12 @@ void UMovieSceneDataLayerSystem::UpdateDesiredStates()
 				continue;
 			}
 
-			EDataLayerState DesiredState = Section->GetDesiredState();
-			if (bPreroll)
-			{
-				// Ignore unloaded data layers in preroll
-				if (DesiredState == EDataLayerState::Unloaded)
-				{
-					continue;
-				}
-
-				// Preroll always means keep the data layer in memory, but do not activate until the section is active
-				DesiredState = EDataLayerState::Loaded;
-			}
+			const bool bRequiresStreamingFlush = bPreroll == false;
+			EDataLayerState DesiredState = bPreroll ? Section->GetPrerollState() : Section->GetDesiredState();
 
 			for (const FActorDataLayer& ActorDataLayer : Section->GetDataLayers())
 			{
-				this->DesiredLayerStates->ApplyNewState(ActorDataLayer.Name, OptHBiases ? OptHBiases[Index] : 0, DesiredState);
+				this->DesiredLayerStates->ApplyNewState(ActorDataLayer.Name, OptHBiases ? OptHBiases[Index] : 0, DesiredState, bRequiresStreamingFlush);
 			}
 		}
 	};
diff --git a/Engine/Source/Runtime/MovieSceneTracks/Public/Sections/MovieSceneDataLayerSection.h b/Engine/Source/Runtime/MovieSceneTracks/Public/Sections/MovieSceneDataLayerSection.h
--- a/Engine/Source/Runtime/MovieSceneTracks/Public/Sections/MovieSceneDataLayerSection.h
+++ b/Engine/Source/Runtime/MovieSceneTracks/Public/Sections/MovieSceneDataLayerSection.h
@@ -26,6 +26,12 @@ public:
 	UFUNCTION(BlueprintCallable, Category = "Sequencer|Section")
 	MOVIESCENETRACKS_API void SetDesiredState(EDataLayerState InDesiredState);
 
+	UFUNCTION(BlueprintPure, Category = "Sequencer|Section")
+	MOVIESCENETRACKS_API EDataLayerState GetPrerollState() const;
+
+	UFUNCTION(BlueprintCallable, Category = "Sequencer|Section")
+	MOVIESCENETRACKS_API void SetPrerollState(EDataLayerState InPrerollState);
+
 	UFUNCTION(BlueprintPure, Category = "Sequencer|Section")
 	const TArray<FActorDataLayer>& GetDataLayers() const { return DataLayers; }
 
@@ -42,7 +48,11 @@ private:
 	UPROPERTY(EditAnywhere, Category=DataLayer)
 	TArray<FActorDataLayer> DataLayers;
 
-	/** Whether or not the levels in this section should be visible or hidden. */
+	/** The desired state for the data layers on this section when the section is actively evaluating. */
 	UPROPERTY(EditAnywhere, Category=DataLayer)
 	EDataLayerState DesiredState;
+
+	/** The desired state for the data layers on this section when the section is pre or post-rolling. */
+	UPROPERTY(EditAnywhere, Category=DataLayer)
+	EDataLayerState PrerollState;
 };
