Sequencer: Add Query to block streaming state only when we are out of preroll

diff --git a/Engine/Source/Runtime/MovieSceneTracks/Private/Systems/MovieSceneDataLayerSystem.cpp b/Engine/Source/Runtime/MovieSceneTracks/Private/Systems/MovieSceneDataLayerSystem.cpp
--- a/Engine/Source/Runtime/MovieSceneTracks/Private/Systems/MovieSceneDataLayerSystem.cpp
+++ b/Engine/Source/Runtime/MovieSceneTracks/Private/Systems/MovieSceneDataLayerSystem.cpp
@@ -11,6 +11,9 @@
 #include "MovieSceneTracksComponentTypes.h"
 #include "Sections/MovieSceneDataLayerSection.h"
 #include "WorldPartition/DataLayer/DataLayerSubsystem.h"
+#include "WorldPartition/WorldPartitionSubsystem.h"
+#include "WorldPartition/WorldPartitionRuntimeCell.h"
+#include "WorldPartition/WorldPartitionStreamingSource.h"
 #include "Engine/World.h"
 #include "Misc/EnumClassFlags.h"
 
@@ -94,7 +97,7 @@ struct FDesiredLayerStates
 {
 	bool IsEmpty() const;
 	void Reset();
-	EDataLayerUpdateFlags Apply(FPreAnimatedDataLayerStorage* PreAnimatedStorage, UDataLayerSubsystem* SubSystem);
+	EDataLayerUpdateFlags Apply(FPreAnimatedDataLayerStorage* PreAnimatedStorage, UDataLayerSubsystem* DataLayerSubsystem, UWorldPartitionSubsystem* WorldPartitionSubsystem);
 #if WITH_EDITOR
 	void ApplyInEditor(FPreAnimatedDataLayerStorage* PreAnimatedStorage, UDataLayerEditorSubsystem* EditorSubSystem);
 #endif
@@ -271,10 +274,36 @@ void FDesiredLayerStates::Reset()
 	}
 }
 
-EDataLayerUpdateFlags FDesiredLayerStates::Apply(FPreAnimatedDataLayerStorage* PreAnimatedStorage, UDataLayerSubsystem* SubSystem)
+EDataLayerUpdateFlags FDesiredLayerStates::Apply(FPreAnimatedDataLayerStorage* PreAnimatedStorage, UDataLayerSubsystem* DataLayerSubsystem, UWorldPartitionSubsystem* WorldPartitionSubsystem)
 {
 	EDataLayerUpdateFlags Flags = EDataLayerUpdateFlags::None;
 
+	auto IsDataLayerReady = [WorldPartitionSubsystem](UDataLayer* DataLayer, EDataLayerState DesireState)
+	{
+		EWorldPartitionRuntimeCellState QueryState;
+		switch (DesireState)
+		{
+		case EDataLayerState::Activated:
+			QueryState = EWorldPartitionRuntimeCellState::Activated;
+			break;
+		case EDataLayerState::Loaded:
+			QueryState = EWorldPartitionRuntimeCellState::Loaded;
+			break;
+		case EDataLayerState::Unloaded:
+			QueryState = EWorldPartitionRuntimeCellState::Unloaded;
+			break;
+		default:
+			UE_LOG(LogMovieScene, Error, TEXT("Unkown data layer state"));
+			return true;
+		}
+
+		FWorldPartitionStreamingQuerySource QuerySource;
+		QuerySource.bDataLayersOnly = true;
+		QuerySource.bSpatialQuery = false; // @todo_ow: how would we support spatial query from sequencer?
+		QuerySource.DataLayers.Add(DataLayer->GetFName());
+		return WorldPartitionSubsystem->IsStreamingCompleted(QueryState, { QuerySource }, true);
+	};
+
 	for (auto It = StatesByLayer.CreateIterator(); It; ++It)
 	{
 		const FDataLayerState& StateValue = It.Value();
@@ -286,19 +315,21 @@ EDataLayerUpdateFlags FDesiredLayerStates::Apply(FPreAnimatedDataLayerStorage* P
 
 		if (TOptional<EDataLayerState> DesiredState = StateValue.ComputeDesiredState())
 		{
-			UDataLayer* DataLayer = SubSystem->GetDataLayerFromName(It.Key());
+			UDataLayer* DataLayer = DataLayerSubsystem->GetDataLayerFromName(It.Key());
 			if (DataLayer)
 			{
 				if (PreAnimatedStorage)
 				{
-					PreAnimatedStorage->SavePreAnimatedState(DataLayer, SubSystem);
+					PreAnimatedStorage->SavePreAnimatedState(DataLayer, DataLayerSubsystem);
 				}
 
-				SubSystem->SetDataLayerState(DataLayer, DesiredState.GetValue());
+				const EDataLayerState DesiredStateValue = DesiredState.GetValue();
+				DataLayerSubsystem->SetDataLayerState(DataLayer, DesiredStateValue);
 
-				if (StateValue.ShouldFlushStreaming(DesiredState.GetValue()))
+				if (StateValue.ShouldFlushStreaming(DesiredStateValue) && !IsDataLayerReady(DataLayer, DesiredStateValue))
 				{
 					Flags |= EDataLayerUpdateFlags::FlushStreaming;
+					UE_LOG(LogMovieScene, Warning, TEXT("Data layer with name '%s' is causing a streaming flush"), *DataLayer->GetDataLayerLabel().ToString());
 				}
 			}
 			else
@@ -458,14 +489,18 @@ void UMovieSceneDataLayerSystem::OnRun(FSystemTaskPrerequisites& InPrerequisites
 
 	// Outside of editor, or in PIE, we use the runtime data layer sub-system
 	{
-		UDataLayerSubsystem* DataLayerSubsystem = World->GetSubsystem<UDataLayerSubsystem>();
-		if (ensureMsgf(DataLayerSubsystem, TEXT("Unable to retrieve data layer subsystem - data layer tracks will not function correctly")))
+		UWorldPartitionSubsystem* WorldPartitionSubsystem = World->GetSubsystem<UWorldPartitionSubsystem>();
+		if (ensureMsgf(WorldPartitionSubsystem, TEXT("Unable to retrieve world partition subsystem - data layer tracks will not function correctly")))
 		{
-			EDataLayerUpdateFlags UpdateFlags = DesiredLayerStates->Apply(WeakPreAnimatedStorage.Pin().Get(), DataLayerSubsystem);
-
-			if (EnumHasAnyFlags(UpdateFlags, EDataLayerUpdateFlags::FlushStreaming))
+			UDataLayerSubsystem* DataLayerSubsystem = World->GetSubsystem<UDataLayerSubsystem>();
+			if (ensureMsgf(DataLayerSubsystem, TEXT("Unable to retrieve data layer subsystem - data layer tracks will not function correctly")))
 			{
-				World->BlockTillLevelStreamingCompleted();
+				EDataLayerUpdateFlags UpdateFlags = DesiredLayerStates->Apply(WeakPreAnimatedStorage.Pin().Get(), DataLayerSubsystem, WorldPartitionSubsystem);
+
+				if (EnumHasAnyFlags(UpdateFlags, EDataLayerUpdateFlags::FlushStreaming))
+				{
+					World->BlockTillLevelStreamingCompleted();
+				}
 			}
 		}
 	}
