Sequencer Data Layer: Distinguish between full streaming flush and visibility flush to minimize blocking loads

diff --git a/Engine/Source/Runtime/MovieSceneTracks/Private/Systems/MovieSceneDataLayerSystem.cpp b/Engine/Source/Runtime/MovieSceneTracks/Private/Systems/MovieSceneDataLayerSystem.cpp
--- a/Engine/Source/Runtime/MovieSceneTracks/Private/Systems/MovieSceneDataLayerSystem.cpp
+++ b/Engine/Source/Runtime/MovieSceneTracks/Private/Systems/MovieSceneDataLayerSystem.cpp
@@ -30,8 +30,9 @@ namespace MovieScene
 
 enum class EDataLayerUpdateFlags : uint8
 {
-	None           = 0,
-	FlushStreaming = 1,
+	None						= 0,
+	FlushStreamingVisibility	= 1,
+	FlushStreamingFull			= 2,
 };
 ENUM_CLASS_FLAGS(EDataLayerUpdateFlags)
 
@@ -278,7 +279,7 @@ EDataLayerUpdateFlags FDesiredLayerStates::Apply(FPreAnimatedDataLayerStorage* P
 {
 	EDataLayerUpdateFlags Flags = EDataLayerUpdateFlags::None;
 
-	auto IsDataLayerReady = [WorldPartitionSubsystem](UDataLayer* DataLayer, EDataLayerState DesireState)
+	auto IsDataLayerReady = [WorldPartitionSubsystem](UDataLayer* DataLayer, EDataLayerState DesireState, bool bExactState)
 	{
 		EWorldPartitionRuntimeCellState QueryState;
 		switch (DesireState)
@@ -301,7 +302,7 @@ EDataLayerUpdateFlags FDesiredLayerStates::Apply(FPreAnimatedDataLayerStorage* P
 		QuerySource.bDataLayersOnly = true;
 		QuerySource.bSpatialQuery = false; // @todo_ow: how would we support spatial query from sequencer?
 		QuerySource.DataLayers.Add(DataLayer->GetFName());
-		return WorldPartitionSubsystem->IsStreamingCompleted(QueryState, { QuerySource }, true);
+		return WorldPartitionSubsystem->IsStreamingCompleted(QueryState, { QuerySource }, bExactState);
 	};
 
 	for (auto It = StatesByLayer.CreateIterator(); It; ++It)
@@ -326,10 +327,19 @@ EDataLayerUpdateFlags FDesiredLayerStates::Apply(FPreAnimatedDataLayerStorage* P
 				const EDataLayerState DesiredStateValue = DesiredState.GetValue();
 				DataLayerSubsystem->SetDataLayerState(DataLayer, DesiredStateValue);
 
-				if (StateValue.ShouldFlushStreaming(DesiredStateValue) && !IsDataLayerReady(DataLayer, DesiredStateValue))
+				if (StateValue.ShouldFlushStreaming(DesiredStateValue) && !IsDataLayerReady(DataLayer, DesiredStateValue, true))
 				{
-					Flags |= EDataLayerUpdateFlags::FlushStreaming;
-					UE_LOG(LogMovieScene, Warning, TEXT("Data layer with name '%s' is causing a streaming flush"), *DataLayer->GetDataLayerLabel().ToString());
+					// Exception for Full flush is if Desired State is Activated but we are not at least in Loaded state
+					if (DesiredStateValue == EDataLayerState::Activated && !IsDataLayerReady(DataLayer, EDataLayerState::Loaded, false))
+					{
+						Flags |= EDataLayerUpdateFlags::FlushStreamingFull;
+						UE_LOG(LogMovieScene, Warning, TEXT("Data layer with name '%s' is causing a full streaming flush"), *DataLayer->GetDataLayerLabel().ToString());
+					}
+					else
+					{
+						Flags |= EDataLayerUpdateFlags::FlushStreamingVisibility;
+						UE_LOG(LogMovieScene, Log, TEXT("Data layer with name '%s' is causing a visibility streaming flush"), *DataLayer->GetDataLayerLabel().ToString());
+					}
 				}
 			}
 			else
@@ -497,10 +507,16 @@ void UMovieSceneDataLayerSystem::OnRun(FSystemTaskPrerequisites& InPrerequisites
 			{
 				EDataLayerUpdateFlags UpdateFlags = DesiredLayerStates->Apply(WeakPreAnimatedStorage.Pin().Get(), DataLayerSubsystem, WorldPartitionSubsystem);
 
-				if (EnumHasAnyFlags(UpdateFlags, EDataLayerUpdateFlags::FlushStreaming))
+				if (EnumHasAnyFlags(UpdateFlags, EDataLayerUpdateFlags::FlushStreamingFull))
 				{
 					World->BlockTillLevelStreamingCompleted();
 				}
+				else if (EnumHasAnyFlags(UpdateFlags, EDataLayerUpdateFlags::FlushStreamingVisibility))
+				{
+					// Make sure any DataLayer state change is processed before flushing visibility
+					WorldPartitionSubsystem->UpdateStreamingState();
+					World->FlushLevelStreaming(EFlushLevelStreamingType::Visibility);
+				}
 			}
 		}
 	}
