Sequencer: Change invalid data layer ensure to log warning since it's possible to just open a level sequence without the map.

diff --git a/Engine/Source/Runtime/MovieSceneTracks/Private/Systems/MovieSceneDataLayerSystem.cpp b/Engine/Source/Runtime/MovieSceneTracks/Private/Systems/MovieSceneDataLayerSystem.cpp
--- a/Engine/Source/Runtime/MovieSceneTracks/Private/Systems/MovieSceneDataLayerSystem.cpp
+++ b/Engine/Source/Runtime/MovieSceneTracks/Private/Systems/MovieSceneDataLayerSystem.cpp
@@ -489,7 +489,11 @@ void UMovieSceneDataLayerSystem::OnRun(FSystemTaskPrerequisites& InPrerequisites
 	if (World->WorldType == EWorldType::Editor)
 	{
 		UDataLayerEditorSubsystem* DataLayerEditorSubsystem = UDataLayerEditorSubsystem::Get();
-		if (ensureMsgf(DataLayerEditorSubsystem, TEXT("Unable to retrieve data layer editor subsystem - data layer tracks will not function correctly")))
+		if (!DataLayerEditorSubsystem)
+		{
+			UE_LOG(LogMovieScene, Warning, TEXT("Unable to retrieve data layer editor subsystem - data layer tracks will not function correctly"));
+		}
+		else
 		{
 			DesiredLayerStates->ApplyInEditor(WeakPreAnimatedStorage.Pin().Get(), DataLayerEditorSubsystem);
 		}
@@ -500,10 +504,18 @@ void UMovieSceneDataLayerSystem::OnRun(FSystemTaskPrerequisites& InPrerequisites
 	// Outside of editor, or in PIE, we use the runtime data layer sub-system
 	{
 		UWorldPartitionSubsystem* WorldPartitionSubsystem = World->GetSubsystem<UWorldPartitionSubsystem>();
-		if (ensureMsgf(WorldPartitionSubsystem, TEXT("Unable to retrieve world partition subsystem - data layer tracks will not function correctly")))
+		if (!WorldPartitionSubsystem)
+		{
+			UE_LOG(LogMovieScene, Warning, TEXT("Unable to retrieve world partition subsystem - data layer tracks will not function correctly"));
+		}
+		else
 		{
 			UDataLayerSubsystem* DataLayerSubsystem = World->GetSubsystem<UDataLayerSubsystem>();
-			if (ensureMsgf(DataLayerSubsystem, TEXT("Unable to retrieve data layer subsystem - data layer tracks will not function correctly")))
+			if (!DataLayerSubsystem)
+			{
+				UE_LOG(LogMovieScene, Warning, TEXT("Unable to retrieve data layer subsystem - data layer tracks will not function correctly"));
+			}
+			else
 			{
 				EDataLayerUpdateFlags UpdateFlags = DesiredLayerStates->Apply(WeakPreAnimatedStorage.Pin().Get(), DataLayerSubsystem, WorldPartitionSubsystem);
 
