Sequencer: Skip applying simulated transforms (for motion blur on first frame of camera cuts) to null attach children.

diff --git a/Engine/Source/Runtime/MovieSceneTracks/Private/Systems/MovieSceneMotionVectorSimulationSystem.cpp b/Engine/Source/Runtime/MovieSceneTracks/Private/Systems/MovieSceneMotionVectorSimulationSystem.cpp
--- a/Engine/Source/Runtime/MovieSceneTracks/Private/Systems/MovieSceneMotionVectorSimulationSystem.cpp
+++ b/Engine/Source/Runtime/MovieSceneTracks/Private/Systems/MovieSceneMotionVectorSimulationSystem.cpp
@@ -317,6 +317,11 @@ void UMovieSceneMotionVectorSimulationSystem::ApplySimulatedTransforms(USceneCom
 
 	for (USceneComponent* Child : InComponent->GetAttachChildren())
 	{
+		if (!Child)
+		{
+			continue;
+		}
+
 		FName AttachSocketName = Child->GetAttachSocketName();
 
 		FTransform SocketTransform = (AttachSocketName == NAME_None) ? FTransform::Identity : GetSocketTransform(InComponent, AttachSocketName);
