Sequencer: Large world coordinates support

diff --git a/Engine/Source/Runtime/MovieScene/Private/EntitySystem/BuiltInComponentTypes.cpp b/Engine/Source/Runtime/MovieScene/Private/EntitySystem/BuiltInComponentTypes.cpp
--- a/Engine/Source/Runtime/MovieScene/Private/EntitySystem/BuiltInComponentTypes.cpp
+++ b/Engine/Source/Runtime/MovieScene/Private/EntitySystem/BuiltInComponentTypes.cpp
@@ -1,6 +1,7 @@
 // Copyright Epic Games, Inc. All Rights Reserved.
 
 #include "EntitySystem/BuiltInComponentTypes.h"
+#include "Channels/MovieSceneDoubleChannel.h"
 #include "Channels/MovieSceneFloatChannel.h"
 #include "EntitySystem/MovieSceneEntityManager.h"
 #include "EntitySystem/MovieSceneEntitySystemLinker.h"
@@ -46,6 +47,15 @@ FBuiltInComponentTypes::FBuiltInComponentTypes()
 	ComponentRegistry->NewComponentType(&FloatChannel[6],         TEXT("Float Channel 6"));
 	ComponentRegistry->NewComponentType(&FloatChannel[7],         TEXT("Float Channel 7"));
 	ComponentRegistry->NewComponentType(&FloatChannel[8],         TEXT("Float Channel 8"));
+	ComponentRegistry->NewComponentType(&DoubleChannel[0],        TEXT("Double Channel 0"));
+	ComponentRegistry->NewComponentType(&DoubleChannel[1],        TEXT("Double Channel 1"));
+	ComponentRegistry->NewComponentType(&DoubleChannel[2],        TEXT("Double Channel 2"));
+	ComponentRegistry->NewComponentType(&DoubleChannel[3],        TEXT("Double Channel 3"));
+	ComponentRegistry->NewComponentType(&DoubleChannel[4],        TEXT("Double Channel 4"));
+	ComponentRegistry->NewComponentType(&DoubleChannel[5],        TEXT("Double Channel 5"));
+	ComponentRegistry->NewComponentType(&DoubleChannel[6],        TEXT("Double Channel 6"));
+	ComponentRegistry->NewComponentType(&DoubleChannel[7],        TEXT("Double Channel 7"));
+	ComponentRegistry->NewComponentType(&DoubleChannel[8],        TEXT("Double Channel 8"));
 	ComponentRegistry->NewComponentType(&WeightChannel,           TEXT("Weight Channel"));
 
 	ComponentRegistry->NewComponentType(&FloatChannelFlags[0],    TEXT("Float Channel 0 Flags"));
@@ -57,6 +67,15 @@ FBuiltInComponentTypes::FBuiltInComponentTypes()
 	ComponentRegistry->NewComponentType(&FloatChannelFlags[6],    TEXT("Float Channel 6 Flags"));
 	ComponentRegistry->NewComponentType(&FloatChannelFlags[7],    TEXT("Float Channel 7 Flags"));
 	ComponentRegistry->NewComponentType(&FloatChannelFlags[8],    TEXT("Float Channel 8 Flags"));
+	ComponentRegistry->NewComponentType(&DoubleChannelFlags[0],   TEXT("Double Channel 0 Flags"));
+	ComponentRegistry->NewComponentType(&DoubleChannelFlags[1],   TEXT("Double Channel 1 Flags"));
+	ComponentRegistry->NewComponentType(&DoubleChannelFlags[2],   TEXT("Double Channel 2 Flags"));
+	ComponentRegistry->NewComponentType(&DoubleChannelFlags[3],   TEXT("Double Channel 3 Flags"));
+	ComponentRegistry->NewComponentType(&DoubleChannelFlags[4],   TEXT("Double Channel 4 Flags"));
+	ComponentRegistry->NewComponentType(&DoubleChannelFlags[5],   TEXT("Double Channel 5 Flags"));
+	ComponentRegistry->NewComponentType(&DoubleChannelFlags[6],   TEXT("Double Channel 6 Flags"));
+	ComponentRegistry->NewComponentType(&DoubleChannelFlags[7],   TEXT("Double Channel 7 Flags"));
+	ComponentRegistry->NewComponentType(&DoubleChannelFlags[8],   TEXT("Double Channel 8 Flags"));
 	ComponentRegistry->NewComponentType(&WeightChannelFlags,      TEXT("Weight Channel Flags"));
 
 	ComponentRegistry->NewComponentType(&Easing,                  TEXT("Easing"));
@@ -84,6 +103,15 @@ FBuiltInComponentTypes::FBuiltInComponentTypes()
 	ComponentRegistry->NewComponentType(&FloatResult[6],        TEXT("Float Result 6"));
 	ComponentRegistry->NewComponentType(&FloatResult[7],        TEXT("Float Result 7"));
 	ComponentRegistry->NewComponentType(&FloatResult[8],        TEXT("Float Result 8"));
+	ComponentRegistry->NewComponentType(&DoubleResult[0],       TEXT("Double Result 0"));
+	ComponentRegistry->NewComponentType(&DoubleResult[1],       TEXT("Double Result 1"));
+	ComponentRegistry->NewComponentType(&DoubleResult[2],       TEXT("Double Result 2"));
+	ComponentRegistry->NewComponentType(&DoubleResult[3],       TEXT("Double Result 3"));
+	ComponentRegistry->NewComponentType(&DoubleResult[4],       TEXT("Double Result 4"));
+	ComponentRegistry->NewComponentType(&DoubleResult[5],       TEXT("Double Result 5"));
+	ComponentRegistry->NewComponentType(&DoubleResult[6],       TEXT("Double Result 6"));
+	ComponentRegistry->NewComponentType(&DoubleResult[7],       TEXT("Double Result 7"));
+	ComponentRegistry->NewComponentType(&DoubleResult[8],       TEXT("Double Result 8"));
 
 	ComponentRegistry->NewComponentType(&BaseInteger,			TEXT("Base Integer"));
 	ComponentRegistry->NewComponentType(&BaseFloat[0],          TEXT("Base Float 0"));
@@ -95,6 +123,15 @@ FBuiltInComponentTypes::FBuiltInComponentTypes()
 	ComponentRegistry->NewComponentType(&BaseFloat[6],          TEXT("Base Float 6"));
 	ComponentRegistry->NewComponentType(&BaseFloat[7],          TEXT("Base Float 7"));
 	ComponentRegistry->NewComponentType(&BaseFloat[8],          TEXT("Base Float 8"));
+	ComponentRegistry->NewComponentType(&BaseDouble[0],         TEXT("Base Double 0"));
+	ComponentRegistry->NewComponentType(&BaseDouble[1],         TEXT("Base Double 1"));
+	ComponentRegistry->NewComponentType(&BaseDouble[2],         TEXT("Base Double 2"));
+	ComponentRegistry->NewComponentType(&BaseDouble[3],         TEXT("Base Double 3"));
+	ComponentRegistry->NewComponentType(&BaseDouble[4],         TEXT("Base Double 4"));
+	ComponentRegistry->NewComponentType(&BaseDouble[5],         TEXT("Base Double 5"));
+	ComponentRegistry->NewComponentType(&BaseDouble[6],         TEXT("Base Double 6"));
+	ComponentRegistry->NewComponentType(&BaseDouble[7],         TEXT("Base Double 7"));
+	ComponentRegistry->NewComponentType(&BaseDouble[8],         TEXT("Base Double 8"));
 
 	ComponentRegistry->NewComponentType(&BaseValueEvalTime,     TEXT("Base Value Eval Time"));
 
@@ -208,6 +245,9 @@ FBuiltInComponentTypes::FBuiltInComponentTypes()
 		static_assert(
 				UE_ARRAY_COUNT(FloatChannel) == UE_ARRAY_COUNT(BaseFloat), 
 				"Base floats and float results should have the same size.");
+		static_assert(
+				UE_ARRAY_COUNT(FloatChannel) == UE_ARRAY_COUNT(FloatChannelFlags),
+				"Float channel flags and flor channels should have the same size.");
 
 		// Duplicate float channels
 		for (int32 Index = 0; Index < UE_ARRAY_COUNT(FloatChannel); ++Index)
@@ -227,6 +267,33 @@ FBuiltInComponentTypes::FBuiltInComponentTypes()
 		}
 	}
 
+	// Double channel relationships
+	{
+		static_assert(
+				UE_ARRAY_COUNT(DoubleChannel) == UE_ARRAY_COUNT(BaseDouble), 
+				"Base doubles and double results should have the same size.");
+		static_assert(
+				UE_ARRAY_COUNT(DoubleChannel) == UE_ARRAY_COUNT(DoubleChannelFlags),
+				"Double channel flags and flor channels should have the same size.");
+
+		// Duplicate double channels
+		for (int32 Index = 0; Index < UE_ARRAY_COUNT(DoubleChannel); ++Index)
+		{
+			ComponentRegistry->Factories.DuplicateChildComponent(DoubleChannel[Index]);
+			ComponentRegistry->Factories.DefineMutuallyInclusiveComponent(DoubleChannel[Index], DoubleResult[Index]);
+			ComponentRegistry->Factories.DefineMutuallyInclusiveComponent(DoubleChannel[Index], EvalTime);
+			ComponentRegistry->Factories.DefineMutuallyInclusiveComponent(DoubleChannel[Index], DoubleChannelFlags[Index]);
+		}
+
+		// Create base double components for double channels that are meant to be additive from base.
+		for (int32 Index = 0; Index < UE_ARRAY_COUNT(BaseDouble); ++Index)
+		{
+			ComponentRegistry->Factories.DefineComplexInclusiveComponents(
+					FComplexInclusivityFilter::All({ DoubleChannel[Index], BaseValueEvalTime, Tags.AdditiveFromBaseBlend }),
+					BaseDouble[Index]);
+		}
+	}
+
 	// Easing component relationships
 	{
 		// Easing components should be duplicated to children
diff --git a/Engine/Source/Runtime/MovieScene/Private/EntitySystem/MovieSceneDecompositionQuery.cpp b/Engine/Source/Runtime/MovieScene/Private/EntitySystem/MovieSceneDecompositionQuery.cpp
--- a/Engine/Source/Runtime/MovieScene/Private/EntitySystem/MovieSceneDecompositionQuery.cpp
+++ b/Engine/Source/Runtime/MovieScene/Private/EntitySystem/MovieSceneDecompositionQuery.cpp
@@ -9,9 +9,9 @@ namespace UE
 namespace MovieScene
 {
 
-void FDecomposedFloat::Decompose(FMovieSceneEntityID EntityID, FWeightedFloat& ThisValue, bool& bOutIsAdditive, FWeightedFloat& Absolutes, FWeightedFloat& Additives) const
+void FDecomposedValue::Decompose(FMovieSceneEntityID EntityID, FWeightedValue& ThisValue, bool& bOutIsAdditive, FWeightedValue& Absolutes, FWeightedValue& Additives) const
 {
-	for (TTuple<FMovieSceneEntityID, FWeightedFloat> Pair : DecomposedAbsolutes)
+	for (TTuple<FMovieSceneEntityID, FWeightedValue> Pair : DecomposedAbsolutes)
 	{
 		if (Pair.Get<0>() == EntityID)
 		{
@@ -24,7 +24,7 @@ void FDecomposedFloat::Decompose(FMovieSceneEntityID EntityID, FWeightedFloat& T
 			Absolutes.Weight += Pair.Value.Weight;
 		}
 	}
-	for (TTuple<FMovieSceneEntityID, FWeightedFloat> Pair : DecomposedAdditives)
+	for (TTuple<FMovieSceneEntityID, FWeightedValue> Pair : DecomposedAdditives)
 	{
 		if (Pair.Get<0>() == EntityID)
 		{
@@ -39,16 +39,29 @@ void FDecomposedFloat::Decompose(FMovieSceneEntityID EntityID, FWeightedFloat& T
 	}
 }
 
-float FDecomposedFloat::Recompose(FMovieSceneEntityID RecomposeEntity, float CurrentValue, const float* InitialValue) const
+float FDecomposedValue::Recompose(FMovieSceneEntityID EntityID, float CurrentValue, const float* InitialValue) const
 {
-	FWeightedFloat DecomposedAbsolute;
-	FWeightedFloat DecomposedAdditive;
+	if (InitialValue)
+	{
+		double WidenedInitialValue = *InitialValue;
+		return (float)Recompose(EntityID, (double)CurrentValue, &WidenedInitialValue);
+	}
+	else
+	{
+		return (float)Recompose(EntityID, (double)CurrentValue, nullptr);
+	}
+}
+
+double FDecomposedValue::Recompose(FMovieSceneEntityID RecomposeEntity, double CurrentValue, const double* InitialValue) const
+{
+	FWeightedValue DecomposedAbsolute;
+	FWeightedValue DecomposedAdditive;
 
-	FWeightedFloat Channel;
+	FWeightedValue Channel;
 	bool bIsAdditive = false;
 	Decompose(RecomposeEntity, Channel, bIsAdditive, DecomposedAbsolute, DecomposedAdditive);
 
-	FWeightedFloat ResultAbsolute = Result.Absolute;
+	FWeightedValue ResultAbsolute = Result.Absolute;
 	float TotalAbsoluteWeight = ResultAbsolute.Weight + DecomposedAbsolute.Weight;
 	if (!bIsAdditive)
 	{
@@ -66,7 +79,7 @@ float FDecomposedFloat::Recompose(FMovieSceneEntityID RecomposeEntity, float Cur
 	{
 		if (bIsAdditive)
 		{
-			const float WeightedAdditiveResult = CurrentValue - ResultAbsolute.Combine(DecomposedAbsolute).WeightedValue() - Result.Additive;
+			const double WeightedAdditiveResult = CurrentValue - ResultAbsolute.Combine(DecomposedAbsolute).WeightedValue() - Result.Additive;
 			return Channel.Weight == 0.f ? WeightedAdditiveResult : WeightedAdditiveResult / Channel.Weight;
 		}
 		else
@@ -74,7 +87,7 @@ float FDecomposedFloat::Recompose(FMovieSceneEntityID RecomposeEntity, float Cur
 			if (Channel.Weight != 0.f)
 			{
 				const float TotalWeight = Channel.Weight + ResultAbsolute.Weight;
-				const float WeightedAbsoluteResult = CurrentValue - Result.Additive - ResultAbsolute.Value / TotalWeight;
+				const double WeightedAbsoluteResult = CurrentValue - Result.Additive - ResultAbsolute.Value / TotalWeight;
 				return WeightedAbsoluteResult * TotalWeight / Channel.Weight;
 			}
 			else
@@ -94,7 +107,7 @@ float FDecomposedFloat::Recompose(FMovieSceneEntityID RecomposeEntity, float Cur
 	{
 		CurrentValue -= ResultAbsolute.Combine(DecomposedAbsolute).WeightedValue();
 
-		const float ThisAdditive = Channel.WeightedValue();
+		const double ThisAdditive = Channel.WeightedValue();
 		if (ThisAdditive == 0.f && DecomposedAdditive.WeightedValue() == 0.f)
 		{
 			const float TotalAdditiveWeight = DecomposedAdditive.Weight + Channel.Weight;
@@ -102,7 +115,7 @@ float FDecomposedFloat::Recompose(FMovieSceneEntityID RecomposeEntity, float Cur
 		}
 
 		// Use the fractions of the values for the recomposition if we have non-zero values
-		const float DecoposeFactor = ThisAdditive / (DecomposedAdditive.WeightedValue() + ThisAdditive);
+		const double DecoposeFactor = ThisAdditive / (DecomposedAdditive.WeightedValue() + ThisAdditive);
 		return CurrentValue * DecoposeFactor / Channel.Weight;
 	}
 	else if (DecomposedAdditives.Num() != 0)
@@ -117,7 +130,7 @@ float FDecomposedFloat::Recompose(FMovieSceneEntityID RecomposeEntity, float Cur
 		CurrentValue *= ResultAbsolute.Weight + TotalDecomposedWeight;
 		CurrentValue -= ResultAbsolute.Value;
 
-		const float AbsValue = Algo::Accumulate(DecomposedAbsolutes, 0.f, [](float Accum, TTuple<FMovieSceneEntityID, FWeightedFloat> In) { return Accum + FMath::Abs(In.Value.Value)*In.Value.Weight; });
+		const double AbsValue = Algo::Accumulate(DecomposedAbsolutes, 0.f, [](double Accum, TTuple<FMovieSceneEntityID, FWeightedValue> In) { return Accum + FMath::Abs(In.Value.Value)*In.Value.Weight; });
 		if (AbsValue != 0.f)
 		{
 			return ((CurrentValue * FMath::Abs(Channel.Value) * Channel.Weight / AbsValue) - Channel.Value) / Channel.Weight;
diff --git a/Engine/Source/Runtime/MovieScene/Public/Channels/MovieSceneFloatChannel.h b/Engine/Source/Runtime/MovieScene/Public/Channels/MovieSceneFloatChannel.h
--- a/Engine/Source/Runtime/MovieScene/Public/Channels/MovieSceneFloatChannel.h
+++ b/Engine/Source/Runtime/MovieScene/Public/Channels/MovieSceneFloatChannel.h
@@ -2,81 +2,20 @@
 
 #pragma once
 
-#include "UObject/ObjectMacros.h"
+#include "Channels/MovieSceneCurveChannelCommon.h"
+#include "Curves/RichCurve.h"
+#include "KeyParams.h"
 #include "Misc/FrameNumber.h"
 #include "MovieSceneChannel.h"
 #include "MovieSceneChannelData.h"
 #include "MovieSceneChannelTraits.h"
-#include "KeyParams.h"
-
-#include "Curves/RichCurve.h"
+#include "UObject/ObjectMacros.h"
 
 #include "MovieSceneFloatChannel.generated.h"
 
 
-USTRUCT()
-struct FMovieSceneTangentData
-{
-	GENERATED_BODY()
-
-	FMovieSceneTangentData()
-		: ArriveTangent(0.f)
-		, LeaveTangent(0.f)
-		, ArriveTangentWeight(0.f)
-		, LeaveTangentWeight(0.f)
-		, TangentWeightMode(RCTWM_WeightedNone)
-
-	{}
-
-	bool Serialize(FArchive& Ar);
-	bool operator==(const FMovieSceneTangentData& Other) const;
-	bool operator!=(const FMovieSceneTangentData& Other) const;
-	friend FArchive& operator<<(FArchive& Ar, FMovieSceneTangentData& P)
-	{
-		P.Serialize(Ar);
-		return Ar;
-	}
-
-	/** If RCIM_Cubic, the arriving tangent at this key */
-	UPROPERTY(EditAnywhere, Category = "Key")
-	float ArriveTangent;
-
-	/** If RCIM_Cubic, the leaving tangent at this key */
-	UPROPERTY(EditAnywhere, Category = "Key")
-	float LeaveTangent;
-
-	/** If RCTWM_WeightedArrive or RCTWM_WeightedBoth, the weight of the left tangent */
-	UPROPERTY(EditAnywhere, Category = "Key")
-	float ArriveTangentWeight;
-
-	/** If RCTWM_WeightedLeave or RCTWM_WeightedBoth, the weight of the right tangent */
-	UPROPERTY(EditAnywhere, Category = "Key")
-	float LeaveTangentWeight;
-
-	/** If RCIM_Cubic, the tangent weight mode */
-	UPROPERTY(EditAnywhere, Category = "Key")
-	TEnumAsByte<ERichCurveTangentWeightMode> TangentWeightMode;
-
-};
+template<typename> struct TMovieSceneCurveChannelImpl;
 
-template<>
-struct TIsPODType<FMovieSceneTangentData>
-{
-	enum { Value = true };
-};
-
-
-template<>
-struct TStructOpsTypeTraits<FMovieSceneTangentData>
-	: public TStructOpsTypeTraitsBase2<FMovieSceneTangentData>
-{
-	enum
-	{
-		WithSerializer = true,
-		WithCopy = false,
-		WithIdenticalViaEquality = true,
-	};
-};
 
 USTRUCT()
 struct FMovieSceneFloatValue
@@ -112,6 +51,12 @@ struct FMovieSceneFloatValue
 	UPROPERTY(EditAnywhere, Category = "Key")
 	TEnumAsByte<ERichCurveTangentMode> TangentMode;
 
+	/**
+	 * float value = 4 bytes
+	 * tangent data = 4 floats + byte enum = 4*4 + 1 = 17 bytes, rounded up to 20 bytes on clang-win64
+	 * interp and tangent modes = 2 byte enums = 2 bytes
+	 * total = 26 bytes
+	 */
 	UPROPERTY()
 	uint8 PaddingByte;
 };
@@ -136,11 +81,15 @@ struct TStructOpsTypeTraits<FMovieSceneFloatValue>
 	};
 };
 
+
 USTRUCT()
 struct MOVIESCENE_API FMovieSceneFloatChannel : public FMovieSceneChannel
 {
 	GENERATED_BODY()
 
+	typedef float CurveValueType;
+	typedef FMovieSceneFloatValue ChannelValueType;
+
 	FMovieSceneFloatChannel() 
 		: PreInfinityExtrap(RCCE_Constant)
 		, PostInfinityExtrap(RCCE_Constant)
@@ -187,9 +136,6 @@ struct MOVIESCENE_API FMovieSceneFloatChannel : public FMovieSceneChannel
 		return Values;
 	}
 
-
-
-
 	/**
 	* Evaluate this channel with the frame resolution 
 	*
@@ -200,24 +146,10 @@ struct MOVIESCENE_API FMovieSceneFloatChannel : public FMovieSceneChannel
 	*/
 	bool Evaluate(FFrameTime InTime, float& OutValue) const;
 
-
-
 	/**
 	 * Set the channel's times and values to the requested values
 	 */
-	FORCEINLINE void Set(TArray<FFrameNumber> InTimes, TArray<FMovieSceneFloatValue> InValues)
-	{
-		check(InTimes.Num() == InValues.Num());
-
-		Times = MoveTemp(InTimes);
-		Values = MoveTemp(InValues);
-
-		KeyHandles.Reset();
-		for (int32 Index = 0; Index < Times.Num(); ++Index)
-		{
-			KeyHandles.AllocateHandle(Index);
-		}
-	}
+	void Set(TArray<FFrameNumber> InTimes, TArray<FMovieSceneFloatValue> InValues);
 
 public:
 
@@ -288,7 +220,6 @@ public:
 		return Ar;
 	}
 
-
 	/** Serialize this float function from a mismatching property tag (FRichCurve) */
 	bool SerializeFromMismatchedTag(const FPropertyTag& Tag, FStructuredArchive::FSlot Slot);
 
@@ -300,6 +231,11 @@ public:
 
 	void AutoSetTangents(float Tension = 0.f);
 
+	/** Get the channel's frame resolution */
+	FFrameRate GetTickResolution() const { return TickResolution; }
+	/** Set the channel's frame resolution */
+	void SetTickResolution(FFrameRate InTickSolution) { TickResolution = InTickSolution; }
+
 	/**
 	 * Populate the specified array with times and values that represent the smooth interpolation of this channel across the specified range
 	 *
@@ -339,27 +275,6 @@ public:
 
 private:
 
-	int32 InsertKeyInternal(FFrameNumber InTime);
-
-	/**
-	 * Evaluate this channel's extrapolation. Assumes more than 1 key is present.
-	 *
-	 * @param InTime     The time to evaluate at
-	 * @param OutValue   A value to receive the result
-	 * @return true if the time was evaluated with extrapolation, false otherwise
-	 */
-	bool EvaluateExtrapolation(FFrameTime InTime, float& OutValue) const;
-
-	/**
-	 * Adds median points between each of the supplied points if their evaluated value is significantly different than the linear interpolation of those points
-	 *
-	 * @param TickResolution        The tick resolution with which to interpret this channel's times
-	 * @param TimeThreshold         A small time threshold in seconds below which we should stop adding new points
-	 * @param ValueThreshold        A small value threshold below which we should stop adding new points where the linear interpolation would suffice
-	 * @param InOutPoints           An array to populate with the evaluated points
-	 */
-	void RefineCurvePoints(FFrameRate TickResolution, double TimeThreshold, float ValueThreshold, TArray<TTuple<double, double>>& InOutPoints) const;
-
 	UPROPERTY(meta=(KeyTimes))
 	TArray<FFrameNumber> Times;
 
@@ -386,12 +301,8 @@ private:
 
 #endif
 
-public:
-	//Set it's frame resolution
-	void SetTickResolution(FFrameRate InTickSolution)
-	{
-		TickResolution = InTickSolution;
-	}
+	friend struct TMovieSceneCurveChannelImpl<FMovieSceneFloatChannel>;
+	using FMovieSceneFloatChannelImpl = TMovieSceneCurveChannelImpl<FMovieSceneFloatChannel>;
 };
 
 template<>
@@ -404,7 +315,7 @@ struct TStructOpsTypeTraits<FMovieSceneFloatChannel> : public TStructOpsTypeTrai
 #if WITH_EDITORONLY_DATA
 		WithPostSerialize = true,
 #endif
-    };
+	};
 };
 
 template<>
@@ -418,30 +329,6 @@ struct TMovieSceneChannelTraits<FMovieSceneFloatChannel> : TMovieSceneChannelTra
 #endif
 };
 
-inline bool ValueExistsAtTime(const FMovieSceneFloatChannel* Channel, FFrameNumber InFrameNumber, float Value)
-{
-	const FFrameTime FrameTime(InFrameNumber);
-
-	float ExistingValue = 0.f;
-	return Channel->Evaluate(FrameTime, ExistingValue) && FMath::IsNearlyEqual(ExistingValue, Value, KINDA_SMALL_NUMBER);
-}
-
-inline bool ValueExistsAtTime(const FMovieSceneFloatChannel* Channel, FFrameNumber InFrameNumber, const FMovieSceneFloatValue& InValue)
-{
-	return ValueExistsAtTime(Channel, InFrameNumber, InValue.Value);
-}
-
-inline void AssignValue(FMovieSceneFloatChannel* InChannel, FKeyHandle InKeyHandle, float InValue)
-{
-	TMovieSceneChannelData<FMovieSceneFloatValue> ChannelData = InChannel->GetData();
-	int32 ValueIndex = ChannelData.GetIndex(InKeyHandle);
-
-	if (ValueIndex != INDEX_NONE)
-	{
-		ChannelData.GetValues()[ValueIndex].Value = InValue;
-	}
-}
-
 /**
  * Overload for adding a new key to a float channel at a given time. See UE::MovieScene::AddKeyToChannel for default implementation.
  */
@@ -453,3 +340,11 @@ MOVIESCENE_API FKeyHandle AddKeyToChannel(FMovieSceneFloatChannel* Channel, FFra
  */
 MOVIESCENE_API void Dilate(FMovieSceneFloatChannel* InChannel, FFrameNumber Origin, float DilationFactor);
 
+
+/**
+ * Overloads for common utility functions.
+ */
+MOVIESCENE_API bool ValueExistsAtTime(const FMovieSceneFloatChannel* InChannel, FFrameNumber InFrameNumber, float InValue);
+MOVIESCENE_API bool ValueExistsAtTime(const FMovieSceneFloatChannel* InChannel, FFrameNumber InFrameNumber, const FMovieSceneFloatValue& InValue);
+MOVIESCENE_API void AssignValue(FMovieSceneFloatChannel* InChannel, FKeyHandle InKeyHandle, float InValue);
+
diff --git a/Engine/Source/Runtime/MovieScene/Public/EntitySystem/BuiltInComponentTypes.h b/Engine/Source/Runtime/MovieScene/Public/EntitySystem/BuiltInComponentTypes.h
--- a/Engine/Source/Runtime/MovieScene/Public/EntitySystem/BuiltInComponentTypes.h
+++ b/Engine/Source/Runtime/MovieScene/Public/EntitySystem/BuiltInComponentTypes.h
@@ -20,6 +20,7 @@ enum class EMovieSceneBlendType : uint8;
 struct FMovieSceneByteChannel;
 struct FMovieSceneIntegerChannel;
 struct FMovieSceneFloatChannel;
+struct FMovieSceneDoubleChannel;
 struct FMovieScenePropertyBinding;
 
 class UMovieSceneSection;
@@ -138,12 +139,32 @@ struct FSourceFloatChannel
 	const FMovieSceneFloatChannel* Source;
 };
 
-
 struct FSourceFloatChannelFlags
 {
 	bool bNeedsEvaluate = true;
 };
 
+/**
+ * The component data for evaluation a double channel
+ */
+struct FSourceDoubleChannel
+{
+	FSourceDoubleChannel()
+		: Source(nullptr)
+	{}
+
+	FSourceDoubleChannel(const FMovieSceneDoubleChannel* InSource)
+		: Source(InSource)
+	{}
+
+	const FMovieSceneDoubleChannel* Source;
+};
+
+struct FSourceDoubleChannelFlags
+{
+	bool bNeedsEvaluate = true;
+};
+
 struct FEvaluationHookFlags
 {
 	bool bHasBegun = false;
@@ -219,6 +240,10 @@ public:
 	TComponentTypeID<FSourceFloatChannel> FloatChannel[9];
 	TComponentTypeID<FSourceFloatChannelFlags> FloatChannelFlags[9];
 
+	// An FMovieSceneDoubleChannel considered to be at index N within the source structure (ie 0 = Location.X, Vector.X; 1 = Location.Y, Vector.Y)
+	TComponentTypeID<FSourceDoubleChannel> DoubleChannel[9];
+	TComponentTypeID<FSourceDoubleChannelFlags> DoubleChannelFlags[9];
+
 	// An FMovieSceneFloatChannel that represents an arbitrary weight
 	TComponentTypeID<FSourceFloatChannel> WeightChannel;
 	TComponentTypeID<FSourceFloatChannelFlags> WeightChannelFlags;
@@ -226,10 +251,16 @@ public:
 	// A float representing the output of the channel considered to be at index N within the source structure (ie 0 = Location.X, Vector.X, Color.R; 1 = Location.Y, Vector.Y, Color.G)
 	TComponentTypeID<float> FloatResult[9];
 
+	// A double considered to be at index N within the source structure (ie 0 = Location.X, Vector.X; 1 = Location.Y, Vector.Y)
+	TComponentTypeID<double> DoubleResult[9];
+
 	// A float representing the base value for the float channel at index N, for the purposes of "additive from base" blending.
 	TComponentTypeID<float> BaseFloat[9];
 
-	// The time at which to evaluate a base value, such as BaseFloat[].
+	// A double representing the base value for the double channel at index N, for the purposes of "additive from base" blending.
+	TComponentTypeID<double> BaseDouble[9];
+
+	// The time at which to evaluate a base value, such as BaseFloat[] or BaseDouble[].
 	TComponentTypeID<FFrameTime> BaseValueEvalTime;
 
 	// A float representing the evaluated output of a weight channel
diff --git a/Engine/Source/Runtime/MovieScene/Public/EntitySystem/IMovieScenePropertyComponentHandler.h b/Engine/Source/Runtime/MovieScene/Public/EntitySystem/IMovieScenePropertyComponentHandler.h
--- a/Engine/Source/Runtime/MovieScene/Public/EntitySystem/IMovieScenePropertyComponentHandler.h
+++ b/Engine/Source/Runtime/MovieScene/Public/EntitySystem/IMovieScenePropertyComponentHandler.h
@@ -20,7 +20,7 @@ struct FPropertyDefinition;
 struct IPreAnimatedStorage;
 struct FSystemSubsequentTasks;
 struct FSystemTaskPrerequisites;
-struct FFloatDecompositionParams;
+struct FValueDecompositionParams;
 struct FPreAnimatedStateExtension;
 struct FPropertyCompositeDefinition;
 
@@ -187,7 +187,7 @@ struct IPropertyComponentHandler
 	 * @param InCurrentValue   The current value (of type StorageType) to recompose using. For instance, if a property comprises 3 additive values (a:1, b:2, c:3), and we recompose 'a' with an InCurrentValue of 10, the result for 'a' would be 5.
 	 * @param OutResult        The result to receieve recomposed values, one for every entitiy in Params.Query.Entities. Must be of type StorageType.
 	 */
-	virtual void RecomposeBlendOperational(const FPropertyDefinition& Definition, TArrayView<const FPropertyCompositeDefinition> Composites, const FFloatDecompositionParams& Params, UMovieSceneBlenderSystem* Blender, FConstPropertyComponentView InCurrentValue, FPropertyComponentArrayView OutResult) = 0;
+	virtual void RecomposeBlendOperational(const FPropertyDefinition& Definition, TArrayView<const FPropertyCompositeDefinition> Composites, const FValueDecompositionParams& Params, UMovieSceneBlenderSystem* Blender, FConstPropertyComponentView InCurrentValue, FPropertyComponentArrayView OutResult) = 0;
 
 	/**
 	 * Run a recomposition using the specified params and values.
@@ -199,7 +199,7 @@ struct IPropertyComponentHandler
 	 * @param InCurrentValue   The current value (of type StorageType) to recompose using. For instance, if a property comprises 3 additive values (a:1, b:2, c:3), and we recompose 'a' with an InCurrentValue of 10, the result for 'a' would be 5.
 	 * @param OutResults       The result to receieve recomposed values, one for every entitiy in Params.Query.Entities.
 	 */
-	virtual void RecomposeBlendChannel(const FPropertyDefinition& Definition, const FPropertyCompositeDefinition& Composite, const FFloatDecompositionParams& Params, UMovieSceneBlenderSystem* Blender, float InCurrentValue, TArrayView<float> OutResults) = 0;
+	virtual void RecomposeBlendChannel(const FPropertyDefinition& Definition, TArrayView<const FPropertyCompositeDefinition> Composites, int32 CompositeIndex, const FValueDecompositionParams& Params, UMovieSceneBlenderSystem* Blender, double InCurrentValue, TArrayView<double> OutResults) = 0;
 
 	/**
 	 * Rebuild operational values from the given entities. These entities are expected to store the value type's composite values.
diff --git a/Engine/Source/Runtime/MovieScene/Public/EntitySystem/MovieSceneDecompositionQuery.h b/Engine/Source/Runtime/MovieScene/Public/EntitySystem/MovieSceneDecompositionQuery.h
--- a/Engine/Source/Runtime/MovieScene/Public/EntitySystem/MovieSceneDecompositionQuery.h
+++ b/Engine/Source/Runtime/MovieScene/Public/EntitySystem/MovieSceneDecompositionQuery.h
@@ -37,56 +37,57 @@ namespace MovieScene
 
 
 	/** Used for decomposing how a final blended value was blended */
-	struct FWeightedFloat
+	struct FWeightedValue
 	{
-		float Value = 0.f;
+		double Value = 0.f;
 		float Weight = 0.f;
 
-		float WeightedValue() const
+		double WeightedValue() const
 		{
 			return Weight != 0.f ? Value / Weight : 0.f;
 		}
 
-		FWeightedFloat Combine(FWeightedFloat Other) const
+		FWeightedValue Combine(FWeightedValue Other) const
 		{
-			return FWeightedFloat{Value + Other.Value, Weight + Other.Weight};
+			return FWeightedValue{Value + Other.Value, Weight + Other.Weight};
 		}
 
-		FWeightedFloat CombineWeighted(FWeightedFloat Other) const
+		FWeightedValue CombineWeighted(FWeightedValue Other) const
 		{
-			return FWeightedFloat{Value + Other.Value * Other.Weight, Weight + Other.Weight};
+			return FWeightedValue{Value + Other.Value * Other.Weight, Weight + Other.Weight};
 		}
 	};
 
-	struct FDecomposedFloat
+	struct FDecomposedValue
 	{
 		struct FResult
 		{
-			FWeightedFloat Absolute;
-			float Additive = 0.f;
+			FWeightedValue Absolute;
+			double Additive = 0.f;
 		};
 
 		FResult Result;
 
-		TArray<TTuple<FMovieSceneEntityID, FWeightedFloat>> DecomposedAbsolutes;
-		TArray<TTuple<FMovieSceneEntityID, FWeightedFloat>> DecomposedAdditives;
+		TArray<TTuple<FMovieSceneEntityID, FWeightedValue>> DecomposedAbsolutes;
+		TArray<TTuple<FMovieSceneEntityID, FWeightedValue>> DecomposedAdditives;
 
 		MOVIESCENE_API float Recompose(FMovieSceneEntityID EntityID, float CurrentValue, const float* InitialValue) const;
-		MOVIESCENE_API void Decompose(FMovieSceneEntityID EntityID, FWeightedFloat& ThisValue, bool& bOutIsAdditive, FWeightedFloat& Absolutes, FWeightedFloat& Additives) const;
+		MOVIESCENE_API double Recompose(FMovieSceneEntityID EntityID, double CurrentValue, const double* InitialValue) const;
+		MOVIESCENE_API void Decompose(FMovieSceneEntityID EntityID, FWeightedValue& ThisValue, bool& bOutIsAdditive, FWeightedValue& Absolutes, FWeightedValue& Additives) const;
 	};
 
 	// Align results to cache lines so there's no contention between cores
-	struct MS_ALIGN(PLATFORM_CACHE_LINE_SIZE) FAlignedDecomposedFloat
+	struct MS_ALIGN(PLATFORM_CACHE_LINE_SIZE) FAlignedDecomposedValue
 	{
-		FDecomposedFloat Value;
+		FDecomposedValue Value;
 	} GCC_ALIGN(PLATFORM_CACHE_LINE_SIZE);
 
-	struct FFloatDecompositionParams
+	struct FValueDecompositionParams
 	{
 		FDecompositionQuery Query;
 		uint16 DecomposeBlendChannel;
 		FMovieSceneEntityID PropertyEntityID;
-		TComponentTypeID<float> ResultComponentType;
+		FComponentTypeID ResultComponentType;
 		FComponentTypeID PropertyTag;
 	};
 
@@ -109,16 +110,16 @@ namespace MovieScene
 
 
 UINTERFACE()
-class MOVIESCENE_API UMovieSceneFloatDecomposer : public UInterface
+class MOVIESCENE_API UMovieSceneValueDecomposer : public UInterface
 {
 public:
 	GENERATED_BODY()
 };
 
-class IMovieSceneFloatDecomposer
+class IMovieSceneValueDecomposer
 {
 public:
 	GENERATED_BODY()
 
-	virtual FGraphEventRef DispatchDecomposeTask(const UE::MovieScene::FFloatDecompositionParams& Params, UE::MovieScene::FAlignedDecomposedFloat* Output) = 0;
+	virtual FGraphEventRef DispatchDecomposeTask(const UE::MovieScene::FValueDecompositionParams& Params, UE::MovieScene::FAlignedDecomposedValue* Output) = 0;
 };
diff --git a/Engine/Source/Runtime/MovieScene/Public/EntitySystem/MovieScenePropertyComponentHandler.h b/Engine/Source/Runtime/MovieScene/Public/EntitySystem/MovieScenePropertyComponentHandler.h
--- a/Engine/Source/Runtime/MovieScene/Public/EntitySystem/MovieScenePropertyComponentHandler.h
+++ b/Engine/Source/Runtime/MovieScene/Public/EntitySystem/MovieScenePropertyComponentHandler.h
@@ -74,7 +74,7 @@ struct TInitialValueProcessorImpl<PropertyTraits, TPropertyMetaData<MetaDataType
 
 		if (PropertyDefinition->CustomPropertyRegistration)
 		{
-		CustomAccessors = PropertyDefinition->CustomPropertyRegistration->GetAccessors();
+			CustomAccessors = PropertyDefinition->CustomPropertyRegistration->GetAccessors();
 		}
 
 		if (InitialValueCache)
@@ -378,17 +378,17 @@ struct TPropertyComponentHandlerImpl<PropertyTraits, TPropertyMetaData<MetaDataT
 		return &Processor;
 	}
 
-	virtual void RecomposeBlendOperational(const FPropertyDefinition& PropertyDefinition, TArrayView<const FPropertyCompositeDefinition> Composites, const FFloatDecompositionParams& InParams, UMovieSceneBlenderSystem* Blender, FConstPropertyComponentView InCurrentValue, FPropertyComponentArrayView OutResult) override
+	virtual void RecomposeBlendOperational(const FPropertyDefinition& PropertyDefinition, TArrayView<const FPropertyCompositeDefinition> Composites, const FValueDecompositionParams& InParams, UMovieSceneBlenderSystem* Blender, FConstPropertyComponentView InCurrentValue, FPropertyComponentArrayView OutResult) override
 	{
 		RecomposeBlendImpl(PropertyDefinition, Composites, InParams, Blender, InCurrentValue.ReinterpretCast<StorageType>(), OutResult.ReinterpretCast<StorageType>());
 	}
 
-	void RecomposeBlendImpl(const FPropertyDefinition& PropertyDefinition, TArrayView<const FPropertyCompositeDefinition> Composites, const FFloatDecompositionParams& InParams, UMovieSceneBlenderSystem* Blender, const StorageType& InCurrentValue, TArrayView<StorageType> OutResults)
+	void RecomposeBlendImpl(const FPropertyDefinition& PropertyDefinition, TArrayView<const FPropertyCompositeDefinition> Composites, const FValueDecompositionParams& InParams, UMovieSceneBlenderSystem* Blender, const StorageType& InCurrentValue, TArrayView<StorageType> OutResults)
 	{
 		check(OutResults.Num() == InParams.Query.Entities.Num());
 
-		IMovieSceneFloatDecomposer* FloatDecomposer = Cast<IMovieSceneFloatDecomposer>(Blender);
-		if (!FloatDecomposer)
+		IMovieSceneValueDecomposer* ValueDecomposer = Cast<IMovieSceneValueDecomposer>(Blender);
+		if (!ValueDecomposer)
 		{
 			return;
 		}
@@ -399,20 +399,20 @@ struct TPropertyComponentHandlerImpl<PropertyTraits, TPropertyMetaData<MetaDataT
 		constexpr int32 NumComposites = sizeof...(CompositeTypes);
 		check(Composites.Num() == NumComposites);
 
-		FAlignedDecomposedFloat AlignedOutputs[NumComposites];
+		FAlignedDecomposedValue AlignedOutputs[NumComposites];
 
-		FFloatDecompositionParams LocalParams = InParams;
+		FValueDecompositionParams LocalParams = InParams;
 
 		FGraphEventArray Tasks;
 		for (int32 Index = 0; Index < NumComposites; ++Index)
 		{
-			if ((PropertyDefinition.FloatCompositeMask & (1 << Index)) == 0)
+			if ((PropertyDefinition.FloatCompositeMask & (1 << Index)) == 0 && (PropertyDefinition.DoubleCompositeMask & (1 << Index)) == 0)
 			{
 				continue;
 			}
 
-			LocalParams.ResultComponentType = Composites[Index].ComponentTypeID.ReinterpretCast<float>();
-			FGraphEventRef Task = FloatDecomposer->DispatchDecomposeTask(LocalParams, &AlignedOutputs[Index]);
+			LocalParams.ResultComponentType = Composites[Index].ComponentTypeID;
+			FGraphEventRef Task = ValueDecomposer->DispatchDecomposeTask(LocalParams, &AlignedOutputs[Index]);
 			if (Task)
 			{
 				Tasks.Add(Task);
@@ -440,35 +440,52 @@ struct TPropertyComponentHandlerImpl<PropertyTraits, TPropertyMetaData<MetaDataT
 
 			for (int32 CompositeIndex = 0; CompositeIndex < NumComposites; ++CompositeIndex)
 			{
-				if ((PropertyDefinition.FloatCompositeMask & (1 << CompositeIndex)) == 0)
+				if ((PropertyDefinition.FloatCompositeMask & (1 << CompositeIndex)) != 0)
 				{
-					continue;
+					const float* InitialValueComposite = nullptr;
+					FAlignedDecomposedValue& AlignedOutput = AlignedOutputs[CompositeIndex];
+					if (InitialValueComponent)
+					{
+						const StorageType& InitialValue = (*InitialValueComponent);
+						InitialValueComposite = reinterpret_cast<const float*>(reinterpret_cast<const uint8*>(&InitialValue) + Composites[CompositeIndex].CompositeOffset);
+					}
+
+					const float NewComposite = *reinterpret_cast<const float*>(reinterpret_cast<const uint8*>(&InCurrentValue) + Composites[CompositeIndex].CompositeOffset);
+
+					float* RecomposedComposite = reinterpret_cast<float*>(Result + Composites[CompositeIndex].CompositeOffset);
+					*RecomposedComposite = AlignedOutput.Value.Recompose(EntityID, NewComposite, InitialValueComposite);
 				}
-
-				const float* InitialValueComposite = nullptr;
-				FAlignedDecomposedFloat& AlignedOutput = AlignedOutputs[CompositeIndex];
-				if (InitialValueComponent)
+				else if ((PropertyDefinition.DoubleCompositeMask & ( 1 << Index)) != 0)
 				{
-					const StorageType& InitialValue = (*InitialValueComponent);
-					InitialValueComposite = reinterpret_cast<const float*>(reinterpret_cast<const uint8*>(&InitialValue) + Composites[CompositeIndex].CompositeOffset);
+					const double* InitialValueComposite = nullptr;
+					FAlignedDecomposedValue& AlignedOutput = AlignedOutputs[CompositeIndex];
+					if (InitialValueComponent)
+					{
+						const StorageType& InitialValue = (*InitialValueComponent);
+						InitialValueComposite = reinterpret_cast<const double*>(reinterpret_cast<const uint8*>(&InitialValue) + Composites[CompositeIndex].CompositeOffset);
+					}
+
+					const double NewComposite = *reinterpret_cast<const double*>(reinterpret_cast<const uint8*>(&InCurrentValue) + Composites[CompositeIndex].CompositeOffset);
+
+					double* RecomposedComposite = reinterpret_cast<double*>(Result + Composites[CompositeIndex].CompositeOffset);
+					*RecomposedComposite = AlignedOutput.Value.Recompose(EntityID, NewComposite, InitialValueComposite);
 				}
-
-				const float NewComposite = *reinterpret_cast<const float*>(reinterpret_cast<const uint8*>(&InCurrentValue) + Composites[CompositeIndex].CompositeOffset);
-
-				float* RecomposedComposite = reinterpret_cast<float*>(Result + Composites[CompositeIndex].CompositeOffset);
-				*RecomposedComposite = AlignedOutput.Value.Recompose(EntityID, NewComposite, InitialValueComposite);
 			}
 		}
 
 		EntityManager.ReleaseLockDown();
 	}
 
-	virtual void RecomposeBlendChannel(const FPropertyDefinition& PropertyDefinition, const FPropertyCompositeDefinition& Composite, const FFloatDecompositionParams& InParams, UMovieSceneBlenderSystem* Blender, float InCurrentValue, TArrayView<float> OutResults) override
+	virtual void RecomposeBlendChannel(const FPropertyDefinition& PropertyDefinition, TArrayView<const FPropertyCompositeDefinition> Composites, int32 CompositeIndex, const FValueDecompositionParams& InParams, UMovieSceneBlenderSystem* Blender, double InCurrentValue, TArrayView<double> OutResults) override
 	{
 		check(OutResults.Num() == InParams.Query.Entities.Num());
 
-		IMovieSceneFloatDecomposer* FloatDecomposer = Cast<IMovieSceneFloatDecomposer>(Blender);
-		if (!FloatDecomposer)
+		constexpr int32 NumComposites = sizeof...(CompositeTypes);
+		check(Composites.Num() == NumComposites);
+		const FPropertyCompositeDefinition& Composite = Composites[CompositeIndex];
+
+		IMovieSceneValueDecomposer* ValueDecomposer = Cast<IMovieSceneValueDecomposer>(Blender);
+		if (!ValueDecomposer)
 		{
 			return;
 		}
@@ -476,12 +493,12 @@ struct TPropertyComponentHandlerImpl<PropertyTraits, TPropertyMetaData<MetaDataT
 		FEntityManager& EntityManager = Blender->GetLinker()->EntityManager;
 		EntityManager.LockDown();
 
-		FAlignedDecomposedFloat AlignedOutput;
+		FAlignedDecomposedValue AlignedOutput;
 
-		FFloatDecompositionParams LocalParams = InParams;
+		FValueDecompositionParams LocalParams = InParams;
 
-		LocalParams.ResultComponentType = Composite.ComponentTypeID.ReinterpretCast<float>();
-		FGraphEventRef Task = FloatDecomposer->DispatchDecomposeTask(LocalParams, &AlignedOutput);
+		LocalParams.ResultComponentType = Composite.ComponentTypeID;
+		FGraphEventRef Task = ValueDecomposer->DispatchDecomposeTask(LocalParams, &AlignedOutput);
 		if (Task)
 		{
 			FTaskGraphInterface::Get().WaitUntilTaskCompletes(Task, ENamedThreads::GameThread);
@@ -499,15 +516,30 @@ struct TPropertyComponentHandlerImpl<PropertyTraits, TPropertyMetaData<MetaDataT
 		{
 			FMovieSceneEntityID EntityID = LocalParams.Query.Entities[Index];
 
-			const float* InitialValueComposite = nullptr;
-			if (InitialValueComponent)
+			if ((PropertyDefinition.FloatCompositeMask & (1 << CompositeIndex)) != 0)
 			{
-				const StorageType& InitialValue = (*InitialValueComponent);
-				InitialValueComposite = reinterpret_cast<const float*>(reinterpret_cast<const uint8*>(&InitialValue) + Composite.CompositeOffset);
+				const float* InitialValueComposite = nullptr;
+				if (InitialValueComponent)
+				{
+					const StorageType& InitialValue = (*InitialValueComponent);
+					InitialValueComposite = reinterpret_cast<const float*>(reinterpret_cast<const uint8*>(&InitialValue) + Composite.CompositeOffset);
+				}
+
+				const float RecomposedComposite = AlignedOutput.Value.Recompose(EntityID, InCurrentValue, InitialValueComposite);
+				OutResults[Index] = (double)RecomposedComposite;
 			}
+			else if ((PropertyDefinition.DoubleCompositeMask & (1 << CompositeIndex)) != 0)
+			{
+				const double* InitialValueComposite = nullptr;
+				if (InitialValueComponent)
+				{
+					const StorageType& InitialValue = (*InitialValueComponent);
+					InitialValueComposite = reinterpret_cast<const double*>(reinterpret_cast<const uint8*>(&InitialValue) + Composite.CompositeOffset);
+				}
 
-			float* RecomposedComposite = &OutResults[Index];
-			*RecomposedComposite = AlignedOutput.Value.Recompose(EntityID, InCurrentValue, InitialValueComposite);
+				const double RecomposedComposite = AlignedOutput.Value.Recompose(EntityID, InCurrentValue, InitialValueComposite);
+				OutResults[Index] = RecomposedComposite;
+			}
 		}
 
 		EntityManager.ReleaseLockDown();
@@ -560,6 +592,10 @@ struct TPropertyDefinitionBuilder
 		{
 			Definition->FloatCompositeMask = 1;
 		}
+		else if (TIsSame<typename PropertyTraits::StorageType, double>::Value)
+		{
+			Definition->DoubleCompositeMask = 1;
+		}
 
 		return *this;
 	}
@@ -571,6 +607,19 @@ struct TPropertyDefinitionBuilder
 		return *this;
 	}
 
+	template<typename BlenderSystemType>
+	TPropertyDefinitionBuilder<PropertyTraits>& SetBlenderSystem()
+	{
+		Definition->BlenderSystemClass = BlenderSystemType::StaticClass();
+		return *this;
+	}
+
+	TPropertyDefinitionBuilder<PropertyTraits>& SetBlenderSystem(UClass* BlenderSystemClass)
+	{
+		Definition->BlenderSystemClass = BlenderSystemClass;
+		return *this;
+	}
+
 	void Commit()
 	{
 		Definition->Handler = TPropertyComponentHandler<PropertyTraits, typename PropertyTraits::StorageType>();
@@ -620,6 +669,10 @@ struct TCompositePropertyDefinitionBuilder
 		{
 			Definition->FloatCompositeMask |= 1 << Definition->CompositeSize;
 		}
+		else if (TIsSame<T, double>::Value)
+		{
+			Definition->DoubleCompositeMask |= 1 << Definition->CompositeSize;
+		}
 
 		++Definition->CompositeSize;
 		return TCompositePropertyDefinitionBuilder<PropertyTraits, Composites..., T>(Definition, Registry);
@@ -640,6 +693,21 @@ struct TCompositePropertyDefinitionBuilder
 		return TCompositePropertyDefinitionBuilder<PropertyTraits, Composites..., float>(Definition, Registry);
 	}
 
+	TCompositePropertyDefinitionBuilder<PropertyTraits, Composites..., double> AddComposite(TComponentTypeID<double> InComponent, double StorageType::*DataPtr)
+	{
+		checkf(Definition == &Registry->GetProperties().Last(), TEXT("Cannot re-define a property type after another has been added."));
+
+		const PTRINT CompositeOffset = (PTRINT)&(((StorageType*)0)->*DataPtr);
+
+		FPropertyCompositeDefinition NewChannel = { InComponent, static_cast<uint16>(CompositeOffset) };
+		Registry->CompositeDefinitions.Add(NewChannel);
+
+		Definition->DoubleCompositeMask |= 1 << Definition->CompositeSize;
+
+		++Definition->CompositeSize;
+		return TCompositePropertyDefinitionBuilder<PropertyTraits, Composites..., double>(Definition, Registry);
+	}
+
 	template<int InlineSize>
 	TCompositePropertyDefinitionBuilder<PropertyTraits, Composites...>& SetCustomAccessors(TCustomPropertyRegistration<PropertyTraits, InlineSize>* InCustomAccessors)
 	{
@@ -647,6 +715,19 @@ struct TCompositePropertyDefinitionBuilder
 		return *this;
 	}
 
+	template<typename BlenderSystemType>
+	TCompositePropertyDefinitionBuilder<PropertyTraits, Composites...>& SetBlenderSystem()
+	{
+		Definition->BlenderSystemClass = BlenderSystemType::StaticClass();
+		return *this;
+	}
+
+	TCompositePropertyDefinitionBuilder<PropertyTraits, Composites...>& SetBlenderSystem(UClass* BlenderSystemClass)
+	{
+		Definition->BlenderSystemClass = BlenderSystemClass;
+		return *this;
+	}
+
 	void Commit()
 	{
 		Definition->Handler = TPropertyComponentHandler<PropertyTraits, Composites...>();
@@ -717,7 +798,7 @@ TRecompositionResult<typename PropertyTraits::StorageType> FPropertyRecomposerIm
 		return Result;
 	}
 
-	FFloatDecompositionParams Params;
+	FValueDecompositionParams Params;
 	Params.Query = InQuery;
 	Params.PropertyEntityID = Property.PropertyEntityID;
 	Params.DecomposeBlendChannel = Property.BlendChannel;
diff --git a/Engine/Source/Runtime/MovieScene/Public/EntitySystem/MovieScenePropertyRegistry.h b/Engine/Source/Runtime/MovieScene/Public/EntitySystem/MovieScenePropertyRegistry.h
--- a/Engine/Source/Runtime/MovieScene/Public/EntitySystem/MovieScenePropertyRegistry.h
+++ b/Engine/Source/Runtime/MovieScene/Public/EntitySystem/MovieScenePropertyRegistry.h
@@ -56,11 +56,13 @@ struct FPropertyDefinition
 			FComponentTypeID InPropertyType, FComponentTypeID InInitialValueType)
 		: CustomPropertyRegistration(nullptr)
 		, FloatCompositeMask(0)
+		, DoubleCompositeMask(0)
 		, VariableSizeCompositeOffset(InVariableSizeCompositeOffset)
 		, CompositeSize(0)
 		, StorageType{ InSizeofStorageType, InAlignofStorageType }
 		, PropertyType(InPropertyType)
 		, InitialValueType(InInitialValueType)
+		, BlenderSystemClass(nullptr)
 	{
 	}
 
@@ -84,6 +86,9 @@ struct FPropertyDefinition
 	/** A mask of which composite indices pertain to floats */
 	uint32 FloatCompositeMask = 0;
 
+	/** A mask of which composite indices pertain to doubles */
+	uint32 DoubleCompositeMask = 0;
+
 	/** The number of channels that this property comprises */
 	uint16 VariableSizeCompositeOffset = INDEX_NONE;
 
@@ -106,6 +111,9 @@ struct FPropertyDefinition
 	/** MetaData types */
 	TArrayView<const FComponentTypeID> MetaDataTypes;
 
+	/** The blender system to use by default (if specified here) to blend composites of this property (it can be overriden per-entity with a blender system component) */
+	UClass* BlenderSystemClass;
+
 	/** Implementation of type specific property actions such as applying properties from entities or recomposing values */
 	TInlineValue<IPropertyComponentHandler, 32> Handler;
 };
diff --git a/Engine/Source/Runtime/MovieScene/Public/INodeAndChannelMappings.h b/Engine/Source/Runtime/MovieScene/Public/INodeAndChannelMappings.h
--- a/Engine/Source/Runtime/MovieScene/Public/INodeAndChannelMappings.h
+++ b/Engine/Source/Runtime/MovieScene/Public/INodeAndChannelMappings.h
@@ -27,6 +27,7 @@ enum class FFBXControlRigTypeProxyEnum : uint8
 	Integer
 };
 
+struct FMovieSceneDoubleChannel;
 struct FMovieSceneFloatChannel;
 struct FMovieSceneBoolChannel;
 struct FMovieSceneIntegerChannel;
@@ -38,7 +39,8 @@ struct FFBXNodeAndChannels
 	UMovieSceneTrack* MovieSceneTrack;
 	FFBXControlRigTypeProxyEnum ControlType;
 	FString NodeName;
-	TArray< FMovieSceneFloatChannel*> Channels;
+	TArray< FMovieSceneDoubleChannel*> DoubleChannels;
+	TArray< FMovieSceneFloatChannel*> FloatChannels;
 	//will really only have one ever.
 	TArray< FMovieSceneBoolChannel*> BoolChannels;
 	TArray< FMovieSceneIntegerChannel*> IntegerChannels;
@@ -62,4 +64,4 @@ class MOVIESCENE_API INodeAndChannelMappings
 	
 	/** Get Selected Nodes */
 	virtual void GetSelectedNodes(TArray<FName>& OutSelectedNodes) = 0;
-};
\ No newline at end of file
+};
diff --git a/Engine/Source/Runtime/MovieSceneTracks/Private/EntitySystem/Interrogation/MovieSceneInterrogatedPropertyInstantiator.cpp b/Engine/Source/Runtime/MovieSceneTracks/Private/EntitySystem/Interrogation/MovieSceneInterrogatedPropertyInstantiator.cpp
--- a/Engine/Source/Runtime/MovieSceneTracks/Private/EntitySystem/Interrogation/MovieSceneInterrogatedPropertyInstantiator.cpp
+++ b/Engine/Source/Runtime/MovieSceneTracks/Private/EntitySystem/Interrogation/MovieSceneInterrogatedPropertyInstantiator.cpp
@@ -49,11 +49,11 @@ UE::MovieScene::FPropertyRecomposerPropertyInfo UMovieSceneInterrogatedPropertyI
 	return FPropertyRecomposerPropertyInfo::Invalid();
 }
 
-UMovieSceneInterrogatedPropertyInstantiatorSystem::FFloatRecompositionResult UMovieSceneInterrogatedPropertyInstantiatorSystem::RecomposeBlendFloatChannel(const UE::MovieScene::FPropertyDefinition& PropertyDefinition, int32 ChannelCompositeIndex, const UE::MovieScene::FDecompositionQuery& InQuery, float InCurrentValue)
+UMovieSceneInterrogatedPropertyInstantiatorSystem::FValueRecompositionResult UMovieSceneInterrogatedPropertyInstantiatorSystem::RecomposeBlendChannel(const UE::MovieScene::FPropertyDefinition& PropertyDefinition, int32 ChannelCompositeIndex, const UE::MovieScene::FDecompositionQuery& InQuery, double InCurrentValue)
 {
 	using namespace UE::MovieScene;
 
-	FFloatRecompositionResult Result(InCurrentValue, InQuery.Entities.Num());
+	FValueRecompositionResult Result(InCurrentValue, InQuery.Entities.Num());
 
 	if (InQuery.Entities.Num() == 0)
 	{
@@ -72,7 +72,7 @@ UMovieSceneInterrogatedPropertyInstantiatorSystem::FFloatRecompositionResult UMo
 		return Result;
 	}
 
-	FFloatDecompositionParams Params;
+	FValueDecompositionParams Params;
 	Params.Query = InQuery;
 	Params.PropertyEntityID = Property->PropertyEntityID;
 	Params.DecomposeBlendChannel = Property->BlendChannel;
@@ -81,7 +81,7 @@ UMovieSceneInterrogatedPropertyInstantiatorSystem::FFloatRecompositionResult UMo
 	TArrayView<const FPropertyCompositeDefinition> Composites = BuiltInComponents->PropertyRegistry.GetComposites(PropertyDefinition);
 	check(Composites.IsValidIndex(ChannelCompositeIndex));
 
-	PropertyDefinition.Handler->RecomposeBlendChannel(PropertyDefinition, Composites[ChannelCompositeIndex], Params, Blender, InCurrentValue, Result.Values);
+	PropertyDefinition.Handler->RecomposeBlendChannel(PropertyDefinition, Composites, ChannelCompositeIndex, Params, Blender, InCurrentValue, Result.Values);
 
 	return Result;
 }
@@ -115,11 +115,11 @@ bool UMovieSceneInterrogatedPropertyInstantiatorSystem::PropertySupportsFastPath
 	return true;
 }
 
-UClass* UMovieSceneInterrogatedPropertyInstantiatorSystem::ResolveBlenderClass(TArrayView<const FMovieSceneEntityID> Inputs) const
+UClass* UMovieSceneInterrogatedPropertyInstantiatorSystem::ResolveBlenderClass(const UE::MovieScene::FPropertyDefinition& PropertyDefinition, TArrayView<const FMovieSceneEntityID> Inputs) const
 {
 	using namespace UE::MovieScene;
 
-	UClass* BlenderClass = UMovieScenePiecewiseFloatBlenderSystem::StaticClass();
+	UClass* BlenderClass = PropertyDefinition.BlenderSystemClass;
 
 	for (FMovieSceneEntityID Input : Inputs)
 	{
@@ -130,8 +130,12 @@ UClass* UMovieSceneInterrogatedPropertyInstantiatorSystem::ResolveBlenderClass(T
 			break;
 		}
 	}
-
-	check(BlenderClass);
+	
+	if (!ensureMsgf(BlenderClass, TEXT("No default blender class specified on property, and no custom blender specified on entities. Falling back to float blender.")))
+	{
+		BlenderClass = UMovieScenePiecewiseFloatBlenderSystem::StaticClass();
+	}
+	
 	return BlenderClass;
 }
 
@@ -176,7 +180,7 @@ void UMovieSceneInterrogatedPropertyInstantiatorSystem::UpdateOutput(UE::MovieSc
 	TArrayView<const FPropertyCompositeDefinition> Composites = BuiltInComponents->PropertyRegistry.GetComposites(*PropertyDefinition);
 
 	// Find the blender class to use
-	UClass* BlenderClass = ResolveBlenderClass(Inputs);
+	UClass* BlenderClass = ResolveBlenderClass(*PropertyDefinition, Inputs);
 
 	UMovieSceneBlenderSystem* ExistingBlender = Output->Blender.Get();
 	if (ExistingBlender && BlenderClass != ExistingBlender->GetClass())
diff --git a/Engine/Source/Runtime/MovieSceneTracks/Private/EntitySystem/Interrogation/MovieSceneInterrogationLinker.cpp b/Engine/Source/Runtime/MovieSceneTracks/Private/EntitySystem/Interrogation/MovieSceneInterrogationLinker.cpp
--- a/Engine/Source/Runtime/MovieSceneTracks/Private/EntitySystem/Interrogation/MovieSceneInterrogationLinker.cpp
+++ b/Engine/Source/Runtime/MovieSceneTracks/Private/EntitySystem/Interrogation/MovieSceneInterrogationLinker.cpp
@@ -321,9 +321,9 @@ void FInterrogationChannels::QueryWorldSpaceTransforms(UMovieSceneEntitySystemLi
 	// Gather fully animated transforms - this code is simple
 	{
 		auto PopulateComplete = [&VariableTransformsByChannel, PredicateBits](FInterrogationKey InterrogationKey,
-			float LocationX, float LocationY, float LocationZ,
-			float RotationX, float RotationY, float RotationZ,
-			float ScaleX, float ScaleY, float ScaleZ)
+			double LocationX, double LocationY, double LocationZ,
+			double RotationX, double RotationY, double RotationZ,
+			double ScaleX, double ScaleY, double ScaleZ)
 		{
 			const int32 ChannelIndex = InterrogationKey.Channel.AsIndex();
 
@@ -340,9 +340,9 @@ void FInterrogationChannels::QueryWorldSpaceTransforms(UMovieSceneEntitySystemLi
 		FEntityTaskBuilder()
 		.Read(Components->Interrogation.OutputKey)
 		.ReadAllOf(
-			Components->FloatResult[0], Components->FloatResult[1], Components->FloatResult[2],
-			Components->FloatResult[3], Components->FloatResult[4], Components->FloatResult[5],
-			Components->FloatResult[6], Components->FloatResult[7], Components->FloatResult[8])
+			Components->DoubleResult[0], Components->DoubleResult[1], Components->DoubleResult[2],
+			Components->DoubleResult[3], Components->DoubleResult[4], Components->DoubleResult[5],
+			Components->DoubleResult[6], Components->DoubleResult[7], Components->DoubleResult[8])
 		.FilterAll({ TracksComponents->ComponentTransform.PropertyTag })
 		.Iterate_PerEntity(&Linker->EntityManager, PopulateComplete);
 	}
@@ -352,7 +352,7 @@ void FInterrogationChannels::QueryWorldSpaceTransforms(UMovieSceneEntitySystemLi
 		FComponentMask CompleteMask;
 		for (int32 Index = 0; Index < 9; ++Index)
 		{
-			CompleteMask.Set(Components->FloatResult[Index]);
+			CompleteMask.Set(Components->DoubleResult[Index]);
 		}
 
 		FEntityComponentFilter Filter;
@@ -365,10 +365,10 @@ void FInterrogationChannels::QueryWorldSpaceTransforms(UMovieSceneEntitySystemLi
 			{
 				const int32 Num = Allocation->Num();
 
-				TOptionalComponentReader<float> FloatComponents[9];
+				TOptionalComponentReader<double> DoubleComponents[9];
 				for (int32 Index = 0; Index < 9; ++Index)
 				{
-					FloatComponents[Index] = Allocation->TryReadComponents(Components->FloatResult[Index]);
+					DoubleComponents[Index] = Allocation->TryReadComponents(Components->DoubleResult[Index]);
 				}
 
 				for (int32 ComponentIndex = 0; ComponentIndex < Num; ++ComponentIndex)
@@ -383,15 +383,15 @@ void FInterrogationChannels::QueryWorldSpaceTransforms(UMovieSceneEntitySystemLi
 						FRotator  Rotation = Transform.GetRotation().Rotator();
 						FVector   Scale    = Transform.GetScale3D();
 
-						if (FloatComponents[0]) { Location.X     = FloatComponents[0][ComponentIndex]; }
-						if (FloatComponents[1]) { Location.Y     = FloatComponents[1][ComponentIndex]; }
-						if (FloatComponents[2]) { Location.Z     = FloatComponents[2][ComponentIndex]; }
-						if (FloatComponents[3]) { Rotation.Roll  = FloatComponents[3][ComponentIndex]; }
-						if (FloatComponents[4]) { Rotation.Pitch = FloatComponents[4][ComponentIndex]; }
-						if (FloatComponents[5]) { Rotation.Yaw   = FloatComponents[5][ComponentIndex]; }
-						if (FloatComponents[6]) { Scale.X        = FloatComponents[6][ComponentIndex]; }
-						if (FloatComponents[7]) { Scale.Y        = FloatComponents[7][ComponentIndex]; }
-						if (FloatComponents[8]) { Scale.Z        = FloatComponents[8][ComponentIndex]; }
+						if (DoubleComponents[0]) { Location.X     = DoubleComponents[0][ComponentIndex]; }
+						if (DoubleComponents[1]) { Location.Y     = DoubleComponents[1][ComponentIndex]; }
+						if (DoubleComponents[2]) { Location.Z     = DoubleComponents[2][ComponentIndex]; }
+						if (DoubleComponents[3]) { Rotation.Roll  = DoubleComponents[3][ComponentIndex]; }
+						if (DoubleComponents[4]) { Rotation.Pitch = DoubleComponents[4][ComponentIndex]; }
+						if (DoubleComponents[5]) { Rotation.Yaw   = DoubleComponents[5][ComponentIndex]; }
+						if (DoubleComponents[6]) { Scale.X        = DoubleComponents[6][ComponentIndex]; }
+						if (DoubleComponents[7]) { Scale.Y        = DoubleComponents[7][ComponentIndex]; }
+						if (DoubleComponents[8]) { Scale.Z        = DoubleComponents[8][ComponentIndex]; }
 
 						VariableTransformsByChannel[ChannelIndex][InterrogationKey.InterrogationIndex] = FTransform(Rotation, Location, Scale);
 					}
@@ -536,9 +536,9 @@ void FInterrogationChannels::QueryLocalSpaceTransforms(UMovieSceneEntitySystemLi
 	// Gather fully animated transforms - this code is simple
 	{
 		auto PopulateComplete = [&OnGetOutputForChannel, &PredicateBits](FInterrogationKey InterrogationKey,
-			float LocationX, float LocationY, float LocationZ,
-			float RotationX, float RotationY, float RotationZ,
-			float ScaleX, float ScaleY, float ScaleZ)
+			double LocationX, double LocationY, double LocationZ,
+			double RotationX, double RotationY, double RotationZ,
+			double ScaleX, double ScaleY, double ScaleZ)
 		{
 			if (PredicateBits[InterrogationKey.Channel.AsIndex()] == true)
 			{
@@ -553,9 +553,9 @@ void FInterrogationChannels::QueryLocalSpaceTransforms(UMovieSceneEntitySystemLi
 		FEntityTaskBuilder()
 		.Read(Components->Interrogation.OutputKey)
 		.ReadAllOf(
-			Components->FloatResult[0], Components->FloatResult[1], Components->FloatResult[2],
-			Components->FloatResult[3], Components->FloatResult[4], Components->FloatResult[5],
-			Components->FloatResult[6], Components->FloatResult[7], Components->FloatResult[8])
+			Components->DoubleResult[0], Components->DoubleResult[1], Components->DoubleResult[2],
+			Components->DoubleResult[3], Components->DoubleResult[4], Components->DoubleResult[5],
+			Components->DoubleResult[6], Components->DoubleResult[7], Components->DoubleResult[8])
 		.FilterAll({ TracksComponents->ComponentTransform.PropertyTag })
 		.Iterate_PerEntity(&Linker->EntityManager, PopulateComplete);
 	}
@@ -565,7 +565,7 @@ void FInterrogationChannels::QueryLocalSpaceTransforms(UMovieSceneEntitySystemLi
 		FComponentMask CompleteMask;
 		for (int32 Index = 0; Index < 9; ++Index)
 		{
-			CompleteMask.Set(Components->FloatResult[Index]);
+			CompleteMask.Set(Components->DoubleResult[Index]);
 		}
 
 		FEntityComponentFilter Filter;
@@ -578,10 +578,10 @@ void FInterrogationChannels::QueryLocalSpaceTransforms(UMovieSceneEntitySystemLi
 			{
 				const int32 Num = Allocation->Num();
 
-				TOptionalComponentReader<float> FloatComponents[9];
+				TOptionalComponentReader<double> DoubleComponents[9];
 				for (int32 Index = 0; Index < 9; ++Index)
 				{
-					FloatComponents[Index] = Allocation->TryReadComponents(Components->FloatResult[Index]);
+					DoubleComponents[Index] = Allocation->TryReadComponents(Components->DoubleResult[Index]);
 				}
 
 				for (int32 ComponentIndex = 0; ComponentIndex < Num; ++ComponentIndex)
@@ -591,15 +591,15 @@ void FInterrogationChannels::QueryLocalSpaceTransforms(UMovieSceneEntitySystemLi
 					{
 						FIntermediate3DTransform Transform = BaseValues[ComponentIndex];
 
-						if (FloatComponents[0]) { Transform.T_X = FloatComponents[0][ComponentIndex]; }
-						if (FloatComponents[1]) { Transform.T_Y = FloatComponents[1][ComponentIndex]; }
-						if (FloatComponents[2]) { Transform.T_Z = FloatComponents[2][ComponentIndex]; }
-						if (FloatComponents[3]) { Transform.R_X = FloatComponents[3][ComponentIndex]; }
-						if (FloatComponents[4]) { Transform.R_Y = FloatComponents[4][ComponentIndex]; }
-						if (FloatComponents[5]) { Transform.R_Z = FloatComponents[5][ComponentIndex]; }
-						if (FloatComponents[6]) { Transform.S_X = FloatComponents[6][ComponentIndex]; }
-						if (FloatComponents[7]) { Transform.S_Y = FloatComponents[7][ComponentIndex]; }
-						if (FloatComponents[8]) { Transform.S_Z = FloatComponents[8][ComponentIndex]; }
+						if (DoubleComponents[0]) { Transform.T_X = DoubleComponents[0][ComponentIndex]; }
+						if (DoubleComponents[1]) { Transform.T_Y = DoubleComponents[1][ComponentIndex]; }
+						if (DoubleComponents[2]) { Transform.T_Z = DoubleComponents[2][ComponentIndex]; }
+						if (DoubleComponents[3]) { Transform.R_X = DoubleComponents[3][ComponentIndex]; }
+						if (DoubleComponents[4]) { Transform.R_Y = DoubleComponents[4][ComponentIndex]; }
+						if (DoubleComponents[5]) { Transform.R_Z = DoubleComponents[5][ComponentIndex]; }
+						if (DoubleComponents[6]) { Transform.S_X = DoubleComponents[6][ComponentIndex]; }
+						if (DoubleComponents[7]) { Transform.S_Y = DoubleComponents[7][ComponentIndex]; }
+						if (DoubleComponents[8]) { Transform.S_Z = DoubleComponents[8][ComponentIndex]; }
 
 						OnGetOutputForChannel(InterrogationKey.Channel)[InterrogationKey.InterrogationIndex] = Transform;
 					}
diff --git a/Engine/Source/Runtime/MovieSceneTracks/Private/MovieSceneTracksComponentTypes.cpp b/Engine/Source/Runtime/MovieSceneTracks/Private/MovieSceneTracksComponentTypes.cpp
--- a/Engine/Source/Runtime/MovieSceneTracks/Private/MovieSceneTracksComponentTypes.cpp
+++ b/Engine/Source/Runtime/MovieSceneTracks/Private/MovieSceneTracksComponentTypes.cpp
@@ -9,6 +9,12 @@
 #include "EntitySystem/MovieSceneEntitySystemLinker.h"
 #include "EntitySystem/MovieSceneComponentRegistry.h"
 #include "EntitySystem/MovieSceneBlenderSystem.h"
+#include "Systems/MovieScenePiecewiseBoolBlenderSystem.h"
+#include "Systems/MovieScenePiecewiseByteBlenderSystem.h"
+#include "Systems/MovieScenePiecewiseEnumBlenderSystem.h"
+#include "Systems/MovieScenePiecewiseIntegerBlenderSystem.h"
+#include "Systems/MovieScenePiecewiseFloatBlenderSystem.h"
+#include "Systems/MovieScenePiecewiseDoubleBlenderSystem.h"
 #include "EntitySystem/MovieScenePropertyComponentHandler.h"
 #include "EntitySystem/MovieSceneEntityFactoryTemplates.h"
 #include "EntitySystem/MovieScenePropertyMetaDataTraits.inl"
@@ -88,34 +94,44 @@ void ConvertOperationalProperty(const FSlateColor& InColor, FIntermediateColor&
 /* ---------------------------------------------------------------------------
  * Vector conversion functions
  * ---------------------------------------------------------------------------*/
-void ConvertOperationalProperty(const FIntermediateVector& InVector, FVector2D& Out)
+void ConvertOperationalProperty(const FFloatIntermediateVector& InVector, FVector2D& Out)
 {
 	Out = FVector2D(InVector.X, InVector.Y);
 }
 
-void ConvertOperationalProperty(const FIntermediateVector& InVector, FVector& Out)
+void ConvertOperationalProperty(const FFloatIntermediateVector& InVector, FVector3f& Out)
 {
-	Out = FVector(InVector.X, InVector.Y, InVector.Z);
+	Out = FVector3f(InVector.X, InVector.Y, InVector.Z);
 }
 
-void ConvertOperationalProperty(const FIntermediateVector& InVector, FVector4& Out)
+void ConvertOperationalProperty(const FFloatIntermediateVector& InVector, FVector4& Out)
 {
 	Out = FVector4(InVector.X, InVector.Y, InVector.Z, InVector.W);
 }
 
-void ConvertOperationalProperty(const FVector2D& In, FIntermediateVector& Out)
+void ConvertOperationalProperty(const FVector2D& In, FFloatIntermediateVector& Out)
 {
-	Out = FIntermediateVector(In.X, In.Y);
+	Out = FFloatIntermediateVector(In.X, In.Y);
 }
 
-void ConvertOperationalProperty(const FVector& In, FIntermediateVector& Out)
+void ConvertOperationalProperty(const FVector3f& In, FFloatIntermediateVector& Out)
 {
-	Out = FIntermediateVector(In.X, In.Y, In.Z);
+	Out = FFloatIntermediateVector(In.X, In.Y, In.Z);
 }
 
-void ConvertOperationalProperty(const FVector4& In, FIntermediateVector& Out)
+void ConvertOperationalProperty(const FVector4& In, FFloatIntermediateVector& Out)
 {
-	Out = FIntermediateVector(In.X, In.Y, In.Z, In.W);
+	Out = FFloatIntermediateVector(In.X, In.Y, In.Z, In.W);
+}
+
+void ConvertOperationalProperty(const FDoubleIntermediateVector& InVector, FVector3d& Out)
+{
+	Out = FVector3d(InVector.X, InVector.Y, InVector.Z);
+}
+
+void ConvertOperationalProperty(const FVector3d& In, FDoubleIntermediateVector& Out)
+{
+	Out = FDoubleIntermediateVector(In.X, In.Y, In.Z);
 }
 
 
@@ -368,7 +384,7 @@ struct FColorHandler : TPropertyComponentHandler<FColorPropertyTraits, float, fl
 };
 
 
-struct FVectorHandler : TPropertyComponentHandler<FVectorPropertyTraits, float, float, float, float>
+struct FFloatVectorHandler : TPropertyComponentHandler<FFloatVectorPropertyTraits, float, float, float, float>
 {
 	virtual void DispatchInitializePropertyMetaDataTasks(const FPropertyDefinition& Definition, FSystemTaskPrerequisites& InPrerequisites, FSystemSubsequentTasks& Subsequents, UMovieSceneEntitySystemLinker* Linker) override
 	{
@@ -378,9 +394,9 @@ struct FVectorHandler : TPropertyComponentHandler<FVectorPropertyTraits, float,
 		FEntityTaskBuilder()
 		.Read(BuiltInComponents->BoundObject)
 		.Read(BuiltInComponents->PropertyBinding)
-		.Write(TrackComponents->Vector.MetaDataComponents.GetType<0>())
+		.Write(TrackComponents->FloatVector.MetaDataComponents.GetType<0>())
 		.FilterAll({ BuiltInComponents->Tags.NeedsLink })
-		.Iterate_PerEntity(&Linker->EntityManager, [](UObject* Object, const FMovieScenePropertyBinding& Binding, FVectorChannelMetaData& OutMetaData)
+		.Iterate_PerEntity(&Linker->EntityManager, [](UObject* Object, const FMovieScenePropertyBinding& Binding, FVectorPropertyMetaData& OutMetaData)
 		{
 			FStructProperty* BoundProperty = CastField<FStructProperty>(FTrackInstancePropertyBindings::FindProperty(Object, Binding.PropertyPath.ToString()));
 			if (ensure(BoundProperty && BoundProperty->Struct))
@@ -389,7 +405,11 @@ struct FVectorHandler : TPropertyComponentHandler<FVectorPropertyTraits, float,
 				{
 					OutMetaData.NumChannels = 2;
 				}
-				else if (BoundProperty->Struct == TBaseStructure<FVector>::Get())
+				else if (BoundProperty->Struct->GetFName() == NAME_Vector3f
+#if UE_LARGE_WORLD_COORDINATES_DISABLED
+						|| BoundProperty->Struct->GetFName() == NAME_Vector
+#endif
+						)
 				{
 					OutMetaData.NumChannels = 3;
 				}
@@ -408,7 +428,40 @@ struct FVectorHandler : TPropertyComponentHandler<FVectorPropertyTraits, float,
 };
 
 
-struct FTransformHandler : TPropertyComponentHandler<FTransformPropertyTraits, float, float, float, float, float, float, float, float, float>
+struct FDoubleVectorHandler : TPropertyComponentHandler<FDoubleVectorPropertyTraits, double, double, double, double>
+{
+	virtual void DispatchInitializePropertyMetaDataTasks(const FPropertyDefinition& Definition, FSystemTaskPrerequisites& InPrerequisites, FSystemSubsequentTasks& Subsequents, UMovieSceneEntitySystemLinker* Linker) override
+	{
+		FBuiltInComponentTypes* BuiltInComponents = FBuiltInComponentTypes::Get();
+		FMovieSceneTracksComponentTypes* TrackComponents = FMovieSceneTracksComponentTypes::Get();
+
+		FEntityTaskBuilder()
+		.Read(BuiltInComponents->BoundObject)
+		.Read(BuiltInComponents->PropertyBinding)
+		.Write(TrackComponents->DoubleVector.MetaDataComponents.GetType<0>())
+		.FilterAll({ BuiltInComponents->Tags.NeedsLink })
+		.Iterate_PerEntity(&Linker->EntityManager, [](UObject* Object, const FMovieScenePropertyBinding& Binding, FVectorPropertyMetaData& OutMetaData)
+		{
+			FStructProperty* BoundProperty = CastField<FStructProperty>(FTrackInstancePropertyBindings::FindProperty(Object, Binding.PropertyPath.ToString()));
+			if (ensure(BoundProperty && BoundProperty->Struct))
+			{
+#if UE_LARGE_WORLD_COORDINATES_DISABLED
+				ensure(BoundProperty->Struct->GetFName() == NAME_Vector3d);
+#else
+				ensure(BoundProperty->Struct->GetFName() == NAME_Vector3d || BoundProperty->Struct->GetFName() == NAME_Vector);
+#endif
+				OutMetaData.NumChannels = 3;
+			}
+			else
+			{
+				OutMetaData.NumChannels = 4;
+			}
+		});
+	}
+};
+
+
+struct FTransformHandler : TPropertyComponentHandler<FTransformPropertyTraits, double, double, double, double, double, double, double, double, double>
 {
 	TSharedPtr<IPreAnimatedStorage> GetPreAnimatedStateStorage(const FPropertyDefinition& Definition, FPreAnimatedStateExtension* Container) override
 	{
@@ -424,16 +477,19 @@ FMovieSceneTracksComponentTypes::FMovieSceneTracksComponentTypes()
 	ComponentRegistry->NewPropertyType(Byte, TEXT("byte"));
 	ComponentRegistry->NewPropertyType(Enum, TEXT("enum"));
 	ComponentRegistry->NewPropertyType(Float, TEXT("float"));
+	ComponentRegistry->NewPropertyType(Double, TEXT("double"));
 	ComponentRegistry->NewPropertyType(Color, TEXT("color"));
 	ComponentRegistry->NewPropertyType(Integer, TEXT("int32"));
-	ComponentRegistry->NewPropertyType(Vector, TEXT("vector"));
+	ComponentRegistry->NewPropertyType(FloatVector, TEXT("float vector"));
+	ComponentRegistry->NewPropertyType(DoubleVector, TEXT("double vector"));
 
 	ComponentRegistry->NewPropertyType(Transform, TEXT("FTransform"));
 	ComponentRegistry->NewPropertyType(EulerTransform, TEXT("FEulerTransform"));
 	ComponentRegistry->NewPropertyType(ComponentTransform, TEXT("Component Transform"));
 
 	Color.MetaDataComponents.Initialize(ComponentRegistry, TEXT("Color Type"));
-	Vector.MetaDataComponents.Initialize(ComponentRegistry, TEXT("Num Vector Channels"));
+	FloatVector.MetaDataComponents.Initialize(ComponentRegistry, TEXT("Num Float Vector Channels"));
+	DoubleVector.MetaDataComponents.Initialize(ComponentRegistry, TEXT("Num Double Vector Channels"));
 
 	ComponentRegistry->NewComponentType(&QuaternionRotationChannel[0], TEXT("Quaternion Rotation Channel 0"));
 	ComponentRegistry->NewComponentType(&QuaternionRotationChannel[1], TEXT("Quaternion Rotation Channel 1"));
@@ -452,26 +508,29 @@ FMovieSceneTracksComponentTypes::FMovieSceneTracksComponentTypes()
 	// Set up bool properties
 	BuiltInComponents->PropertyRegistry.DefineProperty(Bool)
 	.AddSoleChannel(BuiltInComponents->BoolResult)
+	.SetBlenderSystem<UMovieScenePiecewiseBoolBlenderSystem>()
 	.SetCustomAccessors(&Accessors.Bool)
 	.Commit();
 
 	// Set up FTransform properties
 	BuiltInComponents->PropertyRegistry.DefineCompositeProperty(Transform)
-	.AddComposite(BuiltInComponents->FloatResult[0], &FIntermediate3DTransform::T_X)
-	.AddComposite(BuiltInComponents->FloatResult[1], &FIntermediate3DTransform::T_Y)
-	.AddComposite(BuiltInComponents->FloatResult[2], &FIntermediate3DTransform::T_Z)
-	.AddComposite(BuiltInComponents->FloatResult[3], &FIntermediate3DTransform::R_X)
-	.AddComposite(BuiltInComponents->FloatResult[4], &FIntermediate3DTransform::R_Y)
-	.AddComposite(BuiltInComponents->FloatResult[5], &FIntermediate3DTransform::R_Z)
-	.AddComposite(BuiltInComponents->FloatResult[6], &FIntermediate3DTransform::S_X)
-	.AddComposite(BuiltInComponents->FloatResult[7], &FIntermediate3DTransform::S_Y)
-	.AddComposite(BuiltInComponents->FloatResult[8], &FIntermediate3DTransform::S_Z)
+	.AddComposite(BuiltInComponents->DoubleResult[0], &FIntermediate3DTransform::T_X)
+	.AddComposite(BuiltInComponents->DoubleResult[1], &FIntermediate3DTransform::T_Y)
+	.AddComposite(BuiltInComponents->DoubleResult[2], &FIntermediate3DTransform::T_Z)
+	.AddComposite(BuiltInComponents->DoubleResult[3], &FIntermediate3DTransform::R_X)
+	.AddComposite(BuiltInComponents->DoubleResult[4], &FIntermediate3DTransform::R_Y)
+	.AddComposite(BuiltInComponents->DoubleResult[5], &FIntermediate3DTransform::R_Z)
+	.AddComposite(BuiltInComponents->DoubleResult[6], &FIntermediate3DTransform::S_X)
+	.AddComposite(BuiltInComponents->DoubleResult[7], &FIntermediate3DTransform::S_Y)
+	.AddComposite(BuiltInComponents->DoubleResult[8], &FIntermediate3DTransform::S_Z)
+	.SetBlenderSystem<UMovieScenePiecewiseDoubleBlenderSystem>()
 	.Commit();
 
 	// --------------------------------------------------------------------------------------------
 	// Set up byte properties
 	BuiltInComponents->PropertyRegistry.DefineProperty(Byte)
 	.AddSoleChannel(BuiltInComponents->ByteResult)
+	.SetBlenderSystem<UMovieScenePiecewiseByteBlenderSystem>()
 	.SetCustomAccessors(&Accessors.Byte)
 	.Commit();
 
@@ -479,6 +538,7 @@ FMovieSceneTracksComponentTypes::FMovieSceneTracksComponentTypes()
 	// Set up enum properties
 	BuiltInComponents->PropertyRegistry.DefineProperty(Enum)
 	.AddSoleChannel(BuiltInComponents->ByteResult)
+	.SetBlenderSystem<UMovieScenePiecewiseEnumBlenderSystem>()
 	.SetCustomAccessors(&Accessors.Enum)
 	.Commit();
 
@@ -486,6 +546,7 @@ FMovieSceneTracksComponentTypes::FMovieSceneTracksComponentTypes()
 	// Set up integer properties
 	BuiltInComponents->PropertyRegistry.DefineProperty(Integer)
 	.AddSoleChannel(BuiltInComponents->IntegerResult)
+	.SetBlenderSystem<UMovieScenePiecewiseIntegerBlenderSystem>()
 	.SetCustomAccessors(&Accessors.Integer)
 	.Commit();
 
@@ -493,6 +554,7 @@ FMovieSceneTracksComponentTypes::FMovieSceneTracksComponentTypes()
 	// Set up float properties
 	BuiltInComponents->PropertyRegistry.DefineProperty(Float)
 	.AddSoleChannel(BuiltInComponents->FloatResult[0])
+	.SetBlenderSystem<UMovieScenePiecewiseFloatBlenderSystem>()
 	.SetCustomAccessors(&Accessors.Float)
 	.Commit();
 
@@ -503,6 +565,7 @@ FMovieSceneTracksComponentTypes::FMovieSceneTracksComponentTypes()
 	.AddComposite(BuiltInComponents->FloatResult[1], &FIntermediateColor::G)
 	.AddComposite(BuiltInComponents->FloatResult[2], &FIntermediateColor::B)
 	.AddComposite(BuiltInComponents->FloatResult[3], &FIntermediateColor::A)
+	.SetBlenderSystem<UMovieScenePiecewiseFloatBlenderSystem>()
 	.SetCustomAccessors(&Accessors.Color)
 	.Commit(FColorHandler());
 
@@ -529,26 +592,37 @@ FMovieSceneTracksComponentTypes::FMovieSceneTracksComponentTypes()
 
 	// --------------------------------------------------------------------------------------------
 	// Set up vector properties
-	BuiltInComponents->PropertyRegistry.DefineCompositeProperty(Vector)
-	.AddComposite(BuiltInComponents->FloatResult[0], &FIntermediateVector::X)
-	.AddComposite(BuiltInComponents->FloatResult[1], &FIntermediateVector::Y)
-	.AddComposite(BuiltInComponents->FloatResult[2], &FIntermediateVector::Z)
-	.AddComposite(BuiltInComponents->FloatResult[3], &FIntermediateVector::W)
-	.SetCustomAccessors(&Accessors.Vector)
-	.Commit(FVectorHandler());
+	BuiltInComponents->PropertyRegistry.DefineCompositeProperty(FloatVector)
+	.AddComposite(BuiltInComponents->FloatResult[0], &FFloatIntermediateVector::X)
+	.AddComposite(BuiltInComponents->FloatResult[1], &FFloatIntermediateVector::Y)
+	.AddComposite(BuiltInComponents->FloatResult[2], &FFloatIntermediateVector::Z)
+	.AddComposite(BuiltInComponents->FloatResult[3], &FFloatIntermediateVector::W)
+	.SetBlenderSystem<UMovieScenePiecewiseFloatBlenderSystem>()
+	.SetCustomAccessors(&Accessors.FloatVector)
+	.Commit(FFloatVectorHandler());
+
+	BuiltInComponents->PropertyRegistry.DefineCompositeProperty(DoubleVector)
+	.AddComposite(BuiltInComponents->DoubleResult[0], &FDoubleIntermediateVector::X)
+	.AddComposite(BuiltInComponents->DoubleResult[1], &FDoubleIntermediateVector::Y)
+	.AddComposite(BuiltInComponents->DoubleResult[2], &FDoubleIntermediateVector::Z)
+	.AddComposite(BuiltInComponents->DoubleResult[3], &FDoubleIntermediateVector::W)
+	.SetBlenderSystem<UMovieScenePiecewiseDoubleBlenderSystem>()
+	.SetCustomAccessors(&Accessors.DoubleVector)
+	.Commit(FDoubleVectorHandler());
 
 	// --------------------------------------------------------------------------------------------
 	// Set up FEulerTransform properties
 	BuiltInComponents->PropertyRegistry.DefineCompositeProperty(EulerTransform)
-	.AddComposite(BuiltInComponents->FloatResult[0], &FIntermediate3DTransform::T_X)
-	.AddComposite(BuiltInComponents->FloatResult[1], &FIntermediate3DTransform::T_Y)
-	.AddComposite(BuiltInComponents->FloatResult[2], &FIntermediate3DTransform::T_Z)
-	.AddComposite(BuiltInComponents->FloatResult[3], &FIntermediate3DTransform::R_X)
-	.AddComposite(BuiltInComponents->FloatResult[4], &FIntermediate3DTransform::R_Y)
-	.AddComposite(BuiltInComponents->FloatResult[5], &FIntermediate3DTransform::R_Z)
-	.AddComposite(BuiltInComponents->FloatResult[6], &FIntermediate3DTransform::S_X)
-	.AddComposite(BuiltInComponents->FloatResult[7], &FIntermediate3DTransform::S_Y)
-	.AddComposite(BuiltInComponents->FloatResult[8], &FIntermediate3DTransform::S_Z)
+	.AddComposite(BuiltInComponents->DoubleResult[0], &FIntermediate3DTransform::T_X)
+	.AddComposite(BuiltInComponents->DoubleResult[1], &FIntermediate3DTransform::T_Y)
+	.AddComposite(BuiltInComponents->DoubleResult[2], &FIntermediate3DTransform::T_Z)
+	.AddComposite(BuiltInComponents->DoubleResult[3], &FIntermediate3DTransform::R_X)
+	.AddComposite(BuiltInComponents->DoubleResult[4], &FIntermediate3DTransform::R_Y)
+	.AddComposite(BuiltInComponents->DoubleResult[5], &FIntermediate3DTransform::R_Z)
+	.AddComposite(BuiltInComponents->DoubleResult[6], &FIntermediate3DTransform::S_X)
+	.AddComposite(BuiltInComponents->DoubleResult[7], &FIntermediate3DTransform::S_Y)
+	.AddComposite(BuiltInComponents->DoubleResult[8], &FIntermediate3DTransform::S_Z)
+	.SetBlenderSystem<UMovieScenePiecewiseDoubleBlenderSystem>()
 	.Commit();
 
 	// --------------------------------------------------------------------------------------------
@@ -557,15 +631,16 @@ FMovieSceneTracksComponentTypes::FMovieSceneTracksComponentTypes()
 		Accessors.ComponentTransform.Add(USceneComponent::StaticClass(), "Transform", &GetComponentTransform, &SetComponentTransformAndVelocity);
 
 		BuiltInComponents->PropertyRegistry.DefineCompositeProperty(ComponentTransform)
-		.AddComposite(BuiltInComponents->FloatResult[0], &FIntermediate3DTransform::T_X)
-		.AddComposite(BuiltInComponents->FloatResult[1], &FIntermediate3DTransform::T_Y)
-		.AddComposite(BuiltInComponents->FloatResult[2], &FIntermediate3DTransform::T_Z)
-		.AddComposite(BuiltInComponents->FloatResult[3], &FIntermediate3DTransform::R_X)
-		.AddComposite(BuiltInComponents->FloatResult[4], &FIntermediate3DTransform::R_Y)
-		.AddComposite(BuiltInComponents->FloatResult[5], &FIntermediate3DTransform::R_Z)
-		.AddComposite(BuiltInComponents->FloatResult[6], &FIntermediate3DTransform::S_X)
-		.AddComposite(BuiltInComponents->FloatResult[7], &FIntermediate3DTransform::S_Y)
-		.AddComposite(BuiltInComponents->FloatResult[8], &FIntermediate3DTransform::S_Z)
+		.AddComposite(BuiltInComponents->DoubleResult[0], &FIntermediate3DTransform::T_X)
+		.AddComposite(BuiltInComponents->DoubleResult[1], &FIntermediate3DTransform::T_Y)
+		.AddComposite(BuiltInComponents->DoubleResult[2], &FIntermediate3DTransform::T_Z)
+		.AddComposite(BuiltInComponents->DoubleResult[3], &FIntermediate3DTransform::R_X)
+		.AddComposite(BuiltInComponents->DoubleResult[4], &FIntermediate3DTransform::R_Y)
+		.AddComposite(BuiltInComponents->DoubleResult[5], &FIntermediate3DTransform::R_Z)
+		.AddComposite(BuiltInComponents->DoubleResult[6], &FIntermediate3DTransform::S_X)
+		.AddComposite(BuiltInComponents->DoubleResult[7], &FIntermediate3DTransform::S_Y)
+		.AddComposite(BuiltInComponents->DoubleResult[8], &FIntermediate3DTransform::S_Z)
+		.SetBlenderSystem<UMovieScenePiecewiseDoubleBlenderSystem>()
 		.SetCustomAccessors(&Accessors.ComponentTransform)
 		.Commit(FTransformHandler());
 	}
@@ -575,7 +650,7 @@ FMovieSceneTracksComponentTypes::FMovieSceneTracksComponentTypes()
 	for (int32 Index = 0; Index < UE_ARRAY_COUNT(QuaternionRotationChannel); ++Index)
 	{
 		ComponentRegistry->Factories.DuplicateChildComponent(QuaternionRotationChannel[Index]);
-		ComponentRegistry->Factories.DefineMutuallyInclusiveComponent(QuaternionRotationChannel[Index], BuiltInComponents->FloatResult[Index + 3]);
+		ComponentRegistry->Factories.DefineMutuallyInclusiveComponent(QuaternionRotationChannel[Index], BuiltInComponents->DoubleResult[Index + 3]);
 		ComponentRegistry->Factories.DefineMutuallyInclusiveComponent(QuaternionRotationChannel[Index], BuiltInComponents->EvalTime);
 	}
 
diff --git a/Engine/Source/Runtime/MovieSceneTracks/Private/Sections/MovieScene3DTransformSection.cpp b/Engine/Source/Runtime/MovieSceneTracks/Private/Sections/MovieScene3DTransformSection.cpp
--- a/Engine/Source/Runtime/MovieSceneTracks/Private/Sections/MovieScene3DTransformSection.cpp
+++ b/Engine/Source/Runtime/MovieSceneTracks/Private/Sections/MovieScene3DTransformSection.cpp
@@ -205,56 +205,57 @@ struct F3DTransformChannelEditorData
 		return TOptional<FVector>();
 	}
 
-	static TOptional<float> ExtractTranslationX(UObject& InObject, FTrackInstancePropertyBindings* Bindings)
+	static TOptional<double> ExtractTranslationX(UObject& InObject, FTrackInstancePropertyBindings* Bindings)
 	{
 		TOptional<FVector> Translation = GetTranslation(InObject, Bindings);
-		return Translation.IsSet() ? Translation->X : TOptional<float>();
+		return Translation.IsSet() ? Translation->X : TOptional<double>();
 	}
-	static TOptional<float> ExtractTranslationY(UObject& InObject, FTrackInstancePropertyBindings* Bindings)
+	static TOptional<double> ExtractTranslationY(UObject& InObject, FTrackInstancePropertyBindings* Bindings)
 	{
 		TOptional<FVector> Translation = GetTranslation(InObject, Bindings);
-		return Translation.IsSet() ? Translation->Y : TOptional<float>();
+		return Translation.IsSet() ? Translation->Y : TOptional<double>();
 	}
-	static TOptional<float> ExtractTranslationZ(UObject& InObject, FTrackInstancePropertyBindings* Bindings)
+	static TOptional<double> ExtractTranslationZ(UObject& InObject, FTrackInstancePropertyBindings* Bindings)
 	{
 		TOptional<FVector> Translation = GetTranslation(InObject, Bindings);
-		return Translation.IsSet() ? Translation->Z : TOptional<float>();
+		return Translation.IsSet() ? Translation->Z : TOptional<double>();
 	}
 
-	static TOptional<float> ExtractRotationX(UObject& InObject, FTrackInstancePropertyBindings* Bindings)
+	static TOptional<double> ExtractRotationX(UObject& InObject, FTrackInstancePropertyBindings* Bindings)
 	{
 		TOptional<FRotator> Rotator = GetRotator(InObject, Bindings);
-		return Rotator.IsSet() ? Rotator->Roll : TOptional<float>();
+		return Rotator.IsSet() ? Rotator->Roll : TOptional<double>();
 	}
-	static TOptional<float> ExtractRotationY(UObject& InObject, FTrackInstancePropertyBindings* Bindings)
+	static TOptional<double> ExtractRotationY(UObject& InObject, FTrackInstancePropertyBindings* Bindings)
 	{
 		TOptional<FRotator> Rotator = GetRotator(InObject, Bindings);
-		return Rotator.IsSet() ? Rotator->Pitch : TOptional<float>();
+		return Rotator.IsSet() ? Rotator->Pitch : TOptional<double>();
 	}
-	static TOptional<float> ExtractRotationZ(UObject& InObject, FTrackInstancePropertyBindings* Bindings)
+	static TOptional<double> ExtractRotationZ(UObject& InObject, FTrackInstancePropertyBindings* Bindings)
 	{
 		TOptional<FRotator> Rotator = GetRotator(InObject, Bindings);
-		return Rotator.IsSet() ? Rotator->Yaw : TOptional<float>();
+		return Rotator.IsSet() ? Rotator->Yaw : TOptional<double>();
 	}
 
-	static TOptional<float> ExtractScaleX(UObject& InObject, FTrackInstancePropertyBindings* Bindings)
+	static TOptional<double> ExtractScaleX(UObject& InObject, FTrackInstancePropertyBindings* Bindings)
 	{
 		TOptional<FVector> Scale = GetScale(InObject, Bindings);
-		return Scale.IsSet() ? Scale->X : TOptional<float>();
+		return Scale.IsSet() ? Scale->X : TOptional<double>();
 	}
-	static TOptional<float> ExtractScaleY(UObject& InObject, FTrackInstancePropertyBindings* Bindings)
+	static TOptional<double> ExtractScaleY(UObject& InObject, FTrackInstancePropertyBindings* Bindings)
 	{
 		TOptional<FVector> Scale = GetScale(InObject, Bindings);
-		return Scale.IsSet() ? Scale->Y : TOptional<float>();
+		return Scale.IsSet() ? Scale->Y : TOptional<double>();
 	}
-	static TOptional<float> ExtractScaleZ(UObject& InObject, FTrackInstancePropertyBindings* Bindings)
+	static TOptional<double> ExtractScaleZ(UObject& InObject, FTrackInstancePropertyBindings* Bindings)
 	{
 		TOptional<FVector> Scale = GetScale(InObject, Bindings);
-		return Scale.IsSet() ? Scale->Z : TOptional<float>();
+		return Scale.IsSet() ? Scale->Z : TOptional<double>();
 	}
 
 	FMovieSceneChannelMetaData      MetaData[10];
-	TMovieSceneExternalValue<float> ExternalValues[10];
+	TMovieSceneExternalValue<double> ExternalValues[9];
+	TMovieSceneExternalValue<float> WeightExternalValue;
 };
 
 #endif // WITH_EDITOR
@@ -375,12 +376,12 @@ void UMovieScene3DTransformSection::BuildEntity(BaseBuilderType& InBaseBuilder,
 		return;
 	}
 
-	TComponentTypeID<FSourceFloatChannel> RotationChannel[3];
+	TComponentTypeID<FSourceDoubleChannel> RotationChannel[3];
 	if (!bUseQuaternionInterpolation)
 	{
-		RotationChannel[0] = BuiltInComponentTypes->FloatChannel[3];
-		RotationChannel[1] = BuiltInComponentTypes->FloatChannel[4];
-		RotationChannel[2] = BuiltInComponentTypes->FloatChannel[5];
+		RotationChannel[0] = BuiltInComponentTypes->DoubleChannel[3];
+		RotationChannel[1] = BuiltInComponentTypes->DoubleChannel[4];
+		RotationChannel[2] = BuiltInComponentTypes->DoubleChannel[5];
 	}
 	else
 	{
@@ -391,16 +392,16 @@ void UMovieScene3DTransformSection::BuildEntity(BaseBuilderType& InBaseBuilder,
 
 	OutImportedEntity->AddBuilder(
 		InBaseBuilder
-		.AddConditional(BuiltInComponentTypes->FloatChannel[0], &Translation[0],          ActiveChannelsMask[0])
-		.AddConditional(BuiltInComponentTypes->FloatChannel[1], &Translation[1],          ActiveChannelsMask[1])
-		.AddConditional(BuiltInComponentTypes->FloatChannel[2], &Translation[2],          ActiveChannelsMask[2])
-		.AddConditional(RotationChannel[0],                     &Rotation[0],             ActiveChannelsMask[3])
-		.AddConditional(RotationChannel[1],                     &Rotation[1],             ActiveChannelsMask[4])
-		.AddConditional(RotationChannel[2],                     &Rotation[2],             ActiveChannelsMask[5])
-		.AddConditional(BuiltInComponentTypes->FloatChannel[6], &Scale[0],                ActiveChannelsMask[6])
-		.AddConditional(BuiltInComponentTypes->FloatChannel[7], &Scale[1],                ActiveChannelsMask[7])
-		.AddConditional(BuiltInComponentTypes->FloatChannel[8], &Scale[2],                ActiveChannelsMask[8])
-		.AddConditional(BuiltInComponentTypes->WeightChannel,   &ManualWeight,            EnumHasAnyFlags(Channels, EMovieSceneTransformChannel::Weight) && ManualWeight.HasAnyData())
+		.AddConditional(BuiltInComponentTypes->DoubleChannel[0], &Translation[0], ActiveChannelsMask[0])
+		.AddConditional(BuiltInComponentTypes->DoubleChannel[1], &Translation[1], ActiveChannelsMask[1])
+		.AddConditional(BuiltInComponentTypes->DoubleChannel[2], &Translation[2], ActiveChannelsMask[2])
+		.AddConditional(RotationChannel[0],                      &Rotation[0],    ActiveChannelsMask[3])
+		.AddConditional(RotationChannel[1],                      &Rotation[1],    ActiveChannelsMask[4])
+		.AddConditional(RotationChannel[2],                      &Rotation[2],    ActiveChannelsMask[5])
+		.AddConditional(BuiltInComponentTypes->DoubleChannel[6], &Scale[0],       ActiveChannelsMask[6])
+		.AddConditional(BuiltInComponentTypes->DoubleChannel[7], &Scale[1],       ActiveChannelsMask[7])
+		.AddConditional(BuiltInComponentTypes->DoubleChannel[8], &Scale[2],       ActiveChannelsMask[8])
+		.AddConditional(BuiltInComponentTypes->WeightChannel,    &ManualWeight,   EnumHasAnyFlags(Channels, EMovieSceneTransformChannel::Weight) && ManualWeight.HasAnyData())
 		.AddTag(PropertyTag)
 	);
 }
@@ -520,7 +521,7 @@ EMovieSceneChannelProxyType UMovieScene3DTransformSection::CacheChannelProxy()
 	Channels.Add(Scale[0],       EditorData.MetaData[6], EditorData.ExternalValues[6]);
 	Channels.Add(Scale[1],       EditorData.MetaData[7], EditorData.ExternalValues[7]);
 	Channels.Add(Scale[2],       EditorData.MetaData[8], EditorData.ExternalValues[8]);
-	Channels.Add(ManualWeight,   EditorData.MetaData[9], EditorData.ExternalValues[9]);
+	Channels.Add(ManualWeight,   EditorData.MetaData[9], EditorData.WeightExternalValue);
 
 #else
 
@@ -551,24 +552,24 @@ TSharedPtr<FStructOnScope> UMovieScene3DTransformSection::GetKeyStruct(TArrayVie
 	FRotator StartingRotation;
 	FVector  StartingScale;
 
-	TArrayView<FMovieSceneFloatChannel* const> FloatChannels = ChannelProxy->GetChannels<FMovieSceneFloatChannel>();
+	TArrayView<FMovieSceneDoubleChannel* const> DoubleChannels = ChannelProxy->GetChannels<FMovieSceneDoubleChannel>();
 
 	TOptional<TTuple<FKeyHandle, FFrameNumber>> LocationKeys[3] = {
-		FMovieSceneChannelValueHelper::FindFirstKey(FloatChannels[0], KeyHandles),
-		FMovieSceneChannelValueHelper::FindFirstKey(FloatChannels[1], KeyHandles),
-		FMovieSceneChannelValueHelper::FindFirstKey(FloatChannels[2], KeyHandles)
+		FMovieSceneChannelValueHelper::FindFirstKey(DoubleChannels[0], KeyHandles),
+		FMovieSceneChannelValueHelper::FindFirstKey(DoubleChannels[1], KeyHandles),
+		FMovieSceneChannelValueHelper::FindFirstKey(DoubleChannels[2], KeyHandles)
 	};
 
 	TOptional<TTuple<FKeyHandle, FFrameNumber>> RotationKeys[3] = {
-		FMovieSceneChannelValueHelper::FindFirstKey(FloatChannels[3], KeyHandles),
-		FMovieSceneChannelValueHelper::FindFirstKey(FloatChannels[4], KeyHandles),
-		FMovieSceneChannelValueHelper::FindFirstKey(FloatChannels[5], KeyHandles)
+		FMovieSceneChannelValueHelper::FindFirstKey(DoubleChannels[3], KeyHandles),
+		FMovieSceneChannelValueHelper::FindFirstKey(DoubleChannels[4], KeyHandles),
+		FMovieSceneChannelValueHelper::FindFirstKey(DoubleChannels[5], KeyHandles)
 	};
 
 	TOptional<TTuple<FKeyHandle, FFrameNumber>> ScaleKeys[3] = {
-		FMovieSceneChannelValueHelper::FindFirstKey(FloatChannels[6], KeyHandles),
-		FMovieSceneChannelValueHelper::FindFirstKey(FloatChannels[7], KeyHandles),
-		FMovieSceneChannelValueHelper::FindFirstKey(FloatChannels[8], KeyHandles)
+		FMovieSceneChannelValueHelper::FindFirstKey(DoubleChannels[6], KeyHandles),
+		FMovieSceneChannelValueHelper::FindFirstKey(DoubleChannels[7], KeyHandles),
+		FMovieSceneChannelValueHelper::FindFirstKey(DoubleChannels[8], KeyHandles)
 	};
 
 	const int32 AnyLocationKeys = Algo::AnyOf(LocationKeys);
@@ -581,17 +582,17 @@ TSharedPtr<FStructOnScope> UMovieScene3DTransformSection::GetKeyStruct(TArrayVie
 		TSharedRef<FStructOnScope> KeyStruct = MakeShareable(new FStructOnScope(FMovieScene3DTransformKeyStruct::StaticStruct()));
 		auto Struct = (FMovieScene3DTransformKeyStruct*)KeyStruct->GetStructMemory();
 
-		Struct->KeyStructInterop.Add(FMovieSceneChannelValueHelper(ChannelProxy->MakeHandle<FMovieSceneFloatChannel>(0), &Struct->Location.X,     LocationKeys[0]));
-		Struct->KeyStructInterop.Add(FMovieSceneChannelValueHelper(ChannelProxy->MakeHandle<FMovieSceneFloatChannel>(1), &Struct->Location.Y,     LocationKeys[1]));
-		Struct->KeyStructInterop.Add(FMovieSceneChannelValueHelper(ChannelProxy->MakeHandle<FMovieSceneFloatChannel>(2), &Struct->Location.Z,     LocationKeys[2]));
+		Struct->KeyStructInterop.Add(FMovieSceneChannelValueHelper(ChannelProxy->MakeHandle<FMovieSceneDoubleChannel>(0), &Struct->Location.X,     LocationKeys[0]));
+		Struct->KeyStructInterop.Add(FMovieSceneChannelValueHelper(ChannelProxy->MakeHandle<FMovieSceneDoubleChannel>(1), &Struct->Location.Y,     LocationKeys[1]));
+		Struct->KeyStructInterop.Add(FMovieSceneChannelValueHelper(ChannelProxy->MakeHandle<FMovieSceneDoubleChannel>(2), &Struct->Location.Z,     LocationKeys[2]));
 
-		Struct->KeyStructInterop.Add(FMovieSceneChannelValueHelper(ChannelProxy->MakeHandle<FMovieSceneFloatChannel>(3), &Struct->Rotation.Roll,  RotationKeys[0]));
-		Struct->KeyStructInterop.Add(FMovieSceneChannelValueHelper(ChannelProxy->MakeHandle<FMovieSceneFloatChannel>(4), &Struct->Rotation.Pitch, RotationKeys[1]));
-		Struct->KeyStructInterop.Add(FMovieSceneChannelValueHelper(ChannelProxy->MakeHandle<FMovieSceneFloatChannel>(5), &Struct->Rotation.Yaw,   RotationKeys[2]));
+		Struct->KeyStructInterop.Add(FMovieSceneChannelValueHelper(ChannelProxy->MakeHandle<FMovieSceneDoubleChannel>(3), &Struct->Rotation.Roll,  RotationKeys[0]));
+		Struct->KeyStructInterop.Add(FMovieSceneChannelValueHelper(ChannelProxy->MakeHandle<FMovieSceneDoubleChannel>(4), &Struct->Rotation.Pitch, RotationKeys[1]));
+		Struct->KeyStructInterop.Add(FMovieSceneChannelValueHelper(ChannelProxy->MakeHandle<FMovieSceneDoubleChannel>(5), &Struct->Rotation.Yaw,   RotationKeys[2]));
 
-		Struct->KeyStructInterop.Add(FMovieSceneChannelValueHelper(ChannelProxy->MakeHandle<FMovieSceneFloatChannel>(6), &Struct->Scale.X,        ScaleKeys[0]));
-		Struct->KeyStructInterop.Add(FMovieSceneChannelValueHelper(ChannelProxy->MakeHandle<FMovieSceneFloatChannel>(7), &Struct->Scale.Y,        ScaleKeys[1]));
-		Struct->KeyStructInterop.Add(FMovieSceneChannelValueHelper(ChannelProxy->MakeHandle<FMovieSceneFloatChannel>(8), &Struct->Scale.Z,        ScaleKeys[2]));
+		Struct->KeyStructInterop.Add(FMovieSceneChannelValueHelper(ChannelProxy->MakeHandle<FMovieSceneDoubleChannel>(6), &Struct->Scale.X,        ScaleKeys[0]));
+		Struct->KeyStructInterop.Add(FMovieSceneChannelValueHelper(ChannelProxy->MakeHandle<FMovieSceneDoubleChannel>(7), &Struct->Scale.Y,        ScaleKeys[1]));
+		Struct->KeyStructInterop.Add(FMovieSceneChannelValueHelper(ChannelProxy->MakeHandle<FMovieSceneDoubleChannel>(8), &Struct->Scale.Z,        ScaleKeys[2]));
 
 		Struct->KeyStructInterop.SetStartingValues();
 		Struct->Time = Struct->KeyStructInterop.GetUnifiedKeyTime().Get(0);
@@ -603,9 +604,9 @@ TSharedPtr<FStructOnScope> UMovieScene3DTransformSection::GetKeyStruct(TArrayVie
 		TSharedRef<FStructOnScope> KeyStruct = MakeShareable(new FStructOnScope(FMovieScene3DLocationKeyStruct::StaticStruct()));
 		auto Struct = (FMovieScene3DLocationKeyStruct*)KeyStruct->GetStructMemory();
 
-		Struct->KeyStructInterop.Add(FMovieSceneChannelValueHelper(ChannelProxy->MakeHandle<FMovieSceneFloatChannel>(0), &Struct->Location.X,     LocationKeys[0]));
-		Struct->KeyStructInterop.Add(FMovieSceneChannelValueHelper(ChannelProxy->MakeHandle<FMovieSceneFloatChannel>(1), &Struct->Location.Y,     LocationKeys[1]));
-		Struct->KeyStructInterop.Add(FMovieSceneChannelValueHelper(ChannelProxy->MakeHandle<FMovieSceneFloatChannel>(2), &Struct->Location.Z,     LocationKeys[2]));
+		Struct->KeyStructInterop.Add(FMovieSceneChannelValueHelper(ChannelProxy->MakeHandle<FMovieSceneDoubleChannel>(0), &Struct->Location.X,     LocationKeys[0]));
+		Struct->KeyStructInterop.Add(FMovieSceneChannelValueHelper(ChannelProxy->MakeHandle<FMovieSceneDoubleChannel>(1), &Struct->Location.Y,     LocationKeys[1]));
+		Struct->KeyStructInterop.Add(FMovieSceneChannelValueHelper(ChannelProxy->MakeHandle<FMovieSceneDoubleChannel>(2), &Struct->Location.Z,     LocationKeys[2]));
 
 		Struct->KeyStructInterop.SetStartingValues();
 		Struct->Time = Struct->KeyStructInterop.GetUnifiedKeyTime().Get(0);
@@ -617,9 +618,9 @@ TSharedPtr<FStructOnScope> UMovieScene3DTransformSection::GetKeyStruct(TArrayVie
 		TSharedRef<FStructOnScope> KeyStruct = MakeShareable(new FStructOnScope(FMovieScene3DRotationKeyStruct::StaticStruct()));
 		auto Struct = (FMovieScene3DRotationKeyStruct*)KeyStruct->GetStructMemory();
 
-		Struct->KeyStructInterop.Add(FMovieSceneChannelValueHelper(ChannelProxy->MakeHandle<FMovieSceneFloatChannel>(3), &Struct->Rotation.Roll,  RotationKeys[0]));
-		Struct->KeyStructInterop.Add(FMovieSceneChannelValueHelper(ChannelProxy->MakeHandle<FMovieSceneFloatChannel>(4), &Struct->Rotation.Pitch, RotationKeys[1]));
-		Struct->KeyStructInterop.Add(FMovieSceneChannelValueHelper(ChannelProxy->MakeHandle<FMovieSceneFloatChannel>(5), &Struct->Rotation.Yaw,   RotationKeys[2]));
+		Struct->KeyStructInterop.Add(FMovieSceneChannelValueHelper(ChannelProxy->MakeHandle<FMovieSceneDoubleChannel>(3), &Struct->Rotation.Roll,  RotationKeys[0]));
+		Struct->KeyStructInterop.Add(FMovieSceneChannelValueHelper(ChannelProxy->MakeHandle<FMovieSceneDoubleChannel>(4), &Struct->Rotation.Pitch, RotationKeys[1]));
+		Struct->KeyStructInterop.Add(FMovieSceneChannelValueHelper(ChannelProxy->MakeHandle<FMovieSceneDoubleChannel>(5), &Struct->Rotation.Yaw,   RotationKeys[2]));
 
 		Struct->KeyStructInterop.SetStartingValues();
 		Struct->Time = Struct->KeyStructInterop.GetUnifiedKeyTime().Get(0);
@@ -631,9 +632,9 @@ TSharedPtr<FStructOnScope> UMovieScene3DTransformSection::GetKeyStruct(TArrayVie
 		TSharedRef<FStructOnScope> KeyStruct = MakeShareable(new FStructOnScope(FMovieScene3DScaleKeyStruct::StaticStruct()));
 		auto Struct = (FMovieScene3DScaleKeyStruct*)KeyStruct->GetStructMemory();
 
-		Struct->KeyStructInterop.Add(FMovieSceneChannelValueHelper(ChannelProxy->MakeHandle<FMovieSceneFloatChannel>(6), &Struct->Scale.X,        ScaleKeys[0]));
-		Struct->KeyStructInterop.Add(FMovieSceneChannelValueHelper(ChannelProxy->MakeHandle<FMovieSceneFloatChannel>(7), &Struct->Scale.Y,        ScaleKeys[1]));
-		Struct->KeyStructInterop.Add(FMovieSceneChannelValueHelper(ChannelProxy->MakeHandle<FMovieSceneFloatChannel>(8), &Struct->Scale.Z,        ScaleKeys[2]));
+		Struct->KeyStructInterop.Add(FMovieSceneChannelValueHelper(ChannelProxy->MakeHandle<FMovieSceneDoubleChannel>(6), &Struct->Scale.X,        ScaleKeys[0]));
+		Struct->KeyStructInterop.Add(FMovieSceneChannelValueHelper(ChannelProxy->MakeHandle<FMovieSceneDoubleChannel>(7), &Struct->Scale.Y,        ScaleKeys[1]));
+		Struct->KeyStructInterop.Add(FMovieSceneChannelValueHelper(ChannelProxy->MakeHandle<FMovieSceneDoubleChannel>(8), &Struct->Scale.Z,        ScaleKeys[2]));
 
 		Struct->KeyStructInterop.SetStartingValues();
 		Struct->Time = Struct->KeyStructInterop.GetUnifiedKeyTime().Get(0);
