Sequencer: Large world coordinates support

diff --git a/Engine/Source/Runtime/MovieSceneTracks/Private/Sections/MovieSceneVectorSection.cpp b/Engine/Source/Runtime/MovieSceneTracks/Private/Sections/MovieSceneVectorSection.cpp
--- a/Engine/Source/Runtime/MovieSceneTracks/Private/Sections/MovieSceneVectorSection.cpp
+++ b/Engine/Source/Runtime/MovieSceneTracks/Private/Sections/MovieSceneVectorSection.cpp
@@ -12,9 +12,62 @@
  *****************************************************************************/
 
 #if WITH_EDITOR
-struct FVectorSectionEditorData
+template<typename ValueType>
+struct TVectorSectionEditorDataCompositeValue
 {
-	FVectorSectionEditorData(int32 NumChannels)
+	ValueType X, Y, Z, W;
+	TVectorSectionEditorDataCompositeValue() : X(0), Y(0), Z(0), W(0) {}
+	TVectorSectionEditorDataCompositeValue(ValueType InX, ValueType InY, ValueType InZ, ValueType InW)
+		: X(InX), Y(InY), Z(InZ), W(InW)
+	{}
+};
+
+template<typename ValueType>
+struct TVectorSectionEditorDataTraits;
+
+template<>
+struct TVectorSectionEditorDataTraits<float>
+{
+	using FCompositeValue = TVectorSectionEditorDataCompositeValue<float>;
+
+	static FCompositeValue GetPropertyValue(UObject& InObject, FTrackInstancePropertyBindings& Bindings, int32 NumChannels)
+	{
+		if (NumChannels == 2)
+		{
+			FVector2D Vector = Bindings.GetCurrentValue<FVector2D>(InObject);
+			return FCompositeValue(Vector.X, Vector.Y, 0.f, 0.f);
+		}
+		else if (NumChannels == 3)
+		{
+			FVector3f Vector = Bindings.GetCurrentValue<FVector3f>(InObject);
+			return FCompositeValue(Vector.X, Vector.Y, Vector.Z, 0.f);
+		}
+		else
+		{
+			ensure(NumChannels == 4);
+			FVector4 Vector = Bindings.GetCurrentValue<FVector4>(InObject);
+			return FCompositeValue(Vector.X, Vector.Y, Vector.Z, Vector.W);
+		}
+	}
+};
+
+template<>
+struct TVectorSectionEditorDataTraits<double>
+{
+	using FCompositeValue = TVectorSectionEditorDataCompositeValue<double>;
+
+	static FCompositeValue GetPropertyValue(UObject& InObject, FTrackInstancePropertyBindings& Bindings, int32 NumChannels)
+	{
+		ensure(NumChannels == 3);
+		FVector3d Vector = Bindings.GetCurrentValue<FVector3d>(InObject);
+		return FCompositeValue(Vector.X, Vector.Y, Vector.Z, 0.f);
+	}
+};
+
+template<typename ValueType>
+struct TVectorSectionEditorData
+{
+	TVectorSectionEditorData(int32 NumChannels)
 	{
 		MetaData[0].SetIdentifiers("Vector.X", FCommonChannelData::ChannelX);
 		MetaData[0].SortOrder = 0;
@@ -40,58 +93,44 @@ struct FVectorSectionEditorData
 		ExternalValues[2].OnGetExternalValue = [NumChannels](UObject& InObject, FTrackInstancePropertyBindings* Bindings) { return ExtractChannelZ(InObject, Bindings, NumChannels); };
 		ExternalValues[3].OnGetExternalValue = [NumChannels](UObject& InObject, FTrackInstancePropertyBindings* Bindings) { return ExtractChannelW(InObject, Bindings, NumChannels); };
 	}
-
-	static FVector4 GetPropertyValue(UObject& InObject, FTrackInstancePropertyBindings& Bindings, int32 NumChannels)
-	{
-		if (NumChannels == 2)
-		{
-			FVector2D Vector = Bindings.GetCurrentValue<FVector2D>(InObject);
-			return FVector4(Vector.X, Vector.Y, 0.f, 0.f);
-		}
-		else if (NumChannels == 3)
-		{
-			FVector Vector = Bindings.GetCurrentValue<FVector>(InObject);
-			return FVector4(Vector.X, Vector.Y, Vector.Z, 0.f);
-		}
-		else
-		{
-			return Bindings.GetCurrentValue<FVector4>(InObject);
-		}
-	}
 	
-	static TOptional<float> ExtractChannelX(UObject& InObject, FTrackInstancePropertyBindings* Bindings, int32 NumChannels)
+	static TOptional<ValueType> ExtractChannelX(UObject& InObject, FTrackInstancePropertyBindings* Bindings, int32 NumChannels)
 	{
-		return Bindings ? GetPropertyValue(InObject, *Bindings, NumChannels).X : TOptional<float>();
+		return Bindings ? TVectorSectionEditorDataTraits<ValueType>::GetPropertyValue(InObject, *Bindings, NumChannels).X : TOptional<ValueType>();
 	}
-	static TOptional<float> ExtractChannelY(UObject& InObject, FTrackInstancePropertyBindings* Bindings, int32 NumChannels)
+	static TOptional<ValueType> ExtractChannelY(UObject& InObject, FTrackInstancePropertyBindings* Bindings, int32 NumChannels)
 	{
-		return Bindings ? GetPropertyValue(InObject, *Bindings, NumChannels).Y : TOptional<float>();
+		return Bindings ? TVectorSectionEditorDataTraits<ValueType>::GetPropertyValue(InObject, *Bindings, NumChannels).Y : TOptional<ValueType>();
 	}
-	static TOptional<float> ExtractChannelZ(UObject& InObject, FTrackInstancePropertyBindings* Bindings, int32 NumChannels)
+	static TOptional<ValueType> ExtractChannelZ(UObject& InObject, FTrackInstancePropertyBindings* Bindings, int32 NumChannels)
 	{
-		return Bindings ? GetPropertyValue(InObject, *Bindings, NumChannels).Z : TOptional<float>();
+		return Bindings ? TVectorSectionEditorDataTraits<ValueType>::GetPropertyValue(InObject, *Bindings, NumChannels).Z : TOptional<ValueType>();
 	}
-	static TOptional<float> ExtractChannelW(UObject& InObject, FTrackInstancePropertyBindings* Bindings, int32 NumChannels)
+	static TOptional<ValueType> ExtractChannelW(UObject& InObject, FTrackInstancePropertyBindings* Bindings, int32 NumChannels)
 	{
-		return Bindings ? GetPropertyValue(InObject, *Bindings, NumChannels).W : TOptional<float>();
+		return Bindings ? TVectorSectionEditorDataTraits<ValueType>::GetPropertyValue(InObject, *Bindings, NumChannels).W : TOptional<ValueType>();
 	}
 
-	FMovieSceneChannelMetaData      MetaData[4];
-	TMovieSceneExternalValue<float> ExternalValues[4];
-
+	FMovieSceneChannelMetaData			MetaData[4];
+	TMovieSceneExternalValue<ValueType> ExternalValues[4];
 };
 #endif
 
-void FMovieSceneVectorKeyStructBase::PropagateChanges(const FPropertyChangedEvent& ChangeEvent)
+void FMovieSceneFloatVectorKeyStructBase::PropagateChanges(const FPropertyChangedEvent& ChangeEvent)
 {
 	KeyStructInterop.Apply(Time);
 }
 
+void FMovieSceneDoubleVectorKeyStructBase::PropagateChanges(const FPropertyChangedEvent& ChangeEvent)
+{
+	KeyStructInterop.Apply(Time);
+}
 
-/* UMovieSceneVectorSection structors
+
+/* UMovieSceneFloatVectorSection structors
  *****************************************************************************/
 
-UMovieSceneVectorSection::UMovieSceneVectorSection(const FObjectInitializer& ObjectInitializer)
+UMovieSceneFloatVectorSection::UMovieSceneFloatVectorSection(const FObjectInitializer& ObjectInitializer)
 	: Super(ObjectInitializer)
 {
 	ChannelsUsed = 0;
@@ -106,7 +145,7 @@ UMovieSceneVectorSection::UMovieSceneVectorSection(const FObjectInitializer& Obj
 	BlendType = EMovieSceneBlendType::Absolute;
 }
 
-void UMovieSceneVectorSection::Serialize(FArchive& Ar)
+void UMovieSceneFloatVectorSection::Serialize(FArchive& Ar)
 {
 	Super::Serialize(Ar);
 
@@ -116,14 +155,14 @@ void UMovieSceneVectorSection::Serialize(FArchive& Ar)
 	}
 }
 
-void UMovieSceneVectorSection::PostEditImport()
+void UMovieSceneFloatVectorSection::PostEditImport()
 {
 	Super::PostEditImport();
 
 	RecreateChannelProxy();
 }
 
-void UMovieSceneVectorSection::RecreateChannelProxy()
+void UMovieSceneFloatVectorSection::RecreateChannelProxy()
 {
 	FMovieSceneChannelProxyData Channels;
 
@@ -131,7 +170,7 @@ void UMovieSceneVectorSection::RecreateChannelProxy()
 
 #if WITH_EDITOR
 
-	const FVectorSectionEditorData EditorData(ChannelsUsed);
+	const TVectorSectionEditorData<float> EditorData(ChannelsUsed);
 	for (int32 Index = 0; Index < ChannelsUsed; ++Index)
 	{
 		Channels.Add(Curves[Index], EditorData.MetaData[Index], EditorData.ExternalValues[Index]);
@@ -149,7 +188,7 @@ void UMovieSceneVectorSection::RecreateChannelProxy()
 	ChannelProxy = MakeShared<FMovieSceneChannelProxy>(MoveTemp(Channels));
 }
 
-TSharedPtr<FStructOnScope> UMovieSceneVectorSection::GetKeyStruct(TArrayView<const FKeyHandle> KeyHandles)
+TSharedPtr<FStructOnScope> UMovieSceneFloatVectorSection::GetKeyStruct(TArrayView<const FKeyHandle> KeyHandles)
 {
 	TSharedPtr<FStructOnScope> KeyStruct;
 	if (ChannelsUsed == 2)
@@ -158,7 +197,7 @@ TSharedPtr<FStructOnScope> UMovieSceneVectorSection::GetKeyStruct(TArrayView<con
 	}
 	else if (ChannelsUsed == 3)
 	{
-		KeyStruct = MakeShareable(new FStructOnScope(FMovieSceneVectorKeyStruct::StaticStruct()));
+		KeyStruct = MakeShareable(new FStructOnScope(FMovieSceneVector3fKeyStruct::StaticStruct()));
 	}
 	else if (ChannelsUsed == 4)
 	{
@@ -167,7 +206,7 @@ TSharedPtr<FStructOnScope> UMovieSceneVectorSection::GetKeyStruct(TArrayView<con
 
 	if (KeyStruct.IsValid())
 	{
-		FMovieSceneVectorKeyStructBase* Struct = (FMovieSceneVectorKeyStructBase*)KeyStruct->GetStructMemory();
+		FMovieSceneFloatVectorKeyStructBase* Struct = (FMovieSceneFloatVectorKeyStructBase*)KeyStruct->GetStructMemory();
 		for (int32 Index = 0; Index < ChannelsUsed; ++Index)
 		{
 			Struct->KeyStructInterop.Add(FMovieSceneChannelValueHelper(ChannelProxy->MakeHandle<FMovieSceneFloatChannel>(Index), Struct->GetPropertyChannelByIndex(Index), KeyHandles));
@@ -180,13 +219,13 @@ TSharedPtr<FStructOnScope> UMovieSceneVectorSection::GetKeyStruct(TArrayView<con
 	return KeyStruct;
 }
 
-bool UMovieSceneVectorSection::PopulateEvaluationFieldImpl(const TRange<FFrameNumber>& EffectiveRange, const FMovieSceneEvaluationFieldEntityMetaData& InMetaData, FMovieSceneEntityComponentFieldBuilder* OutFieldBuilder)
+bool UMovieSceneFloatVectorSection::PopulateEvaluationFieldImpl(const TRange<FFrameNumber>& EffectiveRange, const FMovieSceneEvaluationFieldEntityMetaData& InMetaData, FMovieSceneEntityComponentFieldBuilder* OutFieldBuilder)
 {
 	FMovieScenePropertyTrackEntityImportHelper::PopulateEvaluationField(*this, EffectiveRange, InMetaData, OutFieldBuilder);
 	return true;
 }
 
-void UMovieSceneVectorSection::ImportEntityImpl(UMovieSceneEntitySystemLinker* EntityLinker, const FEntityImportParams& Params, FImportedEntity* OutImportedEntity)
+void UMovieSceneFloatVectorSection::ImportEntityImpl(UMovieSceneEntitySystemLinker* EntityLinker, const FEntityImportParams& Params, FImportedEntity* OutImportedEntity)
 {
 	using namespace UE::MovieScene;
 
@@ -202,10 +241,117 @@ void UMovieSceneVectorSection::ImportEntityImpl(UMovieSceneEntitySystemLinker* E
 	const FBuiltInComponentTypes* Components = FBuiltInComponentTypes::Get();
 	const FMovieSceneTracksComponentTypes* TracksComponents = FMovieSceneTracksComponentTypes::Get();
 
-	FPropertyTrackEntityImportHelper(TracksComponents->Vector)
+	FPropertyTrackEntityImportHelper(TracksComponents->FloatVector)
 		.AddConditional(Components->FloatChannel[0], &Curves[0], ChannelsUsed > 0 && Curves[0].HasAnyData())
 		.AddConditional(Components->FloatChannel[1], &Curves[1], ChannelsUsed > 1 && Curves[1].HasAnyData())
 		.AddConditional(Components->FloatChannel[2], &Curves[2], ChannelsUsed > 2 && Curves[2].HasAnyData())
 		.AddConditional(Components->FloatChannel[3], &Curves[3], ChannelsUsed > 3 && Curves[3].HasAnyData())
 		.Commit(this, Params, OutImportedEntity);
 }
+
+/* UMovieSceneDoubleVectorSection structors
+ *****************************************************************************/
+
+UMovieSceneDoubleVectorSection::UMovieSceneDoubleVectorSection(const FObjectInitializer& ObjectInitializer)
+	: Super(ObjectInitializer)
+{
+	ChannelsUsed = 0;
+	bSupportsInfiniteRange = true;
+	BlendType = EMovieSceneBlendType::Absolute;
+}
+
+void UMovieSceneDoubleVectorSection::Serialize(FArchive& Ar)
+{
+	Super::Serialize(Ar);
+
+	if (Ar.IsLoading())
+	{
+		RecreateChannelProxy();
+	}
+}
+
+void UMovieSceneDoubleVectorSection::PostEditImport()
+{
+	Super::PostEditImport();
+
+	RecreateChannelProxy();
+}
+
+void UMovieSceneDoubleVectorSection::RecreateChannelProxy()
+{
+	FMovieSceneChannelProxyData Channels;
+
+	check(ChannelsUsed <= UE_ARRAY_COUNT(Curves));
+
+#if WITH_EDITOR
+
+	const TVectorSectionEditorData<double> EditorData(ChannelsUsed);
+	for (int32 Index = 0; Index < ChannelsUsed; ++Index)
+	{
+		Channels.Add(Curves[Index], EditorData.MetaData[Index], EditorData.ExternalValues[Index]);
+	}
+
+#else
+
+	for (int32 Index = 0; Index < ChannelsUsed; ++Index)
+	{
+		Channels.Add(Curves[Index]);
+	}
+
+#endif
+
+	ChannelProxy = MakeShared<FMovieSceneChannelProxy>(MoveTemp(Channels));
+}
+
+TSharedPtr<FStructOnScope> UMovieSceneDoubleVectorSection::GetKeyStruct(TArrayView<const FKeyHandle> KeyHandles)
+{
+	TSharedPtr<FStructOnScope> KeyStruct;
+	if (ChannelsUsed == 3)
+	{
+		KeyStruct = MakeShareable(new FStructOnScope(FMovieSceneVector3dKeyStruct::StaticStruct()));
+	}
+
+	if (KeyStruct.IsValid())
+	{
+		FMovieSceneDoubleVectorKeyStructBase* Struct = (FMovieSceneDoubleVectorKeyStructBase*)KeyStruct->GetStructMemory();
+		for (int32 Index = 0; Index < ChannelsUsed; ++Index)
+		{
+			Struct->KeyStructInterop.Add(FMovieSceneChannelValueHelper(ChannelProxy->MakeHandle<FMovieSceneDoubleChannel>(Index), Struct->GetPropertyChannelByIndex(Index), KeyHandles));
+		}
+
+		Struct->KeyStructInterop.SetStartingValues();
+		Struct->Time = Struct->KeyStructInterop.GetUnifiedKeyTime().Get(0);
+	}
+
+	return KeyStruct;
+}
+
+bool UMovieSceneDoubleVectorSection::PopulateEvaluationFieldImpl(const TRange<FFrameNumber>& EffectiveRange, const FMovieSceneEvaluationFieldEntityMetaData& InMetaData, FMovieSceneEntityComponentFieldBuilder* OutFieldBuilder)
+{
+	FMovieScenePropertyTrackEntityImportHelper::PopulateEvaluationField(*this, EffectiveRange, InMetaData, OutFieldBuilder);
+	return true;
+}
+
+void UMovieSceneDoubleVectorSection::ImportEntityImpl(UMovieSceneEntitySystemLinker* EntityLinker, const FEntityImportParams& Params, FImportedEntity* OutImportedEntity)
+{
+	using namespace UE::MovieScene;
+
+	if (ChannelsUsed <= 0)
+	{
+		return;
+	}
+	else if (!Curves[0].HasAnyData() && !Curves[1].HasAnyData() && !Curves[2].HasAnyData() && !Curves[3].HasAnyData())
+	{
+		return;
+	}
+
+	const FBuiltInComponentTypes* Components = FBuiltInComponentTypes::Get();
+	const FMovieSceneTracksComponentTypes* TracksComponents = FMovieSceneTracksComponentTypes::Get();
+
+	FPropertyTrackEntityImportHelper(TracksComponents->DoubleVector)
+		.AddConditional(Components->DoubleChannel[0], &Curves[0], ChannelsUsed > 0 && Curves[0].HasAnyData())
+		.AddConditional(Components->DoubleChannel[1], &Curves[1], ChannelsUsed > 1 && Curves[1].HasAnyData())
+		.AddConditional(Components->DoubleChannel[2], &Curves[2], ChannelsUsed > 2 && Curves[2].HasAnyData())
+		.AddConditional(Components->DoubleChannel[3], &Curves[3], ChannelsUsed > 3 && Curves[3].HasAnyData())
+		.Commit(this, Params, OutImportedEntity);
+}
diff --git a/Engine/Source/Runtime/MovieSceneTracks/Private/Systems/DoubleChannelEvaluatorSystem.cpp b/Engine/Source/Runtime/MovieSceneTracks/Private/Systems/DoubleChannelEvaluatorSystem.cpp
new file mode 100644
--- /dev/null
+++ b/Engine/Source/Runtime/MovieSceneTracks/Private/Systems/DoubleChannelEvaluatorSystem.cpp
@@ -0,0 +1,119 @@
+// Copyright Epic Games, Inc. All Rights Reserved.
+
+#include "Systems/DoubleChannelEvaluatorSystem.h"
+#include "EntitySystem/BuiltInComponentTypes.h"
+#include "EntitySystem/EntityAllocationIterator.h"
+#include "EntitySystem/MovieSceneEntitySystemTask.h"
+#include "EntitySystem/MovieSceneEntitySystemLinker.h"
+#include "EntitySystem/MovieSceneComponentRegistry.h"
+
+#include "EntitySystem/MovieSceneEvalTimeSystem.h"
+#include "Channels/MovieSceneDoubleChannel.h"
+
+#include "Algo/Find.h"
+#include "Math/NumericLimits.h"
+
+DECLARE_CYCLE_STAT(TEXT("MovieScene: Double channel system"), MovieSceneEval_DoubleChannelSystem, STATGROUP_MovieSceneECS);
+DECLARE_CYCLE_STAT(TEXT("MovieScene: Gather double channels"), MovieSceneEval_GatherDoubleChannelTask, STATGROUP_MovieSceneECS);
+DECLARE_CYCLE_STAT(TEXT("MovieScene: Evaluate double channels"), MovieSceneEval_EvaluateDoubleChannelTask, STATGROUP_MovieSceneECS);
+
+namespace UE
+{
+namespace MovieScene
+{
+
+// @todo: for multi-bindings we currently re-evaluate the double channel for each binding, even though the time is the same.
+// Do we need to optimize for this case using something like the code below, while pessimizing the common (non-multi-bind) codepath??
+struct FEvaluateDoubleChannels
+{
+	void ForEachEntity(FSourceDoubleChannel DoubleChannel, FFrameTime FrameTime, double& OutResult, FSourceDoubleChannelFlags& OutFlags)
+	{
+		if (OutFlags.bNeedsEvaluate == false)
+		{
+			return;
+		}
+
+		if (!DoubleChannel.Source->Evaluate(FrameTime, OutResult))
+		{
+			OutResult = MIN_dbl;
+		}
+
+		if (DoubleChannel.Source->GetTimes().Num() <= 1)
+		{
+			OutFlags.bNeedsEvaluate = false;
+		}
+	}
+};
+
+} // namespace MovieScene
+} // namespace UE
+
+
+TArray<UDoubleChannelEvaluatorSystem::FChannelType, TInlineAllocator<4>> UDoubleChannelEvaluatorSystem::StaticChannelTypes;
+
+UDoubleChannelEvaluatorSystem::UDoubleChannelEvaluatorSystem(const FObjectInitializer& ObjInit)
+	: Super(ObjInit)
+{
+	using namespace UE::MovieScene;
+
+	if (HasAnyFlags(RF_ClassDefaultObject))
+	{
+		DefineImplicitPrerequisite(UMovieSceneEvalTimeSystem::StaticClass(), GetClass());
+
+		FBuiltInComponentTypes* Components = FBuiltInComponentTypes::Get();
+
+		for (int32 Index = 0; Index < UE_ARRAY_COUNT(Components->DoubleChannel); ++Index)
+		{
+			RegisterChannelType(Components->DoubleChannel[Index], Components->DoubleChannelFlags[Index], Components->DoubleResult[Index]);
+		}
+	}
+}
+
+void UDoubleChannelEvaluatorSystem::RegisterChannelType(TComponentTypeID<UE::MovieScene::FSourceDoubleChannel> SourceChannelType, TComponentTypeID<UE::MovieScene::FSourceDoubleChannelFlags> ChannelFlagsType, TComponentTypeID<double> ResultType)
+{
+	using namespace UE::MovieScene;
+
+	FChannelType ChannelType;
+	ChannelType.ChannelType = SourceChannelType;
+	ChannelType.ChannelFlagsType = ChannelFlagsType;
+	ChannelType.ResultType  = ResultType;
+
+	DefineComponentProducer(UDoubleChannelEvaluatorSystem::StaticClass(), ResultType);
+
+	StaticChannelTypes.Add(ChannelType);
+}
+
+bool UDoubleChannelEvaluatorSystem::IsRelevantImpl(UMovieSceneEntitySystemLinker* InLinker) const
+{
+	for (const FChannelType& ChannelType : StaticChannelTypes)
+	{
+		if (InLinker->EntityManager.ContainsComponent(ChannelType.ChannelType))
+		{
+			return true;
+		}
+	}
+
+	return false;
+}
+
+void UDoubleChannelEvaluatorSystem::OnRun(FSystemTaskPrerequisites& InPrerequisites, FSystemSubsequentTasks& Subsequents)
+{
+	using namespace UE::MovieScene;
+
+	SCOPE_CYCLE_COUNTER(MovieSceneEval_DoubleChannelSystem)
+
+	FBuiltInComponentTypes* BuiltInComponents = FBuiltInComponentTypes::Get();
+
+	for (const FChannelType& ChannelType : StaticChannelTypes)
+	{
+		// Evaluate double channels per instance and write the evaluated value into the output
+		FEntityTaskBuilder()
+		.Read(ChannelType.ChannelType)
+		.Read(BuiltInComponents->EvalTime)
+		.Write(ChannelType.ResultType)
+		.Write(ChannelType.ChannelFlagsType)
+		.FilterNone({ BuiltInComponents->Tags.Ignored })
+		.SetStat(GET_STATID(MovieSceneEval_EvaluateDoubleChannelTask))
+		.Dispatch_PerEntity<FEvaluateDoubleChannels>(&Linker->EntityManager, InPrerequisites, &Subsequents);
+	}
+}
diff --git a/Engine/Source/Runtime/MovieSceneTracks/Private/Systems/MovieScene3DTransformPropertySystem.cpp b/Engine/Source/Runtime/MovieSceneTracks/Private/Systems/MovieScene3DTransformPropertySystem.cpp
--- a/Engine/Source/Runtime/MovieSceneTracks/Private/Systems/MovieScene3DTransformPropertySystem.cpp
+++ b/Engine/Source/Runtime/MovieSceneTracks/Private/Systems/MovieScene3DTransformPropertySystem.cpp
@@ -3,8 +3,8 @@
 #include "Systems/MovieScene3DTransformPropertySystem.h"
 
 #include "Systems/MovieScenePropertyInstantiator.h"
-#include "Systems/FloatChannelEvaluatorSystem.h"
-#include "Systems/MovieScenePiecewiseFloatBlenderSystem.h"
+#include "Systems/DoubleChannelEvaluatorSystem.h"
+#include "Systems/MovieScenePiecewiseDoubleBlenderSystem.h"
 
 #include "MovieSceneTracksComponentTypes.h"
 
@@ -17,12 +17,12 @@ UMovieScene3DTransformPropertySystem::UMovieScene3DTransformPropertySystem(const
 
 	if (HasAnyFlags(RF_ClassDefaultObject))
 	{
-		DefineImplicitPrerequisite(UMovieScenePiecewiseFloatBlenderSystem::StaticClass(), GetClass());
-		DefineImplicitPrerequisite(UFloatChannelEvaluatorSystem::StaticClass(), GetClass());
+		DefineImplicitPrerequisite(UMovieScenePiecewiseDoubleBlenderSystem::StaticClass(), GetClass());
+		DefineImplicitPrerequisite(UDoubleChannelEvaluatorSystem::StaticClass(), GetClass());
 	}
 }
 
 void UMovieScene3DTransformPropertySystem::OnRun(FSystemTaskPrerequisites& InPrerequisites, FSystemSubsequentTasks& Subsequents)
 {
 	Super::OnRun(InPrerequisites, Subsequents);
-}
\ No newline at end of file
+}
diff --git a/Engine/Source/Runtime/MovieSceneTracks/Private/Systems/MovieSceneComponentTransformSystem.cpp b/Engine/Source/Runtime/MovieSceneTracks/Private/Systems/MovieSceneComponentTransformSystem.cpp
--- a/Engine/Source/Runtime/MovieSceneTracks/Private/Systems/MovieSceneComponentTransformSystem.cpp
+++ b/Engine/Source/Runtime/MovieSceneTracks/Private/Systems/MovieSceneComponentTransformSystem.cpp
@@ -1,8 +1,8 @@
 // Copyright Epic Games, Inc. All Rights Reserved.
 
 #include "Systems/MovieSceneComponentTransformSystem.h"
-#include "Systems/FloatChannelEvaluatorSystem.h"
-#include "Systems/MovieScenePiecewiseFloatBlenderSystem.h"
+#include "Systems/DoubleChannelEvaluatorSystem.h"
+#include "Systems/MovieScenePiecewiseDoubleBlenderSystem.h"
 #include "Systems/MovieScenePropertyInstantiator.h"
 
 #include "EntitySystem/BuiltInComponentTypes.h"
@@ -19,8 +19,8 @@ UMovieSceneComponentTransformSystem::UMovieSceneComponentTransformSystem(const F
 
 	if (HasAnyFlags(RF_ClassDefaultObject))
 	{
-		DefineImplicitPrerequisite(UMovieScenePiecewiseFloatBlenderSystem::StaticClass(), GetClass());
-		DefineImplicitPrerequisite(UFloatChannelEvaluatorSystem::StaticClass(), GetClass());
+		DefineImplicitPrerequisite(UMovieScenePiecewiseDoubleBlenderSystem::StaticClass(), GetClass());
+		DefineImplicitPrerequisite(UDoubleChannelEvaluatorSystem::StaticClass(), GetClass());
 
 		DefineComponentConsumer(GetClass(), UE::MovieScene::FMovieSceneTracksComponentTypes::Get()->ComponentTransform.PropertyTag);
 	}
diff --git a/Engine/Source/Runtime/MovieSceneTracks/Private/Systems/MovieSceneEulerTransformPropertySystem.cpp b/Engine/Source/Runtime/MovieSceneTracks/Private/Systems/MovieSceneEulerTransformPropertySystem.cpp
--- a/Engine/Source/Runtime/MovieSceneTracks/Private/Systems/MovieSceneEulerTransformPropertySystem.cpp
+++ b/Engine/Source/Runtime/MovieSceneTracks/Private/Systems/MovieSceneEulerTransformPropertySystem.cpp
@@ -3,8 +3,8 @@
 #include "Systems/MovieSceneEulerTransformPropertySystem.h"
 
 #include "Systems/MovieScenePropertyInstantiator.h"
-#include "Systems/FloatChannelEvaluatorSystem.h"
-#include "Systems/MovieScenePiecewiseFloatBlenderSystem.h"
+#include "Systems/DoubleChannelEvaluatorSystem.h"
+#include "Systems/MovieScenePiecewiseDoubleBlenderSystem.h"
 
 #include "MovieSceneTracksComponentTypes.h"
 
@@ -17,12 +17,12 @@ UMovieSceneEulerTransformPropertySystem::UMovieSceneEulerTransformPropertySystem
 
 	if (HasAnyFlags(RF_ClassDefaultObject))
 	{
-		DefineImplicitPrerequisite(UMovieScenePiecewiseFloatBlenderSystem::StaticClass(), GetClass());
-		DefineImplicitPrerequisite(UFloatChannelEvaluatorSystem::StaticClass(), GetClass());
+		DefineImplicitPrerequisite(UMovieScenePiecewiseDoubleBlenderSystem::StaticClass(), GetClass());
+		DefineImplicitPrerequisite(UDoubleChannelEvaluatorSystem::StaticClass(), GetClass());
 	}
 }
 
 void UMovieSceneEulerTransformPropertySystem::OnRun(FSystemTaskPrerequisites& InPrerequisites, FSystemSubsequentTasks& Subsequents)
 {
 	Super::OnRun(InPrerequisites, Subsequents);
-}
\ No newline at end of file
+}
diff --git a/Engine/Source/Runtime/MovieSceneTracks/Private/Systems/MovieScenePiecewiseBlenderSystemImpl.cpp b/Engine/Source/Runtime/MovieSceneTracks/Private/Systems/MovieScenePiecewiseBlenderSystemImpl.cpp
new file mode 100644
--- /dev/null
+++ b/Engine/Source/Runtime/MovieSceneTracks/Private/Systems/MovieScenePiecewiseBlenderSystemImpl.cpp
@@ -0,0 +1,713 @@
+// Copyright Epic Games, Inc. All Rights Reserved.
+
+#include "Systems/MovieScenePiecewiseBlenderSystemImpl.h"
+#include "EntitySystem/BuiltInComponentTypes.h"
+#include "EntitySystem/EntityAllocationIterator.h"
+#include "EntitySystem/MovieSceneEntitySystemTask.h"
+#include "EntitySystem/MovieSceneEntityManager.h"
+#include "Systems/MovieScenePiecewiseDoubleBlenderSystem.h"
+#include "Systems/MovieScenePiecewiseFloatBlenderSystem.h"
+#include "Systems/WeightAndEasingEvaluatorSystem.h"
+
+namespace UE
+{
+namespace MovieScene
+{
+
+/** Traits for known blendable values (float and double) */
+template<typename ValueType>
+struct TPiecewiseBlendableValueTraits;
+
+template<>
+struct TPiecewiseBlendableValueTraits<float>
+{
+	static const FMovieSceneBlenderSystemID GetBlenderSystemID()
+	{
+		return UMovieSceneBlenderSystem::GetBlenderSystemID<UMovieScenePiecewiseFloatBlenderSystem>();
+	}
+	static bool HasAnyComposite(const FPropertyDefinition& PropertyDefinition)
+	{
+		return PropertyDefinition.FloatCompositeMask != 0;
+	}
+	static bool IsCompositeSupported(const FPropertyDefinition& PropertyDefinition, int32 CompositeIndex)
+	{
+		return ((PropertyDefinition.FloatCompositeMask & (1 << CompositeIndex)) != 0);
+	}
+	static TArrayView<TComponentTypeID<float>> GetBaseComponents()
+	{
+		return TArrayView<TComponentTypeID<float>>(FBuiltInComponentTypes::Get()->BaseFloat);
+	}
+	static TArrayView<TComponentTypeID<float>> GetResultComponents()
+	{
+		return TArrayView<TComponentTypeID<float>>(FBuiltInComponentTypes::Get()->FloatResult);
+	}
+};
+
+template<>
+struct TPiecewiseBlendableValueTraits<double>
+{
+	static const FMovieSceneBlenderSystemID GetBlenderSystemID()
+	{
+		return UMovieSceneBlenderSystem::GetBlenderSystemID<UMovieScenePiecewiseDoubleBlenderSystem>();
+	}
+	static bool HasAnyComposite(const FPropertyDefinition& PropertyDefinition)
+	{
+		return PropertyDefinition.DoubleCompositeMask != 0;
+	}
+	static bool IsCompositeSupported(const FPropertyDefinition& PropertyDefinition, int32 CompositeIndex)
+	{
+		return ((PropertyDefinition.DoubleCompositeMask & (1 << CompositeIndex)) != 0);
+	}
+	static TArrayView<TComponentTypeID<double>> GetBaseComponents()
+	{
+		return TArrayView<TComponentTypeID<double>>(FBuiltInComponentTypes::Get()->BaseDouble);
+	}
+	static TArrayView<TComponentTypeID<double>> GetResultComponents()
+	{
+		return TArrayView<TComponentTypeID<double>>(FBuiltInComponentTypes::Get()->DoubleResult);
+	}
+};
+
+/** Task for accumulating all weighted blend inputs into arrays based on BlendID. Will be run for Absolute, Additive and Relative blend modes*/
+template<typename ValueType>
+struct TAccumulationTask
+{
+	using FBlendResult = TBlendResult<ValueType>;
+
+	TAccumulationTask(TSortedMap<FComponentTypeID, TArray<FBlendResult>>* InAccumulationBuffers)
+		: AccumulationBuffers(InAccumulationBuffers)
+	{}
+
+	/** Task entry point - iterates the allocation's headers and accumulates results for any required components */
+	void ForEachAllocation(FEntityAllocationIteratorItem InItem, TRead<FMovieSceneBlendChannelID> BlendIDs, TReadOptional<float> OptionalEasingAndWeights)
+	{
+		const FEntityAllocation* Allocation = InItem;
+		const FComponentMask& AllocationType = InItem;
+
+		for (const FComponentHeader& ComponentHeader : Allocation->GetComponentHeaders())
+		{
+			if (TArray<FBlendResult>* AccumulationBuffer = AccumulationBuffers->Find(ComponentHeader.ComponentType))
+			{
+				ComponentHeader.ReadWriteLock.ReadLock();
+
+				const ValueType* Results = static_cast<const ValueType*>(ComponentHeader.GetValuePtr(0));
+				AccumulateResults(Allocation, Results, BlendIDs, OptionalEasingAndWeights, *AccumulationBuffer);
+
+				ComponentHeader.ReadWriteLock.ReadUnlock();
+			}
+		}
+	}
+
+private:
+
+	void AccumulateResults(const FEntityAllocation* InAllocation, const ValueType* InResults, const FMovieSceneBlendChannelID* BlendIDs, const float* OptionalEasingAndWeights, TArray<FBlendResult>& OutBlendResults)
+	{
+		static const FMovieSceneBlenderSystemID BlenderSystemID = TPiecewiseBlendableValueTraits<ValueType>::GetBlenderSystemID();
+
+		const int32 Num = InAllocation->Num();
+		if (OptionalEasingAndWeights)
+		{
+			// We have some easing/weight factors to multiply values with.
+			for (int32 Index = 0; Index < Num; ++Index)
+			{
+				const FMovieSceneBlendChannelID& BlendID(BlendIDs[Index]);
+				ensureMsgf(BlendID.SystemID == BlenderSystemID, TEXT("Overriding the standard blender system of standard types isn't supported."));
+
+				FBlendResult& Result = OutBlendResults[BlendID.ChannelID];
+
+				const float Weight = OptionalEasingAndWeights[Index];
+				Result.Total  += InResults[Index] * Weight;
+				Result.Weight += Weight;
+			}
+		}
+		else
+		{
+			// Faster path for when there's no weight to multiply values with.
+			for (int32 Index = 0; Index < Num; ++Index)
+			{
+				const FMovieSceneBlendChannelID& BlendID(BlendIDs[Index]);
+				ensureMsgf(BlendID.SystemID == BlenderSystemID, TEXT("Overriding the standard blender system of standard types isn't supported."));
+
+				FBlendResult& Result = OutBlendResults[BlendID.ChannelID];
+				Result.Total  += InResults[Index];
+				Result.Weight += 1.f;
+			}
+		}
+	}
+
+	TSortedMap<FComponentTypeID, TArray<FBlendResult>>* AccumulationBuffers;
+};
+
+/** Same as the task above, but also reads a "base value" that is subtracted from all values.
+ *
+ *  Only used by entities with the "additive from base" blend type.
+ */
+template<typename ValueType>
+struct TAdditiveFromBaseBlendTask
+{
+	using FAdditiveFromBaseBuffer = TAdditiveFromBaseBuffer<ValueType>;
+	using FBlendResult = TBlendResult<ValueType>;
+
+	TSortedMap<FComponentTypeID, FAdditiveFromBaseBuffer>* AccumulationBuffers;
+
+	void ForEachAllocation(FEntityAllocationIteratorItem InItem, TRead<FMovieSceneBlendChannelID> BlendIDs, TReadOptional<float> EasingAndWeightResults)
+	{
+		FEntityAllocation* Allocation = InItem;
+		const FComponentMask& AllocationType = InItem;
+
+		for (const FComponentHeader& ComponentHeader : Allocation->GetComponentHeaders())
+		{
+			if (FAdditiveFromBaseBuffer* Buffer = AccumulationBuffers->Find(ComponentHeader.ComponentType))
+			{
+				TComponentReader<ValueType> BaseValues = Allocation->ReadComponents(Buffer->BaseComponent.template ReinterpretCast<ValueType>());
+				TComponentReader<ValueType> Results(&ComponentHeader);
+
+				AccumulateResults(Allocation, Results.AsPtr(), BaseValues.AsPtr(), BlendIDs, EasingAndWeightResults, Buffer->Buffer);
+			}
+		}
+	}
+
+private:
+
+	void AccumulateResults(const FEntityAllocation* InAllocation, const ValueType* Results, const ValueType* BaseValues, const FMovieSceneBlendChannelID* BlendIDs, const float* OptionalEasingAndWeights, TArray<FBlendResult>& OutBlendResults)
+	{
+		static const FMovieSceneBlenderSystemID BlenderSystemID = TPiecewiseBlendableValueTraits<ValueType>::GetBlenderSystemID();
+
+		const int32 Num = InAllocation->Num();
+
+		if (OptionalEasingAndWeights)
+		{
+			for (int32 Index = 0; Index < Num; ++Index)
+			{
+				const FMovieSceneBlendChannelID& BlendID(BlendIDs[Index]);
+				ensureMsgf(BlendID.SystemID == BlenderSystemID, TEXT("Overriding the standard blender system of standard types isn't supported."));
+
+				FBlendResult& Result = OutBlendResults[BlendID.ChannelID];
+
+				const float Weight = OptionalEasingAndWeights[Index];
+				Result.Total  += (Results[Index] - BaseValues[Index]) * Weight;
+				Result.Weight += Weight;
+			}
+		}
+		else
+		{
+			// Faster path for when there's no weight to multiply values with.
+			for (int32 Index = 0; Index < Num; ++Index)
+			{
+				const FMovieSceneBlendChannelID& BlendID(BlendIDs[Index]);
+				ensureMsgf(BlendID.SystemID == BlenderSystemID, TEXT("Overriding the standard blender system of standard types isn't supported."));
+
+				FBlendResult& Result = OutBlendResults[BlendID.ChannelID];
+				Result.Total  += (Results[Index] - BaseValues[Index]);
+				Result.Weight += 1.f;
+			}
+		}
+	}
+};
+
+/** Task that combines all accumulated blends for any tracked property type that has blend inputs/outputs */
+template<typename ValueType>
+struct TCombineBlends
+{
+	using FBlendResult = TBlendResult<ValueType>;
+	using FAccumulationResult = TAccumulationResult<ValueType>;
+	using FAccumulationBuffers = TAccumulationBuffers<ValueType>;
+	using FPiecewiseBlendableValueTraits = TPiecewiseBlendableValueTraits<ValueType>;
+
+	explicit TCombineBlends(const TBitArray<>& InCachedRelevantProperties, const FAccumulationBuffers* InAccumulationBuffers, FEntityAllocationWriteContext InWriteContext)
+		: CachedRelevantProperties(InCachedRelevantProperties)
+		, AccumulationBuffers(InAccumulationBuffers)
+		, PropertyRegistry(&FBuiltInComponentTypes::Get()->PropertyRegistry)
+		, WriteContext(InWriteContext)
+	{}
+
+	void ForEachAllocation(FEntityAllocationIteratorItem InItem, TRead<FMovieSceneBlendChannelID> BlendIDs)
+	{
+		FEntityAllocation* Allocation = InItem;
+		const FComponentMask& AllocationType = InItem;
+
+		// Find out what kind of property this is
+		for (TConstSetBitIterator<> PropertyIndex(CachedRelevantProperties); PropertyIndex; ++PropertyIndex)
+		{
+			const FPropertyDefinition& PropertyDefinition = PropertyRegistry->GetDefinition(FCompositePropertyTypeID::FromIndex(PropertyIndex.GetIndex()));
+			if (AllocationType.Contains(PropertyDefinition.PropertyType))
+			{
+				ProcessPropertyType(Allocation, AllocationType, PropertyDefinition, BlendIDs);
+				return;
+			}
+		}
+	}
+
+private:
+
+	void ProcessPropertyType(FEntityAllocation* Allocation, const FComponentMask& AllocationType, const FPropertyDefinition& PropertyDefinition, const FMovieSceneBlendChannelID* BlendIDs)
+	{
+		static const FMovieSceneBlenderSystemID BlenderSystemID = TPiecewiseBlendableValueTraits<ValueType>::GetBlenderSystemID();
+
+		TArrayView<const FPropertyCompositeDefinition> Composites = PropertyRegistry->GetComposites(PropertyDefinition);
+
+		FOptionalComponentReader OptInitialValues = Allocation->TryReadComponentsErased(PropertyDefinition.InitialValueType);
+
+		for (int32 CompositeIndex = 0; CompositeIndex < Composites.Num(); ++CompositeIndex)
+		{
+			if (!FPiecewiseBlendableValueTraits::IsCompositeSupported(PropertyDefinition, CompositeIndex))
+			{
+				continue;
+			}
+
+			TComponentTypeID<ValueType> ResultComponent = Composites[CompositeIndex].ComponentTypeID.template ReinterpretCast<ValueType>();
+			if (!AllocationType.Contains(ResultComponent))
+			{
+				continue;
+			}
+
+			FAccumulationResult Results = AccumulationBuffers->FindResults(ResultComponent);
+			if (Results.IsValid())
+			{
+				const uint16 InitialValueProjectionOffset = Composites[CompositeIndex].CompositeOffset;
+
+				// Open the result channel for write
+				TComponentWriter<ValueType> ValueResults = Allocation->WriteComponents(ResultComponent, WriteContext);
+
+				if (OptInitialValues)
+				{
+					for (int32 Index = 0; Index < Allocation->Num(); ++Index)
+					{
+						ensureMsgf(BlendIDs[Index].SystemID == BlenderSystemID, TEXT("Overriding the standard blender system of standard types isn't supported."));
+						const ValueType InitialValue = *reinterpret_cast<const ValueType*>(static_cast<const uint8*>(OptInitialValues[Index]) + InitialValueProjectionOffset);
+						BlendResultsWithInitial(Results, BlendIDs[Index].ChannelID, InitialValue, ValueResults[Index]);
+					}
+				}
+				else
+				{
+					for (int32 Index = 0; Index < Allocation->Num(); ++Index)
+					{
+						ensureMsgf(BlendIDs[Index].SystemID == BlenderSystemID, TEXT("Overriding the standard blender system of standard types isn't supported."));
+						BlendResults(Results, BlendIDs[Index].ChannelID, ValueResults[Index]);
+					}
+				}
+			}
+		}
+	}
+
+	void BlendResultsWithInitial(const FAccumulationResult& Results, uint16 BlendID, const ValueType InitialValue, ValueType& OutFinalBlendResult)
+	{
+		FBlendResult AbsoluteResult = Results.GetAbsoluteResult(BlendID);
+		FBlendResult RelativeResult = Results.GetRelativeResult(BlendID);
+		FBlendResult AdditiveResult = Results.GetAdditiveResult(BlendID);
+		FBlendResult AdditiveFromBaseResult = Results.GetAdditiveFromBaseResult(BlendID);
+
+		if (RelativeResult.Weight != 0)
+		{
+			RelativeResult.Total += InitialValue * RelativeResult.Weight;
+		}
+
+		FBlendResult TotalAdditiveResult = { AdditiveResult.Total + AdditiveFromBaseResult.Total, AdditiveResult.Weight + AdditiveFromBaseResult.Weight };
+
+		const float TotalWeight = AbsoluteResult.Weight + RelativeResult.Weight;
+		if (TotalWeight != 0)
+		{
+			// If the absolute value has some partial weighting (for ease-in/out for instance), we ramp it from/to the initial value. This means
+			// that the "initial value" adds a contribution to the entire blending process, so we add its weight to the total that we
+			// normalize absolutes and relatives with.
+			//
+			// Note that "partial weighting" means strictly between 0 and 100%. At 100% and above, we don't need to do this thing with the initial
+			// value. At 0%, we have no absolute value (only a relative value) and we therefore don't want to include the initial value either.
+			const bool bInitialValueContributes = (0.f < AbsoluteResult.Weight && AbsoluteResult.Weight < 1.f);
+			const ValueType AbsoluteBlendedValue = bInitialValueContributes ?
+				(InitialValue * (1.f - AbsoluteResult.Weight) + AbsoluteResult.Total) :
+				AbsoluteResult.Total;
+			const float FinalTotalWeight = bInitialValueContributes ? (TotalWeight + (1.f - AbsoluteResult.Weight)) : TotalWeight;
+
+			const ValueType Value = (AbsoluteBlendedValue + RelativeResult.Total) / FinalTotalWeight + TotalAdditiveResult.Total;
+			OutFinalBlendResult = Value;
+		}
+		else if (TotalAdditiveResult.Weight != 0)
+		{
+			OutFinalBlendResult = TotalAdditiveResult.Total + InitialValue;
+		}
+		else
+		{
+			OutFinalBlendResult = InitialValue;
+		}
+	}
+
+	void BlendResults(const FAccumulationResult& Results, uint16 BlendID, ValueType& OutFinalBlendResult)
+	{
+		FBlendResult AbsoluteResult = Results.GetAbsoluteResult(BlendID);
+		FBlendResult AdditiveResult = Results.GetAdditiveResult(BlendID);
+		FBlendResult AdditiveFromBaseResult = Results.GetAdditiveFromBaseResult(BlendID);
+
+#if DO_GUARD_SLOW
+		ensureMsgf(AbsoluteResult.Weight != 0.f, TEXT("Default blend combine being used for an entity that has no absolute weight. This should have an initial value and should be handled by each system, and excluded by default with UMovieSceneBlenderSystem::FinalCombineExclusionFilter ."));
+#endif
+
+		const float TotalWeight = AbsoluteResult.Weight;
+		if (TotalWeight != 0)
+		{
+			const ValueType Value = AbsoluteResult.Total / AbsoluteResult.Weight + AdditiveResult.Total + AdditiveFromBaseResult.Total;
+			OutFinalBlendResult = Value;
+		}
+	}
+
+private:
+
+	TBitArray<> CachedRelevantProperties;
+	const FAccumulationBuffers* AccumulationBuffers;
+	const FPropertyRegistry* PropertyRegistry;
+	FEntityAllocationWriteContext WriteContext;
+};
+
+template<typename ValueType>
+bool TAccumulationBuffers<ValueType>::IsEmpty() const
+{
+	return Absolute.Num() == 0 && Relative.Num() == 0 && Additive.Num() == 0 && AdditiveFromBase.Num() == 0;
+}
+
+template<typename ValueType>
+void TAccumulationBuffers<ValueType>::Reset()
+{
+	Absolute.Empty();
+	Relative.Empty();
+	Additive.Empty();
+	AdditiveFromBase.Empty();
+}
+
+template<typename ValueType>
+typename TAccumulationBuffers<ValueType>::FAccumulationResult TAccumulationBuffers<ValueType>::FindResults(FComponentTypeID InComponentType) const
+{
+	FAccumulationResult Result;
+	if (const TArray<FBlendResult>* Absolutes = Absolute.Find(InComponentType))
+	{
+		Result.Absolutes = Absolutes->GetData();
+	}
+	if (const TArray<FBlendResult>* Relatives = Relative.Find(InComponentType))
+	{
+		Result.Relatives = Relatives->GetData();
+	}
+	if (const TArray<FBlendResult>* Additives = Additive.Find(InComponentType))
+	{
+		Result.Additives = Additives->GetData();
+	}
+	if (const FAdditiveFromBaseBuffer* AdditivesFromBase = AdditiveFromBase.Find(InComponentType))
+	{
+		Result.AdditivesFromBase = AdditivesFromBase->Buffer.GetData();
+	}
+	return Result;
+}
+
+template<typename ValueType>
+void TPiecewiseBlenderSystemImpl<ValueType>::Run(FPiecewiseBlenderSystemImplRunParams& Params, FEntityManager& EntityManager, FSystemTaskPrerequisites& InPrerequisites, FSystemSubsequentTasks& Subsequents)
+{
+	using FAccumulationTask = TAccumulationTask<ValueType>;
+	using FAdditiveFromBaseBlendTask = TAdditiveFromBaseBlendTask<ValueType>;
+	using FCombineBlends = TCombineBlends<ValueType>;
+
+	// We allocate space for every blend even if there are gaps so we can do a straight index into each array
+	if (Params.MaximumNumBlends == 0)
+	{
+		return;
+	}
+
+	// Update cached channel data if necessary
+	if (ChannelRelevancyCache.Update(EntityManager) == ECachedEntityManagerState::Stale)
+	{
+		ReinitializeAccumulationBuffers(Params.MaximumNumBlends, EntityManager);
+	}
+
+	if (AccumulationBuffers.IsEmpty())
+	{
+		return;
+	}
+
+	ZeroAccumulationBuffers();
+
+	FBuiltInComponentTypes* BuiltInComponents = FBuiltInComponentTypes::Get();
+
+	FSystemTaskPrerequisites Prereqs;
+	if (AccumulationBuffers.Absolute.Num() != 0)
+	{
+		FGraphEventRef Task = FEntityTaskBuilder()
+		.Read(BuiltInComponents->BlendChannelInput)
+		.ReadOptional(BuiltInComponents->WeightAndEasingResult)
+		.FilterAll({ BuiltInComponents->Tags.AbsoluteBlend })
+		.FilterAny(BlendedResultMask)
+		.FilterNone({ BuiltInComponents->Tags.Ignored })
+		.SetStat(Params.BlendValuesStatId)
+		.template Dispatch_PerAllocation<FAccumulationTask>(&EntityManager, InPrerequisites, nullptr, &AccumulationBuffers.Absolute);
+
+		if (Task)
+		{
+			Prereqs.AddMasterTask(Task);
+		}
+	}
+
+	if (AccumulationBuffers.Relative.Num() != 0)
+	{
+		FGraphEventRef Task = FEntityTaskBuilder()
+		.Read(BuiltInComponents->BlendChannelInput)
+		.ReadOptional(BuiltInComponents->WeightAndEasingResult)
+		.FilterAll({ BuiltInComponents->Tags.RelativeBlend })
+		.FilterAny(BlendedResultMask)
+		.FilterNone({ BuiltInComponents->Tags.Ignored })
+		.SetStat(Params.BlendValuesStatId)
+		.template Dispatch_PerAllocation<FAccumulationTask>(&EntityManager, InPrerequisites, nullptr, &AccumulationBuffers.Relative);
+
+		if (Task)
+		{
+			Prereqs.AddMasterTask(Task);
+		}
+	}
+
+	if (AccumulationBuffers.Additive.Num() != 0)
+	{
+		FGraphEventRef Task = FEntityTaskBuilder()
+		.Read(BuiltInComponents->BlendChannelInput)
+		.ReadOptional(BuiltInComponents->WeightAndEasingResult)
+		.FilterAll({ BuiltInComponents->Tags.AdditiveBlend })
+		.FilterAny(BlendedResultMask)
+		.FilterNone({ BuiltInComponents->Tags.Ignored })
+		.SetStat(Params.BlendValuesStatId)
+		.template Dispatch_PerAllocation<FAccumulationTask>(&EntityManager, InPrerequisites, nullptr, &AccumulationBuffers.Additive);
+
+		if (Task)
+		{
+			Prereqs.AddMasterTask(Task);
+		}
+	}
+
+	if (AccumulationBuffers.AdditiveFromBase.Num() != 0)
+	{
+		FGraphEventRef Task = FEntityTaskBuilder()
+		.Read(BuiltInComponents->BlendChannelInput)
+		.ReadOptional(BuiltInComponents->WeightAndEasingResult)
+		.FilterAll({ BuiltInComponents->Tags.AdditiveFromBaseBlend })
+		.FilterAny(BlendedResultMask)
+		.FilterNone({ BuiltInComponents->Tags.Ignored })
+		.SetStat(Params.BlendValuesStatId)
+		.template Dispatch_PerAllocation<FAdditiveFromBaseBlendTask>(&EntityManager, InPrerequisites, nullptr, &AccumulationBuffers.AdditiveFromBase);
+
+		if (Task)
+		{
+			Prereqs.AddMasterTask(Task);
+		}
+	}
+
+	// Master task that performs the actual blends
+	FEntityTaskBuilder()
+	.Read(BuiltInComponents->BlendChannelOutput)
+	.FilterAny(BlendedPropertyMask)
+	.SetStat(Params.CombineBlendsStatId)
+	.template Dispatch_PerAllocation<FCombineBlends>(&EntityManager, Prereqs, &Subsequents, CachedRelevantProperties, &AccumulationBuffers, FEntityAllocationWriteContext(EntityManager));
+}
+
+template<typename ValueType>
+void TPiecewiseBlenderSystemImpl<ValueType>::ReinitializeAccumulationBuffers(int32 MaximumNumBlends, FEntityManager& EntityManager)
+{
+	using FPiecewiseBlendableValueTraits = TPiecewiseBlendableValueTraits<ValueType>;
+
+	FBuiltInComponentTypes* BuiltInComponents = FBuiltInComponentTypes::Get();
+
+	BlendedResultMask.Reset();
+	AccumulationBuffers.Reset();
+
+	TArrayView<TComponentTypeID<ValueType>> BaseComponents = FPiecewiseBlendableValueTraits::GetBaseComponents();
+	TArrayView<TComponentTypeID<ValueType>> ResultComponents = FPiecewiseBlendableValueTraits::GetResultComponents();
+	check(BaseComponents.Num() == ResultComponents.Num());
+
+	// Recompute which result types are blended
+	const int32 NumResults = ResultComponents.Num();
+	for (int32 Index = 0; Index < NumResults; ++Index)
+	{
+		TComponentTypeID<ValueType> Component = ResultComponents[Index];
+
+		const bool bHasAbsolutes         = EntityManager.Contains(FEntityComponentFilter().All({ Component, BuiltInComponents->BlendChannelInput, BuiltInComponents->Tags.AbsoluteBlend }));
+		const bool bHasRelatives         = EntityManager.Contains(FEntityComponentFilter().All({ Component, BuiltInComponents->BlendChannelInput, BuiltInComponents->Tags.RelativeBlend }));
+		const bool bHasAdditives         = EntityManager.Contains(FEntityComponentFilter().All({ Component, BuiltInComponents->BlendChannelInput, BuiltInComponents->Tags.AdditiveBlend }));
+		const bool bHasAdditivesFromBase = EntityManager.Contains(FEntityComponentFilter().All({ Component, BuiltInComponents->BlendChannelInput, BuiltInComponents->Tags.AdditiveFromBaseBlend }));
+
+		if (!(bHasAbsolutes || bHasRelatives || bHasAdditives || bHasAdditivesFromBase))
+		{
+			continue;
+		}
+
+		BlendedResultMask.Set(Component);
+
+		if (bHasAbsolutes)
+		{
+			TArray<FBlendResult>& Buffer = AccumulationBuffers.Absolute.Add(Component);
+			Buffer.SetNum(MaximumNumBlends);
+		}
+		if (bHasRelatives)
+		{
+			TArray<FBlendResult>& Buffer = AccumulationBuffers.Relative.Add(Component);
+			Buffer.SetNum(MaximumNumBlends);
+		}
+		if (bHasAdditives)
+		{
+			TArray<FBlendResult>& Buffer = AccumulationBuffers.Additive.Add(Component);
+			Buffer.SetNum(MaximumNumBlends);
+		}
+		if (bHasAdditivesFromBase)
+		{
+			FAdditiveFromBaseBuffer& Buffer = AccumulationBuffers.AdditiveFromBase.Add(Component);
+			Buffer.Buffer.SetNum(MaximumNumBlends);
+			Buffer.BaseComponent = BaseComponents[Index];
+		}
+	}
+
+	// Update property relevancy
+	CachedRelevantProperties.Empty();
+
+	// If we have no accumulation buffers, we have nothing more to do
+	if (AccumulationBuffers.IsEmpty())
+	{
+		return;
+	}
+
+	// This code works on the assumption that properties can never be removed (which is safe)
+	FEntityComponentFilter InclusionFilter;
+	TArrayView<const FPropertyDefinition> Properties = BuiltInComponents->PropertyRegistry.GetProperties();
+	for (int32 PropertyTypeIndex = 0; PropertyTypeIndex < Properties.Num(); ++PropertyTypeIndex)
+	{
+		const FPropertyDefinition& PropertyDefinition = Properties[PropertyTypeIndex];
+		if (FPiecewiseBlendableValueTraits::HasAnyComposite(PropertyDefinition))
+		{
+			InclusionFilter.Reset();
+			InclusionFilter.All({ BuiltInComponents->BlendChannelOutput, PropertyDefinition.PropertyType });
+			if (EntityManager.Contains(InclusionFilter))
+			{
+				CachedRelevantProperties.PadToNum(PropertyTypeIndex + 1, false);
+				CachedRelevantProperties[PropertyTypeIndex] = true;
+
+				BlendedPropertyMask.Set(PropertyDefinition.PropertyType);
+			}
+		}
+	}
+}
+
+template<typename ValueType>
+void TPiecewiseBlenderSystemImpl<ValueType>::ZeroAccumulationBuffers()
+{
+	// Arrays should only ever exist in these containers if they have size (they are always initialized to MaximumNumBlends in ReinitializeAccumulationBuffers)
+	for (TPair<FComponentTypeID, TArray<FBlendResult>>& Pair : AccumulationBuffers.Absolute)
+	{
+		FMemory::Memzero(Pair.Value.GetData(), sizeof(FBlendResult)*Pair.Value.Num());
+	}
+	for (TPair<FComponentTypeID, TArray<FBlendResult>>& Pair : AccumulationBuffers.Relative)
+	{
+		FMemory::Memzero(Pair.Value.GetData(), sizeof(FBlendResult)*Pair.Value.Num());
+	}
+	for (TPair<FComponentTypeID, TArray<FBlendResult>>& Pair : AccumulationBuffers.Additive)
+	{
+		FMemory::Memzero(Pair.Value.GetData(), sizeof(FBlendResult)*Pair.Value.Num());
+	}
+	for (TPair<FComponentTypeID, FAdditiveFromBaseBuffer>& Pair : AccumulationBuffers.AdditiveFromBase)
+	{
+		FMemory::Memzero(Pair.Value.Buffer.GetData(), sizeof(FBlendResult)*Pair.Value.Buffer.Num());
+	}
+}
+
+template<typename ValueType>
+FGraphEventRef TPiecewiseBlenderSystemImpl<ValueType>::DispatchDecomposeTask(FEntityManager& EntityManager, const UE::MovieScene::FValueDecompositionParams& Params, UE::MovieScene::FAlignedDecomposedValue* Output)
+{
+	using namespace UE::MovieScene;
+	using FPiecewiseBlendableValueTraits = TPiecewiseBlendableValueTraits<ValueType>;
+
+	if (!Params.ResultComponentType)
+	{
+		return nullptr;
+	}
+
+	TComponentTypeID<ValueType> ResultComponentType = Params.ResultComponentType.ReinterpretCast<ValueType>();
+
+	struct FChannelResultTask
+	{
+		TArray<FMovieSceneEntityID, TInlineAllocator<8>> EntitiesToDecompose;
+		FAlignedDecomposedValue* Result;
+		uint16 DecomposeBlendChannel;
+		FComponentTypeID AdditiveBlendTag;
+
+		explicit FChannelResultTask(const FValueDecompositionParams& Params, FAlignedDecomposedValue* InResult)
+			: Result(InResult)
+			, DecomposeBlendChannel(Params.DecomposeBlendChannel)
+			, AdditiveBlendTag(FBuiltInComponentTypes::Get()->Tags.AdditiveBlend)
+		{
+			EntitiesToDecompose.Append(Params.Query.Entities.GetData(), Params.Query.Entities.Num());
+		}
+
+		void ForEachAllocation(const FEntityAllocation* Allocation, TRead<FMovieSceneEntityID> EntityToDecomposeIDs, TRead<FMovieSceneBlendChannelID> BlendChannels, TRead<ValueType> ValueResultComponent, TReadOptional<float> OptionalWeightComponent)
+		{
+			static const FMovieSceneBlenderSystemID BlenderSystemID = FPiecewiseBlendableValueTraits::GetBlenderSystemID();
+
+			const bool bAdditive = Allocation->HasComponent(AdditiveBlendTag);
+
+			const int32 Num = Allocation->Num();
+			for (int32 EntityIndex = 0; EntityIndex < Num; ++EntityIndex)
+			{
+				const FMovieSceneBlendChannelID& BlendChannel(BlendChannels[EntityIndex]);
+				ensureMsgf(BlendChannel.SystemID == BlenderSystemID, TEXT("Overriding the standard blender system of standard types isn't supported."));
+
+				if (BlendChannel.ChannelID != DecomposeBlendChannel)
+				{
+					continue;
+				}
+
+				// We've found a contributor for this blend channel
+				const FMovieSceneEntityID EntityToDecompose = EntityToDecomposeIDs[EntityIndex];
+				const float               Weight            = OptionalWeightComponent ? OptionalWeightComponent[EntityIndex] : 1.f;
+				const ValueType           ValueResult       = ValueResultComponent[EntityIndex];
+
+				if (EntitiesToDecompose.Contains(EntityToDecompose))
+				{
+					if (bAdditive)
+					{
+						Result->Value.DecomposedAdditives.Add(MakeTuple(EntityToDecompose, FWeightedValue{ ValueResult, Weight }));
+					}
+					else
+					{
+						Result->Value.DecomposedAbsolutes.Add(MakeTuple(EntityToDecompose, FWeightedValue{ ValueResult, Weight }));
+					}
+				}
+				else if (bAdditive)
+				{
+					Result->Value.Result.Additive += ValueResult * Weight;
+				}
+				else
+				{
+					Result->Value.Result.Absolute.Value  += ValueResult * Weight;
+					Result->Value.Result.Absolute.Weight += Weight;
+				}
+			}
+		}
+	};
+
+	FBuiltInComponentTypes* BuiltInComponents = FBuiltInComponentTypes::Get();
+
+	if (Params.Query.bConvertFromSourceEntityIDs)
+	{
+		return FEntityTaskBuilder()
+			.Read(BuiltInComponents->ParentEntity)
+			.Read(BuiltInComponents->BlendChannelInput)
+			.Read(ResultComponentType)
+			.ReadOptional(BuiltInComponents->WeightAndEasingResult)
+			.FilterAll({ Params.PropertyTag })
+			.template Dispatch_PerAllocation<FChannelResultTask>(&EntityManager, FSystemTaskPrerequisites(), nullptr, Params, Output);
+	}
+	else
+	{
+		return FEntityTaskBuilder()
+			.ReadEntityIDs()
+			.Read(BuiltInComponents->BlendChannelInput)
+			.Read(ResultComponentType)
+			.ReadOptional(BuiltInComponents->WeightAndEasingResult)
+			.FilterAll({ Params.PropertyTag })
+			.template Dispatch_PerAllocation<FChannelResultTask>(&EntityManager, FSystemTaskPrerequisites(), nullptr, Params, Output);
+	}
+}
+
+template struct MOVIESCENETRACKS_API TPiecewiseBlenderSystemImpl<float>;
+template struct MOVIESCENETRACKS_API TPiecewiseBlenderSystemImpl<double>;
+
+} // namespace MovieScene
+} // namespace UE
diff --git a/Engine/Source/Runtime/MovieSceneTracks/Private/Systems/MovieScenePiecewiseFloatBlenderSystem.cpp b/Engine/Source/Runtime/MovieSceneTracks/Private/Systems/MovieScenePiecewiseFloatBlenderSystem.cpp
--- a/Engine/Source/Runtime/MovieSceneTracks/Private/Systems/MovieScenePiecewiseFloatBlenderSystem.cpp
+++ b/Engine/Source/Runtime/MovieSceneTracks/Private/Systems/MovieScenePiecewiseFloatBlenderSystem.cpp
@@ -2,378 +2,18 @@
 
 #include "Systems/MovieScenePiecewiseFloatBlenderSystem.h"
 #include "EntitySystem/BuiltInComponentTypes.h"
-#include "EntitySystem/EntityAllocationIterator.h"
-#include "EntitySystem/MovieSceneEntitySystemTask.h"
 #include "EntitySystem/MovieSceneEntitySystemLinker.h"
-
-#include "EntitySystem/MovieSceneEvalTimeSystem.h"
+#include "EntitySystem/MovieSceneEntitySystemTask.h"
 #include "Systems/FloatChannelEvaluatorSystem.h"
+#include "Systems/MovieScenePiecewiseBlenderSystemImpl.h"
 #include "Systems/MovieSceneQuaternionInterpolationRotationSystem.h"
 #include "Systems/WeightAndEasingEvaluatorSystem.h"
-#include "Algo/Find.h"
-#include "Algo/AnyOf.h"
-#include "Algo/Accumulate.h"
-#include "Channels/MovieSceneFloatChannel.h"
-#include "Sections/MovieSceneFloatSection.h"
 
 
-DECLARE_CYCLE_STAT(TEXT("Piecewise Blender System"),             MovieSceneEval_PiecewiseBlenderSystem,  STATGROUP_MovieSceneECS);
+DECLARE_CYCLE_STAT(TEXT("Piecewise Float Blender System"),       MovieSceneEval_PiecewiseFloatBlenderSystem,  STATGROUP_MovieSceneECS);
 DECLARE_CYCLE_STAT(TEXT("Blend float values"),                   MovieSceneEval_BlendFloatValues,        STATGROUP_MovieSceneECS);
 DECLARE_CYCLE_STAT(TEXT("Default combine blended float values"), MovieSceneEval_BlendCombineFloatValues, STATGROUP_MovieSceneECS);
  
-namespace UE
-{
-namespace MovieScene
-{
-
-
-struct FAccumulationResult
-{
-	const FBlendResult* Absolutes = nullptr;
-	const FBlendResult* Relatives = nullptr;
-	const FBlendResult* Additives = nullptr;
-	const FBlendResult* AdditivesFromBase = nullptr;
-
-	bool IsValid() const
-	{
-		return Absolutes || Relatives || Additives || AdditivesFromBase;
-	}
-
-	FBlendResult GetAbsoluteResult(uint16 BlendID) const
-	{
-		return Absolutes ? Absolutes[BlendID] : FBlendResult{};
-	}
-	FBlendResult GetRelativeResult(uint16 BlendID) const
-	{
-		return Relatives ? Relatives[BlendID] : FBlendResult{};
-	}
-	FBlendResult GetAdditiveResult(uint16 BlendID) const
-	{
-		return Additives ? Additives[BlendID] : FBlendResult{};
-	}
-	FBlendResult GetAdditiveFromBaseResult(uint16 BlendID) const
-	{
-		return AdditivesFromBase ? AdditivesFromBase[BlendID] : FBlendResult{};
-	}
-};
-
-/** Task for accumulating all weighted blend inputs into arrays based on BlendID. Will be run for Absolute, Additive and Relative blend modes*/
-struct FAccumulationTask
-{
-	FAccumulationTask(TSortedMap<FComponentTypeID, TArray<FBlendResult>>* InAccumulationBuffers)
-		: AccumulationBuffers(InAccumulationBuffers)
-	{}
-
-	/** Task entry point - iterates the allocation's headers and accumulates float results for any required components */
-	void ForEachAllocation(FEntityAllocationIteratorItem InItem, TRead<FMovieSceneBlendChannelID> BlendIDs, TReadOptional<float> OptionalEasingAndWeights)
-	{
-		const FEntityAllocation* Allocation = InItem;
-		const FComponentMask& AllocationType = InItem;
-
-		for (const FComponentHeader& ComponentHeader : Allocation->GetComponentHeaders())
-		{
-			if (TArray<FBlendResult>* AccumulationBuffer = AccumulationBuffers->Find(ComponentHeader.ComponentType))
-			{
-				ComponentHeader.ReadWriteLock.ReadLock();
-
-				const float* FloatResults = static_cast<const float*>(ComponentHeader.GetValuePtr(0));
-				AccumulateResults(Allocation, FloatResults, BlendIDs, OptionalEasingAndWeights, *AccumulationBuffer);
-
-				ComponentHeader.ReadWriteLock.ReadUnlock();
-			}
-		}
-	}
-
-private:
-
-	void AccumulateResults(const FEntityAllocation* InAllocation, const float* InFloatResults, const FMovieSceneBlendChannelID* BlendIDs, const float* OptionalEasingAndWeights, TArray<FBlendResult>& OutBlendResults)
-	{
-		static const FMovieSceneBlenderSystemID FloatBlenderSystemID = UMovieSceneBlenderSystem::GetBlenderSystemID<UMovieScenePiecewiseFloatBlenderSystem>();
-
-		const int32 Num = InAllocation->Num();
-		if (OptionalEasingAndWeights)
-		{
-			// We have some easing/weight factors to multiply values with.
-			for (int32 Index = 0; Index < Num; ++Index)
-			{
-				const FMovieSceneBlendChannelID& BlendID(BlendIDs[Index]);
-				ensureMsgf(BlendID.SystemID == FloatBlenderSystemID, TEXT("Overriding the standard blender system of standard types isn't supported."));
-
-				FBlendResult& Result = OutBlendResults[BlendID.ChannelID];
-
-				const float Weight = OptionalEasingAndWeights[Index];
-				Result.Total  += InFloatResults[Index] * Weight;
-				Result.Weight += Weight;
-			}
-		}
-		else
-		{
-			// Faster path for when there's no weight to multiply values with.
-			for (int32 Index = 0; Index < Num; ++Index)
-			{
-				const FMovieSceneBlendChannelID& BlendID(BlendIDs[Index]);
-				ensureMsgf(BlendID.SystemID == FloatBlenderSystemID, TEXT("Overriding the standard blender system of standard types isn't supported."));
-
-				FBlendResult& Result = OutBlendResults[BlendID.ChannelID];
-				Result.Total  += InFloatResults[Index];
-				Result.Weight += 1.f;
-			}
-		}
-	}
-
-	TSortedMap<FComponentTypeID, TArray<FBlendResult>>* AccumulationBuffers;
-};
-
-/** Same as the task above, but also reads a "base value" that is subtracted from all values.
- *
- *  Only used by entities with the "additive from base" blend type.
- */
-struct FAdditiveFromBaseBlendTask
-{
-	TSortedMap<FComponentTypeID, FAdditiveFromBaseBuffer>* AccumulationBuffers;
-
-	void ForEachAllocation(FEntityAllocationIteratorItem InItem, TRead<FMovieSceneBlendChannelID> BlendIDs, TReadOptional<float> EasingAndWeightResults)
-	{
-		FEntityAllocation* Allocation = InItem;
-		const FComponentMask& AllocationType = InItem;
-
-		for (const FComponentHeader& ComponentHeader : Allocation->GetComponentHeaders())
-		{
-			if (FAdditiveFromBaseBuffer* Buffer = AccumulationBuffers->Find(ComponentHeader.ComponentType))
-			{
-				TComponentReader<float> BaseValues = Allocation->ReadComponents(Buffer->BaseComponent.ReinterpretCast<float>());
-				TComponentReader<float> FloatResults(&ComponentHeader);
-
-				AccumulateResults(Allocation, FloatResults.AsPtr(), BaseValues.AsPtr(), BlendIDs, EasingAndWeightResults, Buffer->Buffer);
-			}
-		}
-	}
-
-private:
-
-	void AccumulateResults(const FEntityAllocation* InAllocation, const float* FloatResults, const float* BaseValues, const FMovieSceneBlendChannelID* BlendIDs, const float* OptionalEasingAndWeights, TArray<FBlendResult>& OutBlendResults)
-	{
-		static const FMovieSceneBlenderSystemID FloatBlenderSystemID = UMovieSceneBlenderSystem::GetBlenderSystemID<UMovieScenePiecewiseFloatBlenderSystem>();
-		
-		const int32 Num = InAllocation->Num();
-
-		if (OptionalEasingAndWeights)
-		{
-			for (int32 Index = 0; Index < Num; ++Index)
-			{
-				const FMovieSceneBlendChannelID& BlendID(BlendIDs[Index]);
-				ensureMsgf(BlendID.SystemID == FloatBlenderSystemID, TEXT("Overriding the standard blender system of standard types isn't supported."));
-
-				FBlendResult& Result = OutBlendResults[BlendID.ChannelID];
-
-				const float Weight = OptionalEasingAndWeights[Index];
-				Result.Total  += (FloatResults[Index] - BaseValues[Index]) * Weight;
-				Result.Weight += Weight;
-			}
-		}
-		else
-		{
-			// Faster path for when there's no weight to multiply values with.
-			for (int32 Index = 0; Index < Num; ++Index)
-			{
-				const FMovieSceneBlendChannelID& BlendID(BlendIDs[Index]);
-				ensureMsgf(BlendID.SystemID == FloatBlenderSystemID, TEXT("Overriding the standard blender system of standard types isn't supported."));
-
-				FBlendResult& Result = OutBlendResults[BlendID.ChannelID];
-				Result.Total  += (FloatResults[Index] - BaseValues[Index]);
-				Result.Weight += 1.f;
-			}
-		}
-	}
-};
-
-/** Task that combines all accumulated blends for any tracked property type that has blend inputs/outputs */
-struct FCombineBlends
-{
-	explicit FCombineBlends(const TBitArray<>& InCachedRelevantProperties, const FAccumulationBuffers* InAccumulationBuffers, FEntityAllocationWriteContext InWriteContext)
-		: CachedRelevantProperties(InCachedRelevantProperties)
-		, AccumulationBuffers(InAccumulationBuffers)
-		, PropertyRegistry(&FBuiltInComponentTypes::Get()->PropertyRegistry)
-		, WriteContext(InWriteContext)
-	{}
-
-	void ForEachAllocation(FEntityAllocationIteratorItem InItem, TRead<FMovieSceneBlendChannelID> BlendIDs)
-	{
-		FEntityAllocation* Allocation = InItem;
-		const FComponentMask& AllocationType = InItem;
-
-		// Find out what kind of property this is
-		for (TConstSetBitIterator<> PropertyIndex(CachedRelevantProperties); PropertyIndex; ++PropertyIndex)
-		{
-			const FPropertyDefinition& PropertyDefinition = PropertyRegistry->GetDefinition(FCompositePropertyTypeID::FromIndex(PropertyIndex.GetIndex()));
-			if (AllocationType.Contains(PropertyDefinition.PropertyType))
-			{
-				ProcessPropertyType(Allocation, AllocationType, PropertyDefinition, BlendIDs);
-				return;
-			}
-		}
-	}
-
-private:
-
-	void ProcessPropertyType(FEntityAllocation* Allocation, const FComponentMask& AllocationType, const FPropertyDefinition& PropertyDefinition, const FMovieSceneBlendChannelID* BlendIDs)
-	{
-		static const FMovieSceneBlenderSystemID FloatBlenderSystemID = UMovieSceneBlenderSystem::GetBlenderSystemID<UMovieScenePiecewiseFloatBlenderSystem>();
-
-		TArrayView<const FPropertyCompositeDefinition> Composites = PropertyRegistry->GetComposites(PropertyDefinition);
-
-		FOptionalComponentReader OptInitialValues = Allocation->TryReadComponentsErased(PropertyDefinition.InitialValueType);
-
-		for (int32 CompositeIndex = 0; CompositeIndex < Composites.Num(); ++CompositeIndex)
-		{
-			if ( (PropertyDefinition.FloatCompositeMask & (1 << CompositeIndex)) == 0)
-			{
-				continue;
-			}
-
-			TComponentTypeID<float> ResultComponent = Composites[CompositeIndex].ComponentTypeID.ReinterpretCast<float>();
-			if (!AllocationType.Contains(ResultComponent))
-			{
-				continue;
-			}
-
-			FAccumulationResult Results = AccumulationBuffers->FindResults(ResultComponent);
-			if (Results.IsValid())
-			{
-				const uint16 InitialValueProjectionOffset = Composites[CompositeIndex].CompositeOffset;
-
-				// Open the float result channel for write
-				TComponentWriter<float> FloatResults = Allocation->WriteComponents(ResultComponent, WriteContext);
-
-				if (OptInitialValues)
-				{
-					for (int32 Index = 0; Index < Allocation->Num(); ++Index)
-					{
-						ensureMsgf(BlendIDs[Index].SystemID == FloatBlenderSystemID, TEXT("Overriding the standard blender system of standard types isn't supported."));
-						const float InitialValue = *reinterpret_cast<const float*>(static_cast<const uint8*>(OptInitialValues[Index]) + InitialValueProjectionOffset);
-						BlendResultsWithInitial(Results, BlendIDs[Index].ChannelID, InitialValue, FloatResults[Index]);
-					}
-				}
-				else
-				{
-					for (int32 Index = 0; Index < Allocation->Num(); ++Index)
-					{
-						ensureMsgf(BlendIDs[Index].SystemID == FloatBlenderSystemID, TEXT("Overriding the standard blender system of standard types isn't supported."));
-						BlendResults(Results, BlendIDs[Index].ChannelID, FloatResults[Index]);
-					}
-				}
-			}
-		}
-	}
-
-	void BlendResultsWithInitial(const FAccumulationResult& Results, uint16 BlendID, const float InitialValue, float& OutFinalBlendResult)
-	{
-		FBlendResult AbsoluteResult = Results.GetAbsoluteResult(BlendID);
-		FBlendResult RelativeResult = Results.GetRelativeResult(BlendID);
-		FBlendResult AdditiveResult = Results.GetAdditiveResult(BlendID);
-		FBlendResult AdditiveFromBaseResult = Results.GetAdditiveFromBaseResult(BlendID);
-
-		if (RelativeResult.Weight != 0)
-		{
-			RelativeResult.Total += InitialValue * RelativeResult.Weight;
-		}
-
-		FBlendResult TotalAdditiveResult = { AdditiveResult.Total + AdditiveFromBaseResult.Total, AdditiveResult.Weight + AdditiveFromBaseResult.Weight };
-
-		const float TotalWeight = AbsoluteResult.Weight + RelativeResult.Weight;
-		if (TotalWeight != 0)
-		{
-			// If the absolute value has some partial weighting (for ease-in/out for instance), we ramp it from/to the initial value. This means
-			// that the "initial value" adds a contribution to the entire blending process, so we add its weight to the total that we
-			// normalize absolutes and relatives with.
-			//
-			// Note that "partial weighting" means strictly between 0 and 100%. At 100% and above, we don't need to do this thing with the initial
-			// value. At 0%, we have no absolute value (only a relative value) and we therefore don't want to include the initial value either.
-			const bool bInitialValueContributes = (0.f < AbsoluteResult.Weight && AbsoluteResult.Weight < 1.f);
-			const float AbsoluteBlendedValue = bInitialValueContributes ?
-				(InitialValue * (1.f - AbsoluteResult.Weight) + AbsoluteResult.Total) :
-				AbsoluteResult.Total;
-			const float FinalTotalWeight = bInitialValueContributes ? (TotalWeight + (1.f - AbsoluteResult.Weight)) : TotalWeight;
-
-			const float Value = (AbsoluteBlendedValue + RelativeResult.Total) / FinalTotalWeight + TotalAdditiveResult.Total;
-			OutFinalBlendResult = Value;
-		}
-		else if (TotalAdditiveResult.Weight != 0)
-		{
-			OutFinalBlendResult = TotalAdditiveResult.Total + InitialValue;
-		}
-		else
-		{
-			OutFinalBlendResult = InitialValue;
-		}
-	}
-
-	void BlendResults(const FAccumulationResult& Results, uint16 BlendID, float& OutFinalBlendResult)
-	{
-		FBlendResult AbsoluteResult = Results.GetAbsoluteResult(BlendID);
-		FBlendResult AdditiveResult = Results.GetAdditiveResult(BlendID);
-		FBlendResult AdditiveFromBaseResult = Results.GetAdditiveFromBaseResult(BlendID);
-
-#if DO_GUARD_SLOW
-		ensureMsgf(AbsoluteResult.Weight != 0.f, TEXT("Default blend combine being used for an entity that has no absolute weight. This should have an initial value and should be handled by each system, and excluded by default with UMovieSceneBlenderSystem::FinalCombineExclusionFilter ."));
-#endif
-
-		const float TotalWeight = AbsoluteResult.Weight;
-		if (TotalWeight != 0)
-		{
-			const float Value = AbsoluteResult.Total / AbsoluteResult.Weight + AdditiveResult.Total + AdditiveFromBaseResult.Total;
-			OutFinalBlendResult = Value;
-		}
-	}
-
-private:
-
-	TBitArray<> CachedRelevantProperties;
-	const FAccumulationBuffers* AccumulationBuffers;
-	const FPropertyRegistry* PropertyRegistry;
-	FEntityAllocationWriteContext WriteContext;
-};
-
-
-
-bool FAccumulationBuffers::IsEmpty() const
-{
-	return Absolute.Num() == 0 && Relative.Num() == 0 && Additive.Num() == 0 && AdditiveFromBase.Num() == 0;
-}
-
-void FAccumulationBuffers::Reset()
-{
-	Absolute.Empty();
-	Relative.Empty();
-	Additive.Empty();
-	AdditiveFromBase.Empty();
-}
-
-FAccumulationResult FAccumulationBuffers::FindResults(FComponentTypeID InComponentType) const
-{
-	FAccumulationResult Result;
-	if (const TArray<FBlendResult>* Absolutes = Absolute.Find(InComponentType))
-	{
-		Result.Absolutes = Absolutes->GetData();
-	}
-	if (const TArray<FBlendResult>* Relatives = Relative.Find(InComponentType))
-	{
-		Result.Relatives = Relatives->GetData();
-	}
-	if (const TArray<FBlendResult>* Additives = Additive.Find(InComponentType))
-	{
-		Result.Additives = Additives->GetData();
-	}
-	if (const FAdditiveFromBaseBuffer* AdditivesFromBase = AdditiveFromBase.Find(InComponentType))
-	{
-		Result.AdditivesFromBase = AdditivesFromBase->Buffer.GetData();
-	}
-	return Result;
-}
-
-
-} // namespace MovieScene
-} // namespace UE
 
 UMovieScenePiecewiseFloatBlenderSystem::UMovieScenePiecewiseFloatBlenderSystem(const FObjectInitializer& ObjInit)
 	: Super(ObjInit)
@@ -386,313 +26,23 @@ UMovieScenePiecewiseFloatBlenderSystem::UMovieScenePiecewiseFloatBlenderSystem(c
 	}
 }
 
-void UMovieScenePiecewiseFloatBlenderSystem::OnLink()
-{
-}
-
 void UMovieScenePiecewiseFloatBlenderSystem::OnRun(FSystemTaskPrerequisites& InPrerequisites, FSystemSubsequentTasks& Subsequents)
 {
 	using namespace UE::MovieScene;
 
-	SCOPE_CYCLE_COUNTER(MovieSceneEval_PiecewiseBlenderSystem)
+	SCOPE_CYCLE_COUNTER(MovieSceneEval_PiecewiseFloatBlenderSystem)
 
 	CompactBlendChannels();
 
-	// We allocate space for every blend even if there are gaps so we can do a straight index into each array
-	const int32 MaximumNumBlends = AllocatedBlendChannels.Num();
-	if (MaximumNumBlends == 0)
-	{
-		return;
-	}
-
-	// Update cached channel data if necessary
-	if (ChannelRelevancyCache.Update(Linker->EntityManager) == ECachedEntityManagerState::Stale)
-	{
-		ReinitializeAccumulationBuffers();
-	}
-
-	if (AccumulationBuffers.IsEmpty())
-	{
-		return;
-	}
-
-	ZeroAccumulationBuffers();
-
-	FBuiltInComponentTypes* BuiltInComponents = FBuiltInComponentTypes::Get();
-
-	FSystemTaskPrerequisites Prereqs;
-	if (AccumulationBuffers.Absolute.Num() != 0)
-	{
-		FGraphEventRef Task = FEntityTaskBuilder()
-		.Read(BuiltInComponents->BlendChannelInput)
-		.ReadOptional(BuiltInComponents->WeightAndEasingResult)
-		.FilterAll({ BuiltInComponents->Tags.AbsoluteBlend })
-		.FilterAny(BlendedResultMask)
-		.FilterNone({ BuiltInComponents->Tags.Ignored })
-		.SetStat(GET_STATID(MovieSceneEval_BlendFloatValues))
-		.Dispatch_PerAllocation<FAccumulationTask>(&Linker->EntityManager, InPrerequisites, nullptr, &AccumulationBuffers.Absolute);
-
-		if (Task)
-		{
-			Prereqs.AddMasterTask(Task);
-		}
-	}
-
-	if (AccumulationBuffers.Relative.Num() != 0)
-	{
-		FGraphEventRef Task = FEntityTaskBuilder()
-		.Read(BuiltInComponents->BlendChannelInput)
-		.ReadOptional(BuiltInComponents->WeightAndEasingResult)
-		.FilterAll({ BuiltInComponents->Tags.RelativeBlend })
-		.FilterAny(BlendedResultMask)
-		.FilterNone({ BuiltInComponents->Tags.Ignored })
-		.SetStat(GET_STATID(MovieSceneEval_BlendFloatValues))
-		.Dispatch_PerAllocation<FAccumulationTask>(&Linker->EntityManager, InPrerequisites, nullptr, &AccumulationBuffers.Relative);
-
-		if (Task)
-		{
-			Prereqs.AddMasterTask(Task);
-		}
-	}
-
-	if (AccumulationBuffers.Additive.Num() != 0)
-	{
-		FGraphEventRef Task = FEntityTaskBuilder()
-		.Read(BuiltInComponents->BlendChannelInput)
-		.ReadOptional(BuiltInComponents->WeightAndEasingResult)
-		.FilterAll({ BuiltInComponents->Tags.AdditiveBlend })
-		.FilterAny(BlendedResultMask)
-		.FilterNone({ BuiltInComponents->Tags.Ignored })
-		.SetStat(GET_STATID(MovieSceneEval_BlendFloatValues))
-		.Dispatch_PerAllocation<FAccumulationTask>(&Linker->EntityManager, InPrerequisites, nullptr, &AccumulationBuffers.Additive);
-
-		if (Task)
-		{
-			Prereqs.AddMasterTask(Task);
-		}
-	}
-
-	if (AccumulationBuffers.AdditiveFromBase.Num() != 0)
-	{
-		FGraphEventRef Task = FEntityTaskBuilder()
-		.Read(BuiltInComponents->BlendChannelInput)
-		.ReadOptional(BuiltInComponents->WeightAndEasingResult)
-		.FilterAll({ BuiltInComponents->Tags.AdditiveFromBaseBlend })
-		.FilterAny(BlendedResultMask)
-		.FilterNone({ BuiltInComponents->Tags.Ignored })
-		.SetStat(GET_STATID(MovieSceneEval_BlendFloatValues))
-		.Dispatch_PerAllocation<FAdditiveFromBaseBlendTask>(&Linker->EntityManager, InPrerequisites, nullptr, &AccumulationBuffers.AdditiveFromBase);
-
-		if (Task)
-		{
-			Prereqs.AddMasterTask(Task);
-		}
-	}
-
-	// Master task that performs the actual blends
-	FEntityTaskBuilder()
-	.Read(BuiltInComponents->BlendChannelOutput)
-	.FilterAny(BlendedPropertyMask)
-	.SetStat(GET_STATID(MovieSceneEval_BlendCombineFloatValues))
-	.Dispatch_PerAllocation<FCombineBlends>(&Linker->EntityManager, Prereqs, &Subsequents, CachedRelevantProperties, &AccumulationBuffers, FEntityAllocationWriteContext(Linker->EntityManager));
-}
-
-void UMovieScenePiecewiseFloatBlenderSystem::ReinitializeAccumulationBuffers()
-{
-	using namespace UE::MovieScene;
-
-	const int32 MaximumNumBlends = AllocatedBlendChannels.Num();
-
-	FBuiltInComponentTypes* BuiltInComponents = FBuiltInComponentTypes::Get();
-
-	BlendedResultMask.Reset();
-	AccumulationBuffers.Reset();
-
-	// Recompute which result types are blended
-	const int32 NumFloats = UE_ARRAY_COUNT(BuiltInComponents->FloatResult);
-	for (int32 Index = 0; Index < NumFloats; ++Index)
-	{
-		TComponentTypeID<float> Component = BuiltInComponents->FloatResult[Index];
-
-		const bool bHasAbsolutes         = Linker->EntityManager.Contains(FEntityComponentFilter().All({ Component, BuiltInComponents->BlendChannelInput, BuiltInComponents->Tags.AbsoluteBlend }));
-		const bool bHasRelatives         = Linker->EntityManager.Contains(FEntityComponentFilter().All({ Component, BuiltInComponents->BlendChannelInput, BuiltInComponents->Tags.RelativeBlend }));
-		const bool bHasAdditives         = Linker->EntityManager.Contains(FEntityComponentFilter().All({ Component, BuiltInComponents->BlendChannelInput, BuiltInComponents->Tags.AdditiveBlend }));
-		const bool bHasAdditivesFromBase = Linker->EntityManager.Contains(FEntityComponentFilter().All({ Component, BuiltInComponents->BlendChannelInput, BuiltInComponents->Tags.AdditiveFromBaseBlend }));
-
-		if (!(bHasAbsolutes || bHasRelatives || bHasAdditives || bHasAdditivesFromBase))
-		{
-			continue;
-		}
-
-		BlendedResultMask.Set(Component);
-
-		if (bHasAbsolutes)
-		{
-			TArray<FBlendResult>& Buffer = AccumulationBuffers.Absolute.Add(Component);
-			Buffer.SetNum(MaximumNumBlends);
-		}
-		if (bHasRelatives)
-		{
-			TArray<FBlendResult>& Buffer = AccumulationBuffers.Relative.Add(Component);
-			Buffer.SetNum(MaximumNumBlends);
-		}
-		if (bHasAdditives)
-		{
-			TArray<FBlendResult>& Buffer = AccumulationBuffers.Additive.Add(Component);
-			Buffer.SetNum(MaximumNumBlends);
-		}
-		if (bHasAdditivesFromBase)
-		{
-			FAdditiveFromBaseBuffer& Buffer = AccumulationBuffers.AdditiveFromBase.Add(Component);
-			Buffer.Buffer.SetNum(MaximumNumBlends);
-			Buffer.BaseComponent = BuiltInComponents->BaseFloat[Index];
-		}
-	}
-
-	// Update property relevancy
-	CachedRelevantProperties.Empty();
-
-	// If we have no accumulation buffers, we have nothing more to do
-	if (AccumulationBuffers.IsEmpty())
-	{
-		return;
-	}
-
-	// This code works on the assumption that properties can never be removed (which is safe)
-	FEntityComponentFilter InclusionFilter;
-	TArrayView<const FPropertyDefinition> Properties = BuiltInComponents->PropertyRegistry.GetProperties();
-	for (int32 PropertyTypeIndex = 0; PropertyTypeIndex < Properties.Num(); ++PropertyTypeIndex)
-	{
-		const FPropertyDefinition& PropertyDefinition = Properties[PropertyTypeIndex];
-		if (PropertyDefinition.FloatCompositeMask != 0)
-		{
-			InclusionFilter.Reset();
-			InclusionFilter.All({ BuiltInComponents->BlendChannelOutput, PropertyDefinition.PropertyType });
-			if (Linker->EntityManager.Contains(InclusionFilter))
-			{
-				CachedRelevantProperties.PadToNum(PropertyTypeIndex + 1, false);
-				CachedRelevantProperties[PropertyTypeIndex] = true;
-
-				BlendedPropertyMask.Set(PropertyDefinition.PropertyType);
-			}
-		}
-	}
+	FPiecewiseBlenderSystemImplRunParams Params;
+	Params.MaximumNumBlends = AllocatedBlendChannels.Num();
+	Params.BlendValuesStatId = GET_STATID(MovieSceneEval_BlendFloatValues);
+	Params.CombineBlendsStatId = GET_STATID(MovieSceneEval_BlendCombineFloatValues);
+	Impl.Run(Params, Linker->EntityManager, InPrerequisites, Subsequents);
 }
 
-void UMovieScenePiecewiseFloatBlenderSystem::ZeroAccumulationBuffers()
+FGraphEventRef UMovieScenePiecewiseFloatBlenderSystem::DispatchDecomposeTask(const UE::MovieScene::FValueDecompositionParams& Params, UE::MovieScene::FAlignedDecomposedValue* Output)
 {
-	using namespace UE::MovieScene;
-
-	// Arrays should only ever exist in these containers if they have size (they are always initialized to MaximumNumBlends in ReinitializeAccumulationBuffers)
-	for (TPair<FComponentTypeID, TArray<FBlendResult>>& Pair : AccumulationBuffers.Absolute)
-	{
-		FMemory::Memzero(Pair.Value.GetData(), sizeof(FBlendResult)*Pair.Value.Num());
-	}
-	for (TPair<FComponentTypeID, TArray<FBlendResult>>& Pair : AccumulationBuffers.Relative)
-	{
-		FMemory::Memzero(Pair.Value.GetData(), sizeof(FBlendResult)*Pair.Value.Num());
-	}
-	for (TPair<FComponentTypeID, TArray<FBlendResult>>& Pair : AccumulationBuffers.Additive)
-	{
-		FMemory::Memzero(Pair.Value.GetData(), sizeof(FBlendResult)*Pair.Value.Num());
-	}
-	for (TPair<FComponentTypeID, FAdditiveFromBaseBuffer>& Pair : AccumulationBuffers.AdditiveFromBase)
-	{
-		FMemory::Memzero(Pair.Value.Buffer.GetData(), sizeof(FBlendResult)*Pair.Value.Buffer.Num());
-	}
+	return Impl.DispatchDecomposeTask(Linker->EntityManager, Params, Output);
 }
 
-
-FGraphEventRef UMovieScenePiecewiseFloatBlenderSystem::DispatchDecomposeTask(const UE::MovieScene::FFloatDecompositionParams& Params, UE::MovieScene::FAlignedDecomposedFloat* Output)
-{
-	using namespace UE::MovieScene;
-
-	if (!Params.ResultComponentType)
-	{
-		return nullptr;
-	}
-
-	struct FChannelResultTask
-	{
-		TArray<FMovieSceneEntityID, TInlineAllocator<8>> EntitiesToDecompose;
-		FAlignedDecomposedFloat* Result;
-		uint16 DecomposeBlendChannel;
-		FComponentTypeID AdditiveBlendTag;
-
-		explicit FChannelResultTask(const FFloatDecompositionParams& Params, FAlignedDecomposedFloat* InResult)
-			: Result(InResult)
-			, DecomposeBlendChannel(Params.DecomposeBlendChannel)
-			, AdditiveBlendTag(FBuiltInComponentTypes::Get()->Tags.AdditiveBlend)
-		{
-			EntitiesToDecompose.Append(Params.Query.Entities.GetData(), Params.Query.Entities.Num());
-		}
-
-		void ForEachAllocation(const FEntityAllocation* Allocation, TRead<FMovieSceneEntityID> EntityToDecomposeIDs, TRead<FMovieSceneBlendChannelID> BlendChannels, TRead<float> FloatResultComponent, TReadOptional<float> OptionalWeightComponent)
-		{
-			static const FMovieSceneBlenderSystemID FloatBlenderSystemID = UMovieSceneBlenderSystem::GetBlenderSystemID<UMovieScenePiecewiseFloatBlenderSystem>();
-
-			const bool bAdditive = Allocation->HasComponent(AdditiveBlendTag);
-
-			const int32 Num = Allocation->Num();
-			for (int32 EntityIndex = 0; EntityIndex < Num; ++EntityIndex)
-			{
-				const FMovieSceneBlendChannelID& BlendChannel(BlendChannels[EntityIndex]);
-				ensureMsgf(BlendChannel.SystemID == FloatBlenderSystemID, TEXT("Overriding the standard blender system of standard types isn't supported."));
-
-				if (BlendChannel.ChannelID != DecomposeBlendChannel)
-				{
-					continue;
-				}
-
-				// We've found a contributor for this blend channel
-				const FMovieSceneEntityID EntityToDecompose = EntityToDecomposeIDs[EntityIndex];
-				const float               Weight            = OptionalWeightComponent ? OptionalWeightComponent[EntityIndex] : 1.f;
-				const float               FloatResult       = FloatResultComponent[EntityIndex];
-
-				if (EntitiesToDecompose.Contains(EntityToDecompose))
-				{
-					if (bAdditive)
-					{
-						Result->Value.DecomposedAdditives.Add(MakeTuple(EntityToDecompose, FWeightedFloat{ FloatResult, Weight }));
-					}
-					else
-					{
-						Result->Value.DecomposedAbsolutes.Add(MakeTuple(EntityToDecompose, FWeightedFloat{ FloatResult, Weight }));
-					}
-				}
-				else if (bAdditive)
-				{
-					Result->Value.Result.Additive += FloatResult * Weight;
-				}
-				else
-				{
-					Result->Value.Result.Absolute.Value  += FloatResult * Weight;
-					Result->Value.Result.Absolute.Weight += Weight;
-				}
-			}
-		}
-	};
-
-	FBuiltInComponentTypes* BuiltInComponents = FBuiltInComponentTypes::Get();
-
-	if (Params.Query.bConvertFromSourceEntityIDs)
-	{
-		return FEntityTaskBuilder()
-			.Read(BuiltInComponents->ParentEntity)
-			.Read(BuiltInComponents->BlendChannelInput)
-			.Read(Params.ResultComponentType)
-			.ReadOptional(BuiltInComponents->WeightAndEasingResult)
-			.FilterAll({ Params.PropertyTag })
-			.Dispatch_PerAllocation<FChannelResultTask>(&Linker->EntityManager, FSystemTaskPrerequisites(), nullptr, Params, Output);
-	}
-	else
-	{
-		return FEntityTaskBuilder()
-			.ReadEntityIDs()
-			.Read(BuiltInComponents->BlendChannelInput)
-			.Read(Params.ResultComponentType)
-			.ReadOptional(BuiltInComponents->WeightAndEasingResult)
-			.FilterAll({ Params.PropertyTag })
-			.Dispatch_PerAllocation<FChannelResultTask>(&Linker->EntityManager, FSystemTaskPrerequisites(), nullptr, Params, Output);
-	}
-}
diff --git a/Engine/Source/Runtime/MovieSceneTracks/Private/Systems/MovieScenePropertyInstantiator.cpp b/Engine/Source/Runtime/MovieSceneTracks/Private/Systems/MovieScenePropertyInstantiator.cpp
--- a/Engine/Source/Runtime/MovieSceneTracks/Private/Systems/MovieScenePropertyInstantiator.cpp
+++ b/Engine/Source/Runtime/MovieSceneTracks/Private/Systems/MovieScenePropertyInstantiator.cpp
@@ -386,7 +386,7 @@ void UMovieScenePropertyInstantiatorSystem::InitializeBlendPath(const FPropertyP
 
 	TArrayView<const FPropertyCompositeDefinition> Composites = BuiltInComponents->PropertyRegistry.GetComposites(*Params.PropertyDefinition);
 
-	UClass* BlenderClass = UMovieScenePiecewiseFloatBlenderSystem::StaticClass();
+	UClass* BlenderClass = Params.PropertyDefinition->BlenderSystemClass;
 
 	// Ensure contributors all have the necessary blend inputs and tags
 	for (auto ContributorIt = Contributors.CreateConstKeyIterator(Params.PropertyInfoIndex); ContributorIt; ++ContributorIt)
@@ -401,7 +401,10 @@ void UMovieScenePropertyInstantiatorSystem::InitializeBlendPath(const FPropertyP
 		}
 	}
 
-	check(BlenderClass);
+	if (!ensureMsgf(BlenderClass, TEXT("No default blender class specified on property, and no custom blender specified on entities. Falling back to float blender.")))
+	{
+		BlenderClass = UMovieScenePiecewiseFloatBlenderSystem::StaticClass();
+	}
 
 	UMovieSceneBlenderSystem* ExistingBlender = Params.PropertyInfo->Blender.Get();
 	if (ExistingBlender && BlenderClass != ExistingBlender->GetClass())
@@ -602,4 +605,4 @@ void UMovieScenePropertyInstantiatorSystem::InitializePropertyMetaData(FSystemTa
 	}
 
 	InitializePropertyMetaDataTasks.Empty();
-}
\ No newline at end of file
+}
diff --git a/Engine/Source/Runtime/MovieSceneTracks/Private/Systems/MovieSceneQuaternionInterpolationRotationSystem.cpp b/Engine/Source/Runtime/MovieSceneTracks/Private/Systems/MovieSceneQuaternionInterpolationRotationSystem.cpp
--- a/Engine/Source/Runtime/MovieSceneTracks/Private/Systems/MovieSceneQuaternionInterpolationRotationSystem.cpp
+++ b/Engine/Source/Runtime/MovieSceneTracks/Private/Systems/MovieSceneQuaternionInterpolationRotationSystem.cpp
@@ -1,7 +1,7 @@
 // Copyright Epic Games, Inc. All Rights Reserved.
 
 #include "Systems/MovieSceneQuaternionInterpolationRotationSystem.h"
-#include "Channels/MovieSceneFloatChannel.h"
+#include "Channels/MovieSceneDoubleChannel.h"
 #include "EntitySystem/BuiltInComponentTypes.h"
 #include "EntitySystem/MovieSceneEvalTimeSystem.h"
 #include "EntitySystem/MovieSceneEntitySystemLinker.h"
@@ -43,12 +43,12 @@ struct FEvaluateQuaternionInterpolationRotationChannels
 	void ForEachAllocation(
 			const FEntityAllocation* Allocation,
 			TRead<FFrameTime> FrameTimes, 
-			TMultiReadOptional<FSourceFloatChannel, FSourceFloatChannel, FSourceFloatChannel> RotChannelAccessors,
-			float* OutResultXs, float* OutResultYs, float* OutResultZs)
+			TMultiReadOptional<FSourceDoubleChannel, FSourceDoubleChannel, FSourceDoubleChannel> RotChannelAccessors,
+			double* OutResultXs, double* OutResultYs, double* OutResultZs)
 	{
-		const FSourceFloatChannel* RotationXs = RotChannelAccessors.Get<0>();
-		const FSourceFloatChannel* RotationYs = RotChannelAccessors.Get<1>();
-		const FSourceFloatChannel* RotationZs = RotChannelAccessors.Get<2>();
+		const FSourceDoubleChannel* RotationXs = RotChannelAccessors.Get<0>();
+		const FSourceDoubleChannel* RotationYs = RotChannelAccessors.Get<1>();
+		const FSourceDoubleChannel* RotationZs = RotChannelAccessors.Get<2>();
 
 		check(
 			(OutResultXs != nullptr) == (RotationXs != nullptr) &&
@@ -60,9 +60,9 @@ struct FEvaluateQuaternionInterpolationRotationChannels
 		{
 			const FFrameTime FrameTime = FrameTimes[Index];
 
-			const FSourceFloatChannel* RotationX = RotationXs ? &RotationXs[Index] : nullptr;
-			const FSourceFloatChannel* RotationY = RotationYs ? &RotationYs[Index] : nullptr;
-			const FSourceFloatChannel* RotationZ = RotationZs ? &RotationZs[Index] : nullptr;
+			const FSourceDoubleChannel* RotationX = RotationXs ? &RotationXs[Index] : nullptr;
+			const FSourceDoubleChannel* RotationY = RotationYs ? &RotationYs[Index] : nullptr;
+			const FSourceDoubleChannel* RotationZ = RotationZs ? &RotationZs[Index] : nullptr;
 
 			// Find the closest keyframes before/after the current time on the 3 rotation channels.
 			TRange<FFrameNumber> FrameRange(TNumericLimits<FFrameNumber>::Min(), TNumericLimits<FFrameNumber>::Max());
@@ -83,7 +83,7 @@ struct FEvaluateQuaternionInterpolationRotationChannels
 			const FFrameNumber UpperBound = FrameRange.GetUpperBoundValue();
 			if (LowerBound != TNumericLimits<FFrameNumber>::Min() && UpperBound != TNumericLimits<FFrameNumber>::Max())
 			{
-				float Value;
+				double Value;
 				FVector FirstRot(0.0f, 0.0f, 0.0f);
 				FVector SecondRot(0.0f, 0.0f, 0.0f);
 				double U = (FrameTime.AsDecimal() - (double) FrameRange.GetLowerBoundValue().Value) /
@@ -144,7 +144,7 @@ struct FEvaluateQuaternionInterpolationRotationChannels
 			}
 			else  // no range found: default to regular, but still do RotToQuat
 			{
-				float Value;
+				double Value;
 				FVector CurrentRot(0.0f, 0.0f, 0.0f);
 				if (RotationX && RotationX->Source->Evaluate(FrameTime, Value))
 				{
@@ -190,9 +190,9 @@ UMovieSceneQuaternionInterpolationRotationSystem::UMovieSceneQuaternionInterpola
 		DefineImplicitPrerequisite(UMovieSceneEvalTimeSystem::StaticClass(), GetClass());
 
 		FBuiltInComponentTypes* BuiltInComponents = FBuiltInComponentTypes::Get();
-		DefineComponentProducer(GetClass(), BuiltInComponents->FloatResult[3]);
-		DefineComponentProducer(GetClass(), BuiltInComponents->FloatResult[4]);
-		DefineComponentProducer(GetClass(), BuiltInComponents->FloatResult[5]);
+		DefineComponentProducer(GetClass(), BuiltInComponents->DoubleResult[3]);
+		DefineComponentProducer(GetClass(), BuiltInComponents->DoubleResult[4]);
+		DefineComponentProducer(GetClass(), BuiltInComponents->DoubleResult[5]);
 	}
 }
 
@@ -224,9 +224,9 @@ void UMovieSceneQuaternionInterpolationRotationSystem::OnRun(FSystemTaskPrerequi
 				TrackComponents->QuaternionRotationChannel[0],
 				TrackComponents->QuaternionRotationChannel[1],
 				TrackComponents->QuaternionRotationChannel[2])
-		.WriteOptional(BuiltInComponents->FloatResult[3])
-		.WriteOptional(BuiltInComponents->FloatResult[4])
-		.WriteOptional(BuiltInComponents->FloatResult[5])
+		.WriteOptional(BuiltInComponents->DoubleResult[3])
+		.WriteOptional(BuiltInComponents->DoubleResult[4])
+		.WriteOptional(BuiltInComponents->DoubleResult[5])
 		.SetStat(GET_STATID(MovieSceneEval_EvaluateQuatInterpRotChannelTask))
 		.Dispatch_PerAllocation<FEvaluateQuaternionInterpolationRotationChannels>(&Linker->EntityManager, InPrerequisites, &Subsequents);
 }
diff --git a/Engine/Source/Runtime/MovieSceneTracks/Private/Systems/MovieSceneTransformOriginSystem.cpp b/Engine/Source/Runtime/MovieSceneTracks/Private/Systems/MovieSceneTransformOriginSystem.cpp
--- a/Engine/Source/Runtime/MovieSceneTracks/Private/Systems/MovieSceneTransformOriginSystem.cpp
+++ b/Engine/Source/Runtime/MovieSceneTracks/Private/Systems/MovieSceneTransformOriginSystem.cpp
@@ -7,8 +7,8 @@
 #include "Tracks/IMovieSceneTransformOrigin.h"
 #include "MovieSceneTracksComponentTypes.h"
 
-#include "Systems/FloatChannelEvaluatorSystem.h"
-#include "Systems/MovieScenePiecewiseFloatBlenderSystem.h"
+#include "Systems/DoubleChannelEvaluatorSystem.h"
+#include "Systems/MovieScenePiecewiseDoubleBlenderSystem.h"
 #include "Systems/MovieSceneComponentTransformSystem.h"
 
 #include "IMovieScenePlayer.h"
@@ -24,19 +24,19 @@ struct FAssignTransformOrigin
 	const TSparseArray<FTransform>* TransformOriginsByInstanceID;
 
 	void ForEachAllocation(const FEntityAllocation* Allocation, TRead<FInstanceHandle> Instances, TRead<UObject*> BoundObjects,
-		TWriteOptional<float> LocationX, TWriteOptional<float> LocationY, TWriteOptional<float> LocationZ,
-		TWriteOptional<float> RotationX, TWriteOptional<float> RotationY, TWriteOptional<float> RotationZ,
-		TWriteOptional<FSourceFloatChannelFlags> FlagsLocationX, TWriteOptional<FSourceFloatChannelFlags> FlagsLocationY, TWriteOptional<FSourceFloatChannelFlags> FlagsLocationZ,
-		TWriteOptional<FSourceFloatChannelFlags> FlagsRotationX, TWriteOptional<FSourceFloatChannelFlags> FlagsRotationY, TWriteOptional<FSourceFloatChannelFlags> FlagsRotationZ)
+		TWriteOptional<double> LocationX, TWriteOptional<double> LocationY, TWriteOptional<double> LocationZ,
+		TWriteOptional<double> RotationX, TWriteOptional<double> RotationY, TWriteOptional<double> RotationZ,
+		TWriteOptional<FSourceDoubleChannelFlags> FlagsLocationX, TWriteOptional<FSourceDoubleChannelFlags> FlagsLocationY, TWriteOptional<FSourceDoubleChannelFlags> FlagsLocationZ,
+		TWriteOptional<FSourceDoubleChannelFlags> FlagsRotationX, TWriteOptional<FSourceDoubleChannelFlags> FlagsRotationY, TWriteOptional<FSourceDoubleChannelFlags> FlagsRotationZ)
 	{
 		TransformLocation(Instances, BoundObjects, LocationX, LocationY, LocationZ, RotationX, RotationY, RotationZ, FlagsLocationX, FlagsLocationY, FlagsLocationZ, FlagsRotationX, FlagsRotationY, FlagsRotationZ, Allocation->Num());
 	}
 
 	void TransformLocation(const FInstanceHandle* Instances, const UObject* const * BoundObjects,
-		float* OutLocationX, float* OutLocationY, float* OutLocationZ,
-		float* OutRotationX, float* OutRotationY, float* OutRotationZ,
-		FSourceFloatChannelFlags* OutFlagsLocationX, FSourceFloatChannelFlags* OutFlagsLocationY, FSourceFloatChannelFlags* OutFlagsLocationZ,
-		FSourceFloatChannelFlags* OutFlagsRotationX, FSourceFloatChannelFlags* OutFlagsRotationY, FSourceFloatChannelFlags* OutFlagsRotationZ,
+		double* OutLocationX, double* OutLocationY, double* OutLocationZ,
+		double* OutRotationX, double* OutRotationY, double* OutRotationZ,
+		FSourceDoubleChannelFlags* OutFlagsLocationX, FSourceDoubleChannelFlags* OutFlagsLocationY, FSourceDoubleChannelFlags* OutFlagsLocationZ,
+		FSourceDoubleChannelFlags* OutFlagsRotationX, FSourceDoubleChannelFlags* OutFlagsRotationY, FSourceDoubleChannelFlags* OutFlagsRotationZ,
 		int32 Num)
 	{
 		for (int32 Index = 0; Index < Num; ++Index)
@@ -95,22 +95,22 @@ UMovieSceneTransformOriginSystem::UMovieSceneTransformOriginSystem(const FObject
 	if (HasAnyFlags(RF_ClassDefaultObject))
 	{
 		// This system relies upon anything that creates entities
-		DefineImplicitPrerequisite(GetClass(), UMovieScenePiecewiseFloatBlenderSystem::StaticClass());
+		DefineImplicitPrerequisite(GetClass(), UMovieScenePiecewiseDoubleBlenderSystem::StaticClass());
 		DefineImplicitPrerequisite(GetClass(), UMovieSceneComponentTransformSystem::StaticClass());
 
-		DefineComponentConsumer(GetClass(), FBuiltInComponentTypes::Get()->FloatResult[0]);
-		DefineComponentConsumer(GetClass(), FBuiltInComponentTypes::Get()->FloatResult[1]);
-		DefineComponentConsumer(GetClass(), FBuiltInComponentTypes::Get()->FloatResult[2]);
-		DefineComponentConsumer(GetClass(), FBuiltInComponentTypes::Get()->FloatResult[3]);
-		DefineComponentConsumer(GetClass(), FBuiltInComponentTypes::Get()->FloatResult[4]);
-		DefineComponentConsumer(GetClass(), FBuiltInComponentTypes::Get()->FloatResult[5]);
-
-		DefineComponentConsumer(GetClass(), FBuiltInComponentTypes::Get()->FloatChannelFlags[0]);
-		DefineComponentConsumer(GetClass(), FBuiltInComponentTypes::Get()->FloatChannelFlags[1]);
-		DefineComponentConsumer(GetClass(), FBuiltInComponentTypes::Get()->FloatChannelFlags[2]);
-		DefineComponentConsumer(GetClass(), FBuiltInComponentTypes::Get()->FloatChannelFlags[3]);
-		DefineComponentConsumer(GetClass(), FBuiltInComponentTypes::Get()->FloatChannelFlags[4]);
-		DefineComponentConsumer(GetClass(), FBuiltInComponentTypes::Get()->FloatChannelFlags[5]);
+		DefineComponentConsumer(GetClass(), FBuiltInComponentTypes::Get()->DoubleResult[0]);
+		DefineComponentConsumer(GetClass(), FBuiltInComponentTypes::Get()->DoubleResult[1]);
+		DefineComponentConsumer(GetClass(), FBuiltInComponentTypes::Get()->DoubleResult[2]);
+		DefineComponentConsumer(GetClass(), FBuiltInComponentTypes::Get()->DoubleResult[3]);
+		DefineComponentConsumer(GetClass(), FBuiltInComponentTypes::Get()->DoubleResult[4]);
+		DefineComponentConsumer(GetClass(), FBuiltInComponentTypes::Get()->DoubleResult[5]);
+
+		DefineComponentConsumer(GetClass(), FBuiltInComponentTypes::Get()->DoubleChannelFlags[0]);
+		DefineComponentConsumer(GetClass(), FBuiltInComponentTypes::Get()->DoubleChannelFlags[1]);
+		DefineComponentConsumer(GetClass(), FBuiltInComponentTypes::Get()->DoubleChannelFlags[2]);
+		DefineComponentConsumer(GetClass(), FBuiltInComponentTypes::Get()->DoubleChannelFlags[3]);
+		DefineComponentConsumer(GetClass(), FBuiltInComponentTypes::Get()->DoubleChannelFlags[4]);
+		DefineComponentConsumer(GetClass(), FBuiltInComponentTypes::Get()->DoubleChannelFlags[5]);
 	}
 }
 
@@ -177,22 +177,22 @@ void UMovieSceneTransformOriginSystem::OnRun(FSystemTaskPrerequisites& InPrerequ
 		FEntityTaskBuilder()
 		.Read(BuiltInComponents->InstanceHandle)
 		.Read(BuiltInComponents->BoundObject)
-		.WriteOptional(BuiltInComponents->FloatResult[0])
-		.WriteOptional(BuiltInComponents->FloatResult[1])
-		.WriteOptional(BuiltInComponents->FloatResult[2])
-		.WriteOptional(BuiltInComponents->FloatResult[3])
-		.WriteOptional(BuiltInComponents->FloatResult[4])
-		.WriteOptional(BuiltInComponents->FloatResult[5])
-		.WriteOptional(BuiltInComponents->FloatChannelFlags[0])
-		.WriteOptional(BuiltInComponents->FloatChannelFlags[1])
-		.WriteOptional(BuiltInComponents->FloatChannelFlags[2])
-		.WriteOptional(BuiltInComponents->FloatChannelFlags[3])
-		.WriteOptional(BuiltInComponents->FloatChannelFlags[4])
-		.WriteOptional(BuiltInComponents->FloatChannelFlags[5])
+		.WriteOptional(BuiltInComponents->DoubleResult[0])
+		.WriteOptional(BuiltInComponents->DoubleResult[1])
+		.WriteOptional(BuiltInComponents->DoubleResult[2])
+		.WriteOptional(BuiltInComponents->DoubleResult[3])
+		.WriteOptional(BuiltInComponents->DoubleResult[4])
+		.WriteOptional(BuiltInComponents->DoubleResult[5])
+		.WriteOptional(BuiltInComponents->DoubleChannelFlags[0])
+		.WriteOptional(BuiltInComponents->DoubleChannelFlags[1])
+		.WriteOptional(BuiltInComponents->DoubleChannelFlags[2])
+		.WriteOptional(BuiltInComponents->DoubleChannelFlags[3])
+		.WriteOptional(BuiltInComponents->DoubleChannelFlags[4])
+		.WriteOptional(BuiltInComponents->DoubleChannelFlags[5])
 		.CombineFilter(Filter)
-		// Must contain at least one float result
-		.FilterAny({ BuiltInComponents->FloatResult[0], BuiltInComponents->FloatResult[1], BuiltInComponents->FloatResult[2],
-			BuiltInComponents->FloatResult[3], BuiltInComponents->FloatResult[4], BuiltInComponents->FloatResult[5] })
+		// Must contain at least one double result
+		.FilterAny({ BuiltInComponents->DoubleResult[0], BuiltInComponents->DoubleResult[1], BuiltInComponents->DoubleResult[2],
+			BuiltInComponents->DoubleResult[3], BuiltInComponents->DoubleResult[4], BuiltInComponents->DoubleResult[5] })
 		.Dispatch_PerAllocation<FAssignTransformOrigin>(&Linker->EntityManager, InPrerequisites, &Subsequents, &TransformOriginsByInstanceID);
 	}
 }
diff --git a/Engine/Source/Runtime/MovieSceneTracks/Private/Systems/MovieSceneVectorPropertySystem.cpp b/Engine/Source/Runtime/MovieSceneTracks/Private/Systems/MovieSceneVectorPropertySystem.cpp
--- a/Engine/Source/Runtime/MovieSceneTracks/Private/Systems/MovieSceneVectorPropertySystem.cpp
+++ b/Engine/Source/Runtime/MovieSceneTracks/Private/Systems/MovieSceneVectorPropertySystem.cpp
@@ -5,16 +5,18 @@
 #include "EntitySystem/MovieSceneDecompositionQuery.h"
 #include "EntitySystem/MovieSceneEntitySystemTypes.h"
 #include "MovieSceneTracksComponentTypes.h"
+#include "Systems/DoubleChannelEvaluatorSystem.h"
 #include "Systems/FloatChannelEvaluatorSystem.h"
+#include "Systems/MovieScenePiecewiseDoubleBlenderSystem.h"
 #include "Systems/MovieScenePiecewiseFloatBlenderSystem.h"
 #include "Systems/MovieScenePropertyInstantiator.h"
 
-UMovieSceneVectorPropertySystem::UMovieSceneVectorPropertySystem(const FObjectInitializer& ObjInit)
+UMovieSceneFloatVectorPropertySystem::UMovieSceneFloatVectorPropertySystem(const FObjectInitializer& ObjInit)
 	: Super(ObjInit)
 {
 	SystemExclusionContext |= UE::MovieScene::EEntitySystemContext::Interrogation;
 
-	BindToProperty(UE::MovieScene::FMovieSceneTracksComponentTypes::Get()->Vector);
+	BindToProperty(UE::MovieScene::FMovieSceneTracksComponentTypes::Get()->FloatVector);
 
 	if (HasAnyFlags(RF_ClassDefaultObject))
 	{
@@ -24,7 +26,28 @@ UMovieSceneVectorPropertySystem::UMovieSceneVectorPropertySystem(const FObjectIn
 	}
 }
 
-void UMovieSceneVectorPropertySystem::OnRun(FSystemTaskPrerequisites& InPrerequisites, FSystemSubsequentTasks& Subsequents)
+void UMovieSceneFloatVectorPropertySystem::OnRun(FSystemTaskPrerequisites& InPrerequisites, FSystemSubsequentTasks& Subsequents)
 {
 	Super::OnRun(InPrerequisites, Subsequents);
 }
+
+UMovieSceneDoubleVectorPropertySystem::UMovieSceneDoubleVectorPropertySystem(const FObjectInitializer& ObjInit)
+	: Super(ObjInit)
+{
+	SystemExclusionContext |= UE::MovieScene::EEntitySystemContext::Interrogation;
+
+	BindToProperty(UE::MovieScene::FMovieSceneTracksComponentTypes::Get()->DoubleVector);
+
+	if (HasAnyFlags(RF_ClassDefaultObject))
+	{
+		// We need our floats correctly evaluated and blended, so we are downstream from those systems.
+		DefineImplicitPrerequisite(UDoubleChannelEvaluatorSystem::StaticClass(), GetClass());
+		DefineImplicitPrerequisite(UMovieScenePiecewiseDoubleBlenderSystem::StaticClass(), GetClass());
+	}
+}
+
+void UMovieSceneDoubleVectorPropertySystem::OnRun(FSystemTaskPrerequisites& InPrerequisites, FSystemSubsequentTasks& Subsequents)
+{
+	Super::OnRun(InPrerequisites, Subsequents);
+}
+
