Sequencer: Large world coordinates support

diff --git a/Engine/Source/Runtime/MovieSceneTracks/Private/Tracks/MovieScene3DTransformTrack.cpp b/Engine/Source/Runtime/MovieSceneTracks/Private/Tracks/MovieScene3DTransformTrack.cpp
--- a/Engine/Source/Runtime/MovieSceneTracks/Private/Tracks/MovieScene3DTransformTrack.cpp
+++ b/Engine/Source/Runtime/MovieSceneTracks/Private/Tracks/MovieScene3DTransformTrack.cpp
@@ -57,7 +57,7 @@ TArray<FTrajectoryKey> UMovieScene3DTransformTrack::GetTrajectoryData(FFrameNumb
 {
 	struct FCurveKeyIterator
 	{
-		FCurveKeyIterator(UMovieScene3DTransformSection* InSection, FMovieSceneFloatChannel* InChannel, FName InChannelName, FFrameNumber StartTime, TRange<FFrameNumber> ViewRange)
+		FCurveKeyIterator(UMovieScene3DTransformSection* InSection, FMovieSceneDoubleChannel* InChannel, FName InChannelName, FFrameNumber StartTime, TRange<FFrameNumber> ViewRange)
 			: Section(InSection), Channel(InChannel->GetData()), ChannelName(InChannelName), SectionRange(TRange<FFrameNumber>::Intersection(ViewRange, InSection->GetRange())), CurrentIndex(INDEX_NONE)
 		{
 			TArrayView<const FFrameNumber> Times = Channel.GetTimes();
@@ -162,7 +162,7 @@ TArray<FTrajectoryKey> UMovieScene3DTransformTrack::GetTrajectoryData(FFrameNumb
 
 	private:
 		UMovieScene3DTransformSection* Section;
-		TMovieSceneChannelData<FMovieSceneFloatValue> Channel;
+		TMovieSceneChannelData<FMovieSceneDoubleValue> Channel;
 		FName ChannelName;
 		TRange<FFrameNumber> SectionRange;
 		int32 CurrentIndex;
@@ -177,39 +177,39 @@ TArray<FTrajectoryKey> UMovieScene3DTransformTrack::GetTrajectoryData(FFrameNumb
 		UMovieScene3DTransformSection* TransformSection = Cast<UMovieScene3DTransformSection>(Section);
 		if (TransformSection)
 		{
-			TArrayView<FMovieSceneFloatChannel*>         FloatChannels = TransformSection->GetChannelProxy().GetChannels<FMovieSceneFloatChannel>();
-			TArrayView<const FMovieSceneChannelMetaData> MetaData      = TransformSection->GetChannelProxy().GetMetaData<FMovieSceneFloatChannel>();
+			TArrayView<FMovieSceneDoubleChannel*>        DoubleChannels = TransformSection->GetChannelProxy().GetChannels<FMovieSceneDoubleChannel>();
+			TArrayView<const FMovieSceneChannelMetaData> MetaData      = TransformSection->GetChannelProxy().GetMetaData<FMovieSceneDoubleChannel>();
 
 			EMovieSceneTransformChannel Mask = TransformSection->GetMask().GetChannels();
 			if (EnumHasAnyFlags(Mask, EMovieSceneTransformChannel::TranslationX))
 			{
-				ForwardIters.Emplace(TransformSection, FloatChannels[0], MetaData[0].Name, Time, ViewRange);
-				BackwardIters.Emplace(TransformSection, FloatChannels[0], MetaData[0].Name, Time, ViewRange);
+				ForwardIters.Emplace(TransformSection, DoubleChannels[0], MetaData[0].Name, Time, ViewRange);
+				BackwardIters.Emplace(TransformSection, DoubleChannels[0], MetaData[0].Name, Time, ViewRange);
 			}
 			if (EnumHasAnyFlags(Mask, EMovieSceneTransformChannel::TranslationY))
 			{
-				ForwardIters.Emplace(TransformSection, FloatChannels[1], MetaData[1].Name, Time, ViewRange);
-				BackwardIters.Emplace(TransformSection, FloatChannels[1], MetaData[1].Name, Time, ViewRange);
+				ForwardIters.Emplace(TransformSection, DoubleChannels[1], MetaData[1].Name, Time, ViewRange);
+				BackwardIters.Emplace(TransformSection, DoubleChannels[1], MetaData[1].Name, Time, ViewRange);
 			}
 			if (EnumHasAnyFlags(Mask, EMovieSceneTransformChannel::TranslationZ))
 			{
-				ForwardIters.Emplace(TransformSection, FloatChannels[2], MetaData[2].Name, Time, ViewRange);
-				BackwardIters.Emplace(TransformSection, FloatChannels[2], MetaData[2].Name, Time, ViewRange);
+				ForwardIters.Emplace(TransformSection, DoubleChannels[2], MetaData[2].Name, Time, ViewRange);
+				BackwardIters.Emplace(TransformSection, DoubleChannels[2], MetaData[2].Name, Time, ViewRange);
 			}
 			if (EnumHasAnyFlags(Mask, EMovieSceneTransformChannel::RotationX))
 			{
-				ForwardIters.Emplace(TransformSection, FloatChannels[3], MetaData[3].Name, Time, ViewRange);
-				BackwardIters.Emplace(TransformSection, FloatChannels[3], MetaData[3].Name, Time, ViewRange);
+				ForwardIters.Emplace(TransformSection, DoubleChannels[3], MetaData[3].Name, Time, ViewRange);
+				BackwardIters.Emplace(TransformSection, DoubleChannels[3], MetaData[3].Name, Time, ViewRange);
 			}
 			if (EnumHasAnyFlags(Mask, EMovieSceneTransformChannel::RotationY))
 			{
-				ForwardIters.Emplace(TransformSection, FloatChannels[4], MetaData[4].Name, Time, ViewRange);
-				BackwardIters.Emplace(TransformSection, FloatChannels[4], MetaData[4].Name, Time, ViewRange);
+				ForwardIters.Emplace(TransformSection, DoubleChannels[4], MetaData[4].Name, Time, ViewRange);
+				BackwardIters.Emplace(TransformSection, DoubleChannels[4], MetaData[4].Name, Time, ViewRange);
 			}
 			if (EnumHasAnyFlags(Mask, EMovieSceneTransformChannel::RotationZ))
 			{
-				ForwardIters.Emplace(TransformSection, FloatChannels[5], MetaData[5].Name, Time, ViewRange);
-				BackwardIters.Emplace(TransformSection, FloatChannels[5], MetaData[5].Name, Time, ViewRange);
+				ForwardIters.Emplace(TransformSection, DoubleChannels[5], MetaData[5].Name, Time, ViewRange);
+				BackwardIters.Emplace(TransformSection, DoubleChannels[5], MetaData[5].Name, Time, ViewRange);
 			}
 		}
 	}
diff --git a/Engine/Source/Runtime/MovieSceneTracks/Private/Tracks/MovieSceneVectorTrack.cpp b/Engine/Source/Runtime/MovieSceneTracks/Private/Tracks/MovieSceneVectorTrack.cpp
--- a/Engine/Source/Runtime/MovieSceneTracks/Private/Tracks/MovieSceneVectorTrack.cpp
+++ b/Engine/Source/Runtime/MovieSceneTracks/Private/Tracks/MovieSceneVectorTrack.cpp
@@ -4,7 +4,7 @@
 #include "Sections/MovieSceneVectorSection.h"
 
 
-UMovieSceneVectorTrack::UMovieSceneVectorTrack( const FObjectInitializer& ObjectInitializer )
+UMovieSceneFloatVectorTrack::UMovieSceneFloatVectorTrack( const FObjectInitializer& ObjectInitializer )
 	: Super( ObjectInitializer )
 {
 	NumChannelsUsed = 0;
@@ -12,14 +12,35 @@ UMovieSceneVectorTrack::UMovieSceneVectorTrack( const FObjectInitializer& Object
 }
 
 
-bool UMovieSceneVectorTrack::SupportsType(TSubclassOf<UMovieSceneSection> SectionClass) const
+bool UMovieSceneFloatVectorTrack::SupportsType(TSubclassOf<UMovieSceneSection> SectionClass) const
 {
-	return SectionClass == UMovieSceneVectorSection::StaticClass();
+	return SectionClass == UMovieSceneFloatVectorSection::StaticClass();
 }
 
-UMovieSceneSection* UMovieSceneVectorTrack::CreateNewSection()
+UMovieSceneSection* UMovieSceneFloatVectorTrack::CreateNewSection()
 {
-	UMovieSceneVectorSection* NewSection = NewObject<UMovieSceneVectorSection>(this, NAME_None, RF_Transactional);
+	UMovieSceneFloatVectorSection* NewSection = NewObject<UMovieSceneFloatVectorSection>(this, NAME_None, RF_Transactional);
+	NewSection->SetChannelsUsed(NumChannelsUsed);
+	return NewSection;
+}
+
+
+UMovieSceneDoubleVectorTrack::UMovieSceneDoubleVectorTrack( const FObjectInitializer& ObjectInitializer )
+	: Super( ObjectInitializer )
+{
+	NumChannelsUsed = 0;
+	SupportedBlendTypes = FMovieSceneBlendTypeField::All();
+}
+
+
+bool UMovieSceneDoubleVectorTrack::SupportsType(TSubclassOf<UMovieSceneSection> SectionClass) const
+{
+	return SectionClass == UMovieSceneDoubleVectorSection::StaticClass();
+}
+
+UMovieSceneSection* UMovieSceneDoubleVectorTrack::CreateNewSection()
+{
+	UMovieSceneDoubleVectorSection* NewSection = NewObject<UMovieSceneDoubleVectorSection>(this, NAME_None, RF_Transactional);
 	NewSection->SetChannelsUsed(NumChannelsUsed);
 	return NewSection;
 }
diff --git a/Engine/Source/Runtime/MovieSceneTracks/Public/EntitySystem/Interrogation/MovieSceneInterrogatedPropertyInstantiator.h b/Engine/Source/Runtime/MovieSceneTracks/Public/EntitySystem/Interrogation/MovieSceneInterrogatedPropertyInstantiator.h
--- a/Engine/Source/Runtime/MovieSceneTracks/Public/EntitySystem/Interrogation/MovieSceneInterrogatedPropertyInstantiator.h
+++ b/Engine/Source/Runtime/MovieSceneTracks/Public/EntitySystem/Interrogation/MovieSceneInterrogatedPropertyInstantiator.h
@@ -21,7 +21,7 @@ class MOVIESCENETRACKS_API UMovieSceneInterrogatedPropertyInstantiatorSystem
 public:
 
 	using FMovieSceneEntityID = UE::MovieScene::FMovieSceneEntityID;
-	using FFloatRecompositionResult = UE::MovieScene::TRecompositionResult<float>;
+	using FValueRecompositionResult = UE::MovieScene::TRecompositionResult<double>;
 
 	GENERATED_BODY()
 
@@ -36,7 +36,7 @@ public:
 	 * @param InCurrentValue The value of the property to recompose
 	 * @return A result containing the recomposed value for each of the entities specified in InQuery
 	 */
-	FFloatRecompositionResult RecomposeBlendFloatChannel(const UE::MovieScene::FPropertyDefinition& PropertyDefinition, int32 ChannelCompositeIndex, const UE::MovieScene::FDecompositionQuery& InQuery, float InCurrentValue);
+	FValueRecompositionResult RecomposeBlendChannel(const UE::MovieScene::FPropertyDefinition& PropertyDefinition, int32 ChannelCompositeIndex, const UE::MovieScene::FDecompositionQuery& InQuery, double InCurrentValue);
 
 	/**
 	 * Recompose a value from the constituent parts specified in InQuery, taking into account the weightings of each channel.
@@ -110,7 +110,7 @@ private:
 	virtual bool IsRelevantImpl(UMovieSceneEntitySystemLinker* InLinker) const override;
 
 	bool PropertySupportsFastPath(TArrayView<const FMovieSceneEntityID> Inputs, FPropertyInfo* Output) const;
-	UClass* ResolveBlenderClass(TArrayView<const FMovieSceneEntityID> Inputs) const;
+	UClass* ResolveBlenderClass(const UE::MovieScene::FPropertyDefinition& PropertyDefinition, TArrayView<const FMovieSceneEntityID> Inputs) const;
 	UE::MovieScene::FPropertyRecomposerPropertyInfo FindPropertyFromSource(FMovieSceneEntityID EntityID, UObject* Object) const;
 
 private:
diff --git a/Engine/Source/Runtime/MovieSceneTracks/Public/MovieSceneTracksComponentTypes.h b/Engine/Source/Runtime/MovieSceneTracks/Public/MovieSceneTracksComponentTypes.h
--- a/Engine/Source/Runtime/MovieSceneTracks/Public/MovieSceneTracksComponentTypes.h
+++ b/Engine/Source/Runtime/MovieSceneTracks/Public/MovieSceneTracksComponentTypes.h
@@ -146,80 +146,137 @@ struct FColorPropertyTraits
 	}
 };
 
-struct FVectorChannelMetaData
+struct FFloatVectorPropertyTraits
 {
-	uint8 NumChannels = 0;
+	using StorageType  = FFloatIntermediateVector;
+	using MetaDataType = TPropertyMetaData<FVectorPropertyMetaData>;
+
+	static void GetObjectPropertyValue(const UObject* InObject, FVectorPropertyMetaData MetaData, const FCustomPropertyAccessor& BaseCustomAccessor, FFloatIntermediateVector& OutValue)
+	{
+		const TCustomPropertyAccessor<FFloatVectorPropertyTraits>& CustomAccessor = static_cast<const TCustomPropertyAccessor<FFloatVectorPropertyTraits>&>(BaseCustomAccessor);
+		OutValue = (*CustomAccessor.Functions.Getter)(InObject, MetaData);
+	}
+	static void GetObjectPropertyValue(const UObject* InObject, FVectorPropertyMetaData MetaData, uint16 PropertyOffset, FFloatIntermediateVector& OutValue)
+	{
+		switch (MetaData.NumChannels)
+		{
+			case 2: TIndirectPropertyTraits<FVector2D, FFloatIntermediateVector>::GetObjectPropertyValue(InObject, PropertyOffset, OutValue); return;
+			case 3: TIndirectPropertyTraits<FVector3f, FFloatIntermediateVector>::GetObjectPropertyValue(InObject, PropertyOffset, OutValue); return;
+			case 4: TIndirectPropertyTraits<FVector4,  FFloatIntermediateVector>::GetObjectPropertyValue(InObject, PropertyOffset, OutValue); return;
+		}
+	}
+	static void GetObjectPropertyValue(const UObject* InObject, FVectorPropertyMetaData MetaData, FTrackInstancePropertyBindings* PropertyBindings, FFloatIntermediateVector& OutValue)
+	{
+		switch (MetaData.NumChannels)
+		{
+			case 2: TIndirectPropertyTraits<FVector2D, FFloatIntermediateVector>::GetObjectPropertyValue(InObject, PropertyBindings, OutValue); return;
+			case 3: TIndirectPropertyTraits<FVector3f, FFloatIntermediateVector>::GetObjectPropertyValue(InObject, PropertyBindings, OutValue); return;
+			case 4: TIndirectPropertyTraits<FVector4,  FFloatIntermediateVector>::GetObjectPropertyValue(InObject, PropertyBindings, OutValue); return;
+		}
+	}
+	static void GetObjectPropertyValue(const UObject* InObject, FVectorPropertyMetaData MetaData, const FName& PropertyPath, StorageType& OutValue)
+	{
+		switch (MetaData.NumChannels)
+		{
+			case 2: TIndirectPropertyTraits<FVector2D, FFloatIntermediateVector>::GetObjectPropertyValue(InObject, PropertyPath, OutValue); return;
+			case 3: TIndirectPropertyTraits<FVector3f, FFloatIntermediateVector>::GetObjectPropertyValue(InObject, PropertyPath, OutValue); return;
+			case 4: TIndirectPropertyTraits<FVector4,  FFloatIntermediateVector>::GetObjectPropertyValue(InObject, PropertyPath, OutValue); return;
+		}
+	}
+
+	static void SetObjectPropertyValue(UObject* InObject, FVectorPropertyMetaData MetaData, const FCustomPropertyAccessor& BaseCustomAccessor, const FFloatIntermediateVector& InValue)
+	{
+		const TCustomPropertyAccessor<FFloatVectorPropertyTraits>& CustomAccessor = static_cast<const TCustomPropertyAccessor<FFloatVectorPropertyTraits>&>(BaseCustomAccessor);
+		(*CustomAccessor.Functions.Setter)(InObject, MetaData, InValue);
+	}
+	static void SetObjectPropertyValue(UObject* InObject, FVectorPropertyMetaData MetaData, uint16 PropertyOffset, const FFloatIntermediateVector& InValue)
+	{
+		switch (MetaData.NumChannels)
+		{
+			case 2: TIndirectPropertyTraits<FVector2D, FFloatIntermediateVector>::SetObjectPropertyValue(InObject, PropertyOffset, InValue); return;
+			case 3: TIndirectPropertyTraits<FVector3f, FFloatIntermediateVector>::SetObjectPropertyValue(InObject, PropertyOffset, InValue); return;
+			case 4: TIndirectPropertyTraits<FVector4,  FFloatIntermediateVector>::SetObjectPropertyValue(InObject, PropertyOffset, InValue); return;
+		}
+
+		checkf(false, TEXT("Invalid number of channels"));
+	}
+	static void SetObjectPropertyValue(UObject* InObject, FVectorPropertyMetaData MetaData, FTrackInstancePropertyBindings* PropertyBindings, const FFloatIntermediateVector& InValue)
+	{
+		switch (MetaData.NumChannels)
+		{
+			case 2: TIndirectPropertyTraits<FVector2D, FFloatIntermediateVector>::SetObjectPropertyValue(InObject, PropertyBindings, InValue); return;
+			case 3: TIndirectPropertyTraits<FVector3f, FFloatIntermediateVector>::SetObjectPropertyValue(InObject, PropertyBindings, InValue); return;
+			case 4: TIndirectPropertyTraits<FVector4,  FFloatIntermediateVector>::SetObjectPropertyValue(InObject, PropertyBindings, InValue); return;
+		}
+
+		checkf(false, TEXT("Invalid number of channels"));
+	}
+
+	static FFloatIntermediateVector CombineComposites(FVectorPropertyMetaData MetaData, float InX, float InY, float InZ, float InW)
+	{
+		return FFloatIntermediateVector(InX, InY, InZ, InW);
+	}
 };
 
-struct FVectorPropertyTraits
+struct FDoubleVectorPropertyTraits
 {
-	using StorageType  = FIntermediateVector;
-	using MetaDataType = TPropertyMetaData<FVectorChannelMetaData>;
+	using StorageType  = FDoubleIntermediateVector;
+	using MetaDataType = TPropertyMetaData<FVectorPropertyMetaData>;
 
-	static void GetObjectPropertyValue(const UObject* InObject, FVectorChannelMetaData MetaData, const FCustomPropertyAccessor& BaseCustomAccessor, FIntermediateVector& OutValue)
+	static void GetObjectPropertyValue(const UObject* InObject, FVectorPropertyMetaData MetaData, const FCustomPropertyAccessor& BaseCustomAccessor, FDoubleIntermediateVector& OutValue)
 	{
-		const TCustomPropertyAccessor<FVectorPropertyTraits>& CustomAccessor = static_cast<const TCustomPropertyAccessor<FVectorPropertyTraits>&>(BaseCustomAccessor);
+		const TCustomPropertyAccessor<FDoubleVectorPropertyTraits>& CustomAccessor = static_cast<const TCustomPropertyAccessor<FDoubleVectorPropertyTraits>&>(BaseCustomAccessor);
 		OutValue = (*CustomAccessor.Functions.Getter)(InObject, MetaData);
 	}
-	static void GetObjectPropertyValue(const UObject* InObject, FVectorChannelMetaData MetaData, uint16 PropertyOffset, FIntermediateVector& OutValue)
+	static void GetObjectPropertyValue(const UObject* InObject, FVectorPropertyMetaData MetaData, uint16 PropertyOffset, FDoubleIntermediateVector& OutValue)
 	{
 		switch (MetaData.NumChannels)
 		{
-			case 2: TIndirectPropertyTraits<FVector2D, FIntermediateVector>::GetObjectPropertyValue(InObject, PropertyOffset, OutValue); return;
-			case 3: TIndirectPropertyTraits<FVector,   FIntermediateVector>::GetObjectPropertyValue(InObject, PropertyOffset, OutValue); return;
-			case 4: TIndirectPropertyTraits<FVector4,  FIntermediateVector>::GetObjectPropertyValue(InObject, PropertyOffset, OutValue); return;
+			case 3: TIndirectPropertyTraits<FVector3d, FDoubleIntermediateVector>::GetObjectPropertyValue(InObject, PropertyOffset, OutValue); return;
 		}
 	}
-	static void GetObjectPropertyValue(const UObject* InObject, FVectorChannelMetaData MetaData, FTrackInstancePropertyBindings* PropertyBindings, FIntermediateVector& OutValue)
+	static void GetObjectPropertyValue(const UObject* InObject, FVectorPropertyMetaData MetaData, FTrackInstancePropertyBindings* PropertyBindings, FDoubleIntermediateVector& OutValue)
 	{
 		switch (MetaData.NumChannels)
 		{
-			case 2: TIndirectPropertyTraits<FVector2D, FIntermediateVector>::GetObjectPropertyValue(InObject, PropertyBindings, OutValue); return;
-			case 3: TIndirectPropertyTraits<FVector,   FIntermediateVector>::GetObjectPropertyValue(InObject, PropertyBindings, OutValue); return;
-			case 4: TIndirectPropertyTraits<FVector4,  FIntermediateVector>::GetObjectPropertyValue(InObject, PropertyBindings, OutValue); return;
+			case 3: TIndirectPropertyTraits<FVector3d, FDoubleIntermediateVector>::GetObjectPropertyValue(InObject, PropertyBindings, OutValue); return;
 		}
 	}
-	static void GetObjectPropertyValue(const UObject* InObject, FVectorChannelMetaData MetaData, const FName& PropertyPath, StorageType& OutValue)
+	static void GetObjectPropertyValue(const UObject* InObject, FVectorPropertyMetaData MetaData, const FName& PropertyPath, StorageType& OutValue)
 	{
 		switch (MetaData.NumChannels)
 		{
-			case 2: TIndirectPropertyTraits<FVector2D, FIntermediateVector>::GetObjectPropertyValue(InObject, PropertyPath, OutValue); return;
-			case 3: TIndirectPropertyTraits<FVector,   FIntermediateVector>::GetObjectPropertyValue(InObject, PropertyPath, OutValue); return;
-			case 4: TIndirectPropertyTraits<FVector4,  FIntermediateVector>::GetObjectPropertyValue(InObject, PropertyPath, OutValue); return;
+			case 3: TIndirectPropertyTraits<FVector3d, FDoubleIntermediateVector>::GetObjectPropertyValue(InObject, PropertyPath, OutValue); return;
 		}
 	}
 
-	static void SetObjectPropertyValue(UObject* InObject, FVectorChannelMetaData MetaData, const FCustomPropertyAccessor& BaseCustomAccessor, const FIntermediateVector& InValue)
+	static void SetObjectPropertyValue(UObject* InObject, FVectorPropertyMetaData MetaData, const FCustomPropertyAccessor& BaseCustomAccessor, const FDoubleIntermediateVector& InValue)
 	{
-		const TCustomPropertyAccessor<FVectorPropertyTraits>& CustomAccessor = static_cast<const TCustomPropertyAccessor<FVectorPropertyTraits>&>(BaseCustomAccessor);
+		const TCustomPropertyAccessor<FDoubleVectorPropertyTraits>& CustomAccessor = static_cast<const TCustomPropertyAccessor<FDoubleVectorPropertyTraits>&>(BaseCustomAccessor);
 		(*CustomAccessor.Functions.Setter)(InObject, MetaData, InValue);
 	}
-	static void SetObjectPropertyValue(UObject* InObject, FVectorChannelMetaData MetaData, uint16 PropertyOffset, const FIntermediateVector& InValue)
+	static void SetObjectPropertyValue(UObject* InObject, FVectorPropertyMetaData MetaData, uint16 PropertyOffset, const FDoubleIntermediateVector& InValue)
 	{
 		switch (MetaData.NumChannels)
 		{
-			case 2: TIndirectPropertyTraits<FVector2D, FIntermediateVector>::SetObjectPropertyValue(InObject, PropertyOffset, InValue); return;
-			case 3: TIndirectPropertyTraits<FVector,   FIntermediateVector>::SetObjectPropertyValue(InObject, PropertyOffset, InValue); return;
-			case 4: TIndirectPropertyTraits<FVector4,  FIntermediateVector>::SetObjectPropertyValue(InObject, PropertyOffset, InValue); return;
+			case 3: TIndirectPropertyTraits<FVector3d, FDoubleIntermediateVector>::SetObjectPropertyValue(InObject, PropertyOffset, InValue); return;
 		}
 
 		checkf(false, TEXT("Invalid number of channels"));
 	}
-	static void SetObjectPropertyValue(UObject* InObject, FVectorChannelMetaData MetaData, FTrackInstancePropertyBindings* PropertyBindings, const FIntermediateVector& InValue)
+	static void SetObjectPropertyValue(UObject* InObject, FVectorPropertyMetaData MetaData, FTrackInstancePropertyBindings* PropertyBindings, const FDoubleIntermediateVector& InValue)
 	{
 		switch (MetaData.NumChannels)
 		{
-			case 2: TIndirectPropertyTraits<FVector2D, FIntermediateVector>::SetObjectPropertyValue(InObject, PropertyBindings, InValue); return;
-			case 3: TIndirectPropertyTraits<FVector,   FIntermediateVector>::SetObjectPropertyValue(InObject, PropertyBindings, InValue); return;
-			case 4: TIndirectPropertyTraits<FVector4,  FIntermediateVector>::SetObjectPropertyValue(InObject, PropertyBindings, InValue); return;
+			case 3: TIndirectPropertyTraits<FVector3d, FDoubleIntermediateVector>::SetObjectPropertyValue(InObject, PropertyBindings, InValue); return;
 		}
 
 		checkf(false, TEXT("Invalid number of channels"));
 	}
 
-	static FIntermediateVector CombineComposites(FVectorChannelMetaData MetaData, float InX, float InY, float InZ, float InW)
+	static FDoubleIntermediateVector CombineComposites(FVectorPropertyMetaData MetaData, double InX, double InY, double InZ, double InW)
 	{
-		return FIntermediateVector(InX, InY, InZ, InW);
+		return FDoubleIntermediateVector(InX, InY, InZ, InW);
 	}
 };
 
@@ -228,6 +285,7 @@ using FBytePropertyTraits               = TDirectPropertyTraits<uint8>;
 using FEnumPropertyTraits               = TDirectPropertyTraits<uint8>;
 using FIntPropertyTraits                = TDirectPropertyTraits<int32>;
 using FFloatPropertyTraits              = TDirectPropertyTraits<float>;
+using FDoublePropertyTraits             = TDirectPropertyTraits<double>;
 using FTransformPropertyTraits          = TIndirectPropertyTraits<FTransform, FIntermediate3DTransform>;
 using FEulerTransformPropertyTraits     = TIndirectPropertyTraits<FEulerTransform, FIntermediate3DTransform>;
 using FComponentTransformPropertyTraits = TDirectPropertyTraits<FIntermediate3DTransform>;
@@ -242,12 +300,14 @@ struct MOVIESCENETRACKS_API FMovieSceneTracksComponentTypes
 	TPropertyComponents<FEnumPropertyTraits> Enum;
 	TPropertyComponents<FIntPropertyTraits> Integer;
 	TPropertyComponents<FFloatPropertyTraits> Float;
+	TPropertyComponents<FDoublePropertyTraits> Double;
 	TPropertyComponents<FColorPropertyTraits> Color;
-	TPropertyComponents<FVectorPropertyTraits> Vector;
+	TPropertyComponents<FFloatVectorPropertyTraits> FloatVector;
+	TPropertyComponents<FDoubleVectorPropertyTraits> DoubleVector;
 	TPropertyComponents<FTransformPropertyTraits> Transform;
 	TPropertyComponents<FEulerTransformPropertyTraits> EulerTransform;
 	TPropertyComponents<FComponentTransformPropertyTraits> ComponentTransform;
-	TComponentTypeID<FSourceFloatChannel> QuaternionRotationChannel[3];
+	TComponentTypeID<FSourceDoubleChannel> QuaternionRotationChannel[3];
 
 	TComponentTypeID<USceneComponent*> AttachParent;
 	TComponentTypeID<FAttachmentComponent> AttachComponent;
@@ -260,8 +320,10 @@ struct MOVIESCENETRACKS_API FMovieSceneTracksComponentTypes
 		TCustomPropertyRegistration<FEnumPropertyTraits> Enum;
 		TCustomPropertyRegistration<FIntPropertyTraits> Integer;
 		TCustomPropertyRegistration<FFloatPropertyTraits> Float;
+		TCustomPropertyRegistration<FDoublePropertyTraits> Double;
 		TCustomPropertyRegistration<FColorPropertyTraits> Color;
-		TCustomPropertyRegistration<FVectorPropertyTraits> Vector;
+		TCustomPropertyRegistration<FFloatVectorPropertyTraits> FloatVector;
+		TCustomPropertyRegistration<FDoubleVectorPropertyTraits> DoubleVector;
 		TCustomPropertyRegistration<FComponentTransformPropertyTraits, 1> ComponentTransform;
 	} Accessors;
 
diff --git a/Engine/Source/Runtime/MovieSceneTracks/Public/MovieSceneTracksPropertyTypes.h b/Engine/Source/Runtime/MovieSceneTracks/Public/MovieSceneTracksPropertyTypes.h
--- a/Engine/Source/Runtime/MovieSceneTracks/Public/MovieSceneTracksPropertyTypes.h
+++ b/Engine/Source/Runtime/MovieSceneTracks/Public/MovieSceneTracksPropertyTypes.h
@@ -10,6 +10,7 @@
 #include "EulerTransform.h"
 #include "Styling/SlateColor.h"
 #include "EntitySystem/MovieSceneComponentDebug.h"
+#include "Misc/LargeWorldCoordinates.h"
 
 class USceneComponent;
 
@@ -19,27 +20,49 @@ namespace UE
 namespace MovieScene
 {
 
-
-
+struct FVectorPropertyMetaData
+{
+	uint8 NumChannels = 0;
+};
 
 /** Intermediate type for the vector property system that lets us store how many dimensions the vector should have */
-struct FIntermediateVector
+struct FFloatIntermediateVector
 {
 	float X, Y, Z, W;
 
-	FIntermediateVector()
+	FFloatIntermediateVector()
 		: X(0), Y(0), Z(0), W(0)
 	{}
 
-	FIntermediateVector(float InX, float InY)
+	FFloatIntermediateVector(float InX, float InY)
 		: X(InX), Y(InY), Z(0), W(0)
 	{}
 
-	FIntermediateVector(float InX, float InY, float InZ)
+	FFloatIntermediateVector(float InX, float InY, float InZ)
 		: X(InX), Y(InY), Z(InZ), W(0)
 	{}
 
-	FIntermediateVector(float InX, float InY, float InZ, float InW)
+	FFloatIntermediateVector(float InX, float InY, float InZ, float InW)
+		: X(InX), Y(InY), Z(InZ), W(InW)
+	{}
+};
+struct FDoubleIntermediateVector
+{
+	double X, Y, Z, W;
+
+	FDoubleIntermediateVector()
+		: X(0), Y(0), Z(0), W(0)
+	{}
+
+	FDoubleIntermediateVector(double InX, double InY)
+		: X(InX), Y(InY), Z(0), W(0)
+	{}
+
+	FDoubleIntermediateVector(double InX, double InY, double InZ)
+		: X(InX), Y(InY), Z(InZ), W(0)
+	{}
+
+	FDoubleIntermediateVector(double InX, double InY, double InZ, double InW)
 		: X(InX), Y(InY), Z(InZ), W(InW)
 	{}
 };
@@ -118,13 +141,15 @@ struct FIntermediateColor
 /** Intermediate type used for applying partially animated transforms. Saves us from repteatedly recomposing quaternions from euler angles */
 struct FIntermediate3DTransform
 {
-	float T_X, T_Y, T_Z, R_X, R_Y, R_Z, S_X, S_Y, S_Z;
+	// When LWC is enabled, translations are manipulated as doubles.
+	double T_X, T_Y, T_Z, R_X, R_Y, R_Z, S_X, S_Y, S_Z;
 
 	FIntermediate3DTransform()
-		: T_X(0.f), T_Y(0.f), T_Z(0.f), R_X(0.f), R_Y(0.f), R_Z(0.f), S_X(0.f), S_Y(0.f), S_Z(0.f)
+		: T_X(0.), T_Y(0.), T_Z(0.), R_X(0.), R_Y(0.), R_Z(0.), S_X(0.), S_Y(0.), S_Z(0.)
 	{}
 
-	FIntermediate3DTransform(float InT_X, float InT_Y, float InT_Z, float InR_X, float InR_Y, float InR_Z, float InS_X, float InS_Y, float InS_Z)
+	FIntermediate3DTransform(
+			double InT_X, double InT_Y, double InT_Z, double InR_X, double InR_Y, double InR_Z, double InS_X, double InS_Y, double InS_Z)
 		: T_X(InT_X), T_Y(InT_Y), T_Z(InT_Z), R_X(InR_X), R_Y(InR_Y), R_Z(InR_Z), S_X(InS_X), S_Y(InS_Y), S_Z(InS_Z)
 	{}
 
@@ -134,7 +159,7 @@ struct FIntermediate3DTransform
 		, S_X(InScale.X), S_Y(InScale.Y), S_Z(InScale.Z)
 	{}
 
-	float operator[](int32 Index) const
+	double operator[](int32 Index) const
 	{
 		check(Index >= 0 && Index < 9);
 		return (&T_X)[Index];
@@ -171,13 +196,15 @@ MOVIESCENETRACKS_API void ConvertOperationalProperty(const FColor& InColor, FInt
 MOVIESCENETRACKS_API void ConvertOperationalProperty(const FLinearColor& InColor, FIntermediateColor& OutIntermediate);
 MOVIESCENETRACKS_API void ConvertOperationalProperty(const FSlateColor& InColor, FIntermediateColor& OutIntermediate);
 
-MOVIESCENETRACKS_API void ConvertOperationalProperty(const FIntermediateVector& InVector, FVector2D& Out);
-MOVIESCENETRACKS_API void ConvertOperationalProperty(const FIntermediateVector& InVector, FVector& Out);
-MOVIESCENETRACKS_API void ConvertOperationalProperty(const FIntermediateVector& InVector, FVector4& Out);
-MOVIESCENETRACKS_API void ConvertOperationalProperty(const FVector2D& In, FIntermediateVector& Out);
-MOVIESCENETRACKS_API void ConvertOperationalProperty(const FVector& In, FIntermediateVector& Out);
-MOVIESCENETRACKS_API void ConvertOperationalProperty(const FVector4& In, FIntermediateVector& Out);
+MOVIESCENETRACKS_API void ConvertOperationalProperty(const FFloatIntermediateVector& InVector, FVector2D& Out);
+MOVIESCENETRACKS_API void ConvertOperationalProperty(const FFloatIntermediateVector& InVector, FVector3f& Out);
+MOVIESCENETRACKS_API void ConvertOperationalProperty(const FFloatIntermediateVector& InVector, FVector4& Out);
+MOVIESCENETRACKS_API void ConvertOperationalProperty(const FVector2D& In, FFloatIntermediateVector& Out);
+MOVIESCENETRACKS_API void ConvertOperationalProperty(const FVector3f& In, FFloatIntermediateVector& Out);
+MOVIESCENETRACKS_API void ConvertOperationalProperty(const FVector4& In, FFloatIntermediateVector& Out);
 
+MOVIESCENETRACKS_API void ConvertOperationalProperty(const FDoubleIntermediateVector& InVector, FVector3d& Out);
+MOVIESCENETRACKS_API void ConvertOperationalProperty(const FVector3d& In, FDoubleIntermediateVector& Out);
 
 } // namespace MovieScene
 } // namespace UE
diff --git a/Engine/Source/Runtime/MovieSceneTracks/Public/Sections/MovieScene3DTransformSection.h b/Engine/Source/Runtime/MovieSceneTracks/Public/Sections/MovieScene3DTransformSection.h
--- a/Engine/Source/Runtime/MovieSceneTracks/Public/Sections/MovieScene3DTransformSection.h
+++ b/Engine/Source/Runtime/MovieSceneTracks/Public/Sections/MovieScene3DTransformSection.h
@@ -8,9 +8,11 @@
 #include "Curves/KeyHandle.h"
 #include "MovieSceneSection.h"
 #include "MovieSceneKeyStruct.h"
+#include "Channels/MovieSceneDoubleChannel.h"
 #include "Channels/MovieSceneFloatChannel.h"
 #include "EntitySystem/IMovieSceneEntityProvider.h"
 #include "TransformData.h"
+#include "Misc/LargeWorldCoordinates.h"
 #include "MovieScene3DTransformSection.generated.h"
 
 
@@ -26,7 +28,6 @@ enum class EShow3DTrajectory : uint8
 #endif
 
 
-
 /**
  * Proxy structure for translation keys in 3D transform sections.
  */
@@ -38,7 +39,7 @@ struct FMovieScene3DLocationKeyStruct
 
 	/** The key's translation value. */
 	UPROPERTY(EditAnywhere, Category=Key)
-	FVector3f Location = FVector3f::ZeroVector; // LWC_TODO: Stored into FMovieSceneFloatChannel (float*). Implement FMovieSceneDoubleChannel?
+	FVector Location = FVector::ZeroVector;
 
 	/** The key's time. */
 	UPROPERTY(EditAnywhere, Category=Key)
@@ -85,7 +86,7 @@ struct FMovieScene3DScaleKeyStruct
 
 	/** The key's scale value. */
 	UPROPERTY(EditAnywhere, Category=Key)
-	FVector3f Scale = FVector::OneVector;	// LWC_TODO: Stored into FMovieSceneFloatChannel (float*). Implement FMovieSceneDoubleChannel?
+	FVector3f Scale = FVector::OneVector;
 
 	/** The key's time. */
 	UPROPERTY(EditAnywhere, Category=Key)
@@ -109,7 +110,7 @@ struct FMovieScene3DTransformKeyStruct
 
 	/** The key's translation value. */
 	UPROPERTY(EditAnywhere, Category=Key)
-	FVector3f Location = FVector::ZeroVector;	// LWC_TODO: Stored into FMovieSceneFloatChannel (float*). Implement FMovieSceneDoubleChannel?
+	FVector Location = FVector::ZeroVector;
 
 	/** The key's rotation value. */
 	UPROPERTY(EditAnywhere, Category=Key)
@@ -117,7 +118,7 @@ struct FMovieScene3DTransformKeyStruct
 
 	/** The key's scale value. */
 	UPROPERTY(EditAnywhere, Category=Key)
-	FVector3f Scale = FVector::OneVector;	// LWC_TODO: Stored into FMovieSceneFloatChannel (float*). Implement FMovieSceneDoubleChannel?
+	FVector3f Scale = FVector::OneVector;
 
 	/** The key's time. */
 	UPROPERTY(EditAnywhere, Category=Key)
@@ -271,15 +272,15 @@ private:
 
 	/** Translation curves */
 	UPROPERTY()
-	FMovieSceneFloatChannel Translation[3];
+	FMovieSceneDoubleChannel Translation[3];
 	
 	/** Rotation curves */
 	UPROPERTY()
-	FMovieSceneFloatChannel Rotation[3];
+	FMovieSceneDoubleChannel Rotation[3];
 
 	/** Scale curves */
 	UPROPERTY()
-	FMovieSceneFloatChannel Scale[3];
+	FMovieSceneDoubleChannel Scale[3];
 
 	/** Manual weight curve */
 	UPROPERTY()
diff --git a/Engine/Source/Runtime/MovieSceneTracks/Public/Sections/MovieSceneVectorSection.h b/Engine/Source/Runtime/MovieSceneTracks/Public/Sections/MovieSceneVectorSection.h
--- a/Engine/Source/Runtime/MovieSceneTracks/Public/Sections/MovieSceneVectorSection.h
+++ b/Engine/Source/Runtime/MovieSceneTracks/Public/Sections/MovieSceneVectorSection.h
@@ -9,6 +9,7 @@
 #include "MovieSceneSection.h"
 #include "MovieSceneKeyStruct.h"
 #include "Channels/MovieSceneFloatChannel.h"
+#include "Channels/MovieSceneDoubleChannel.h"
 #include "EntitySystem/IMovieSceneEntityProvider.h"
 #include "MovieSceneVectorSection.generated.h"
 
@@ -16,10 +17,10 @@ class FStructOnScope;
 struct FPropertyChangedEvent;
 
 /**
-* Base Proxy structure for vector section key data.
+* Base Proxy structure for float vector section key data.
 */
 USTRUCT()
-struct FMovieSceneVectorKeyStructBase
+struct FMovieSceneFloatVectorKeyStructBase
 	: public FMovieSceneKeyStruct
 {
 	GENERATED_BODY();
@@ -33,9 +34,9 @@ struct FMovieSceneVectorKeyStructBase
 	virtual void PropagateChanges(const FPropertyChangedEvent& ChangeEvent) override;
 
 	/** Gets a ptr value of a channel by index, 0-3 = x-w */
-	virtual float* GetPropertyChannelByIndex(int32 Index) PURE_VIRTUAL(FMovieSceneVectorKeyStructBase::GetPropertyChannelByIndex, return nullptr; );
+	virtual float* GetPropertyChannelByIndex(int32 Index) PURE_VIRTUAL(FMovieSceneFloatVectorKeyStructBase::GetPropertyChannelByIndex, return nullptr; );
 };
-template<> struct TStructOpsTypeTraits<FMovieSceneVectorKeyStructBase> : public TStructOpsTypeTraitsBase2<FMovieSceneVectorKeyStructBase> { enum { WithCopy = false, WithPureVirtual = true, }; };
+template<> struct TStructOpsTypeTraits<FMovieSceneFloatVectorKeyStructBase> : public TStructOpsTypeTraitsBase2<FMovieSceneFloatVectorKeyStructBase> { enum { WithCopy = false, WithPureVirtual = true, }; };
 
 
 /**
@@ -43,7 +44,7 @@ template<> struct TStructOpsTypeTraits<FMovieSceneVectorKeyStructBase> : public
  */
 USTRUCT()
 struct FMovieSceneVector2DKeyStruct
-	: public FMovieSceneVectorKeyStructBase
+	: public FMovieSceneFloatVectorKeyStructBase
 {
 	GENERATED_BODY()
 
@@ -51,17 +52,17 @@ struct FMovieSceneVector2DKeyStruct
 	UPROPERTY(EditAnywhere, Category=Key)
 	FVector2D Vector = FVector2D::ZeroVector;
 
-	//~ FMovieSceneVectorKeyStructBase interface
+	//~ FMovieSceneFloatVectorKeyStructBase interface
 	virtual float* GetPropertyChannelByIndex(int32 Index) override { return &Vector[Index]; }
 };
 template<> struct TStructOpsTypeTraits<FMovieSceneVector2DKeyStruct> : public TStructOpsTypeTraitsBase2<FMovieSceneVector2DKeyStruct> { enum { WithCopy = false }; };
 
 /**
-* Proxy structure for vector section key data.
+* Proxy structure for float vector section key data.
 */
 USTRUCT()
-struct FMovieSceneVectorKeyStruct
-	: public FMovieSceneVectorKeyStructBase
+struct FMovieSceneVector3fKeyStruct
+	: public FMovieSceneFloatVectorKeyStructBase
 {
 	GENERATED_BODY()
 
@@ -69,17 +70,17 @@ struct FMovieSceneVectorKeyStruct
 	UPROPERTY(EditAnywhere, Category = Key)
 	FVector3f Vector = FVector3f::ZeroVector;
 
-	//~ FMovieSceneVectorKeyStructBase interface
+	//~ FMovieSceneFloatVectorKeyStructBase interface
 	virtual float* GetPropertyChannelByIndex(int32 Index) override { return &Vector[Index]; }
 };
-template<> struct TStructOpsTypeTraits<FMovieSceneVectorKeyStruct> : public TStructOpsTypeTraitsBase2<FMovieSceneVectorKeyStruct> { enum { WithCopy = false }; };
+template<> struct TStructOpsTypeTraits<FMovieSceneVector3fKeyStruct> : public TStructOpsTypeTraitsBase2<FMovieSceneVector3fKeyStruct> { enum { WithCopy = false }; };
 
 /**
 * Proxy structure for vector4 section key data.
 */
 USTRUCT()
 struct FMovieSceneVector4KeyStruct
-	: public FMovieSceneVectorKeyStructBase
+	: public FMovieSceneFloatVectorKeyStructBase
 {
 	GENERATED_BODY()
 
@@ -87,17 +88,58 @@ struct FMovieSceneVector4KeyStruct
 	UPROPERTY(EditAnywhere, Category = Key)
 	FVector4 Vector = FVector4(FVector::ZeroVector);
 
-	//~ FMovieSceneVectorKeyStructBase interface
+	//~ FMovieSceneFloatVectorKeyStructBase interface
 	virtual float* GetPropertyChannelByIndex(int32 Index) override { return &Vector[Index]; }
 };
 template<> struct TStructOpsTypeTraits<FMovieSceneVector4KeyStruct> : public TStructOpsTypeTraitsBase2<FMovieSceneVector4KeyStruct> { enum { WithCopy = false }; };
 
 
 /**
- * A vector section.
+* Base Proxy structure for double vector section key data.
+*/
+USTRUCT()
+struct FMovieSceneDoubleVectorKeyStructBase
+	: public FMovieSceneKeyStruct
+{
+	GENERATED_BODY();
+
+	/** The key's time. */
+	UPROPERTY(EditAnywhere, Category=Key)
+	FFrameNumber Time;
+
+	FMovieSceneKeyStructHelper KeyStructInterop;
+
+	virtual void PropagateChanges(const FPropertyChangedEvent& ChangeEvent) override;
+
+	/** Gets a ptr value of a channel by index, 0-3 = x-w */
+	virtual double* GetPropertyChannelByIndex(int32 Index) PURE_VIRTUAL(FMovieSceneDoubleVectorKeyStructBase::GetPropertyChannelByIndex, return nullptr; );
+};
+template<> struct TStructOpsTypeTraits<FMovieSceneDoubleVectorKeyStructBase> : public TStructOpsTypeTraitsBase2<FMovieSceneDoubleVectorKeyStructBase> { enum { WithCopy = false, WithPureVirtual = true, }; };
+
+/**
+* Proxy structure for double vector section key data.
+*/
+USTRUCT()
+struct FMovieSceneVector3dKeyStruct
+	: public FMovieSceneDoubleVectorKeyStructBase
+{
+	GENERATED_BODY()
+
+	/** They key's vector value. */
+	UPROPERTY(EditAnywhere, Category = Key)
+	FVector3d Vector = FVector3d::ZeroVector;
+
+	//~ FMovieSceneDoubleVectorKeyStructBase interface
+	virtual double* GetPropertyChannelByIndex(int32 Index) override { return &Vector[Index]; }
+};
+template<> struct TStructOpsTypeTraits<FMovieSceneVector3dKeyStruct> : public TStructOpsTypeTraitsBase2<FMovieSceneVector3dKeyStruct> { enum { WithCopy = false }; };
+
+
+/**
+ * A float vector section.
  */
 UCLASS(MinimalAPI)
-class UMovieSceneVectorSection
+class UMovieSceneFloatVectorSection
 	: public UMovieSceneSection
 	, public IMovieSceneEntityProvider
 {
@@ -150,3 +192,62 @@ private:
 	UPROPERTY()
 	int32 ChannelsUsed;
 };
+
+
+/**
+ * A double vector section.
+ */
+UCLASS(MinimalAPI)
+class UMovieSceneDoubleVectorSection
+	: public UMovieSceneSection
+	, public IMovieSceneEntityProvider
+{
+	GENERATED_UCLASS_BODY()
+
+public:
+
+	/** Sets how many channels are to be used */
+	void SetChannelsUsed(int32 InChannelsUsed) 
+	{
+		checkf(InChannelsUsed >= 2 && InChannelsUsed <= 4, TEXT("Only 2-4 channels are supported."));
+		ChannelsUsed = InChannelsUsed;
+		RecreateChannelProxy();
+	}
+
+	/** Gets the number of channels in use */
+	int32 GetChannelsUsed() const { return ChannelsUsed; }
+
+	/**
+	 * Public access to this section's internal data function
+	 */
+	const FMovieSceneDoubleChannel& GetChannel(int32 Index) const
+	{
+		check(Index >= 0 && Index < GetChannelsUsed());
+		return Curves[Index];
+	}
+
+protected:
+
+	//~ UMovieSceneSection interface
+	virtual void Serialize(FArchive& Ar) override;
+	virtual void PostEditImport() override;
+	virtual TSharedPtr<FStructOnScope> GetKeyStruct(TArrayView<const FKeyHandle> KeyHandles) override;
+
+	MOVIESCENETRACKS_API void RecreateChannelProxy();
+
+private:
+
+	//~ IMovieSceneEntityProvider interface
+	virtual void ImportEntityImpl(UMovieSceneEntitySystemLinker* EntityLinker, const FEntityImportParams& Params, FImportedEntity* OutImportedEntity) override;
+	virtual bool PopulateEvaluationFieldImpl(const TRange<FFrameNumber>& EffectiveRange, const FMovieSceneEvaluationFieldEntityMetaData& InMetaData, FMovieSceneEntityComponentFieldBuilder* OutFieldBuilder) override;
+
+private:
+
+	/** Double functions for the X,Y,Z,W components of the vector */
+	UPROPERTY()
+	FMovieSceneDoubleChannel Curves[4];
+
+	/** How many curves are actually used */
+	UPROPERTY()
+	int32 ChannelsUsed;
+};
diff --git a/Engine/Source/Runtime/MovieSceneTracks/Public/Systems/DoubleChannelEvaluatorSystem.h b/Engine/Source/Runtime/MovieSceneTracks/Public/Systems/DoubleChannelEvaluatorSystem.h
new file mode 100644
--- /dev/null
+++ b/Engine/Source/Runtime/MovieSceneTracks/Public/Systems/DoubleChannelEvaluatorSystem.h
@@ -0,0 +1,51 @@
+// Copyright Epic Games, Inc. All Rights Reserved.
+
+#pragma once
+
+#include "Containers/SortedMap.h"
+#include "Containers/Set.h"
+#include "EntitySystem/MovieSceneEntitySystem.h"
+#include "EntitySystem/EntityAllocationIterator.h"
+
+#include "DoubleChannelEvaluatorSystem.generated.h"
+
+namespace UE
+{
+namespace MovieScene
+{
+
+	struct FSourceDoubleChannel;
+	struct FSourceDoubleChannelFlags;
+
+} // namespace MovieScene
+} // namespace UE
+
+/**
+ * System that is responsible for evaluating double channels.
+ */
+UCLASS()
+class MOVIESCENETRACKS_API UDoubleChannelEvaluatorSystem : public UMovieSceneEntitySystem
+{
+public:
+
+
+	GENERATED_BODY()
+
+	UDoubleChannelEvaluatorSystem(const FObjectInitializer& ObjInit);
+
+	virtual void OnRun(FSystemTaskPrerequisites& InPrerequisites, FSystemSubsequentTasks& Subsequents) override;
+	virtual bool IsRelevantImpl(UMovieSceneEntitySystemLinker* InLinker) const override;
+
+	static void RegisterChannelType(TComponentTypeID<UE::MovieScene::FSourceDoubleChannel> SourceChannelType, TComponentTypeID<UE::MovieScene::FSourceDoubleChannelFlags> ChannelFlagsType, TComponentTypeID<double> ResultType);
+
+private:
+
+	struct FChannelType
+	{
+		TComponentTypeID<UE::MovieScene::FSourceDoubleChannel> ChannelType;
+		TComponentTypeID<UE::MovieScene::FSourceDoubleChannelFlags> ChannelFlagsType;
+		TComponentTypeID<double> ResultType;
+	};
+
+	static TArray<FChannelType, TInlineAllocator<4>> StaticChannelTypes;
+};
diff --git a/Engine/Source/Runtime/MovieSceneTracks/Public/Systems/MovieSceneDoublePropertySystem.h b/Engine/Source/Runtime/MovieSceneTracks/Public/Systems/MovieSceneDoublePropertySystem.h
new file mode 100644
--- /dev/null
+++ b/Engine/Source/Runtime/MovieSceneTracks/Public/Systems/MovieSceneDoublePropertySystem.h
@@ -0,0 +1,18 @@
+// Copyright Epic Games, Inc. All Rights Reserved.
+
+#pragma once
+
+#include "Systems/MovieScenePropertySystem.h"
+#include "MovieSceneDoublePropertySystem.generated.h"
+
+
+
+UCLASS()
+class MOVIESCENETRACKS_API UMovieSceneDoublePropertySystem : public UMovieScenePropertySystem
+{
+	GENERATED_BODY()
+
+	UMovieSceneDoublePropertySystem(const FObjectInitializer& ObjInit);
+
+	virtual void OnRun(FSystemTaskPrerequisites& InPrerequisites, FSystemSubsequentTasks& Subsequents) override;
+};
diff --git a/Engine/Source/Runtime/MovieSceneTracks/Public/Systems/MovieScenePiecewiseBlenderSystemImpl.h b/Engine/Source/Runtime/MovieSceneTracks/Public/Systems/MovieScenePiecewiseBlenderSystemImpl.h
new file mode 100644
--- /dev/null
+++ b/Engine/Source/Runtime/MovieSceneTracks/Public/Systems/MovieScenePiecewiseBlenderSystemImpl.h
@@ -0,0 +1,142 @@
+// Copyright Epic Games, Inc. All Rights Reserved.
+
+#pragma once
+
+#include "EntitySystem/MovieSceneEntitySystem.h"
+#include "EntitySystem/MovieSceneDecompositionQuery.h"
+#include "EntitySystem/MovieSceneCachedEntityFilterResult.h"
+#include "Containers/SortedMap.h"
+
+namespace UE
+{
+namespace MovieScene
+{
+
+/** Blend result struct that stores the cumulative sum of pre-weighted values, alongside the total weight */
+template<typename ValueType>
+struct TBlendResult
+{
+	/** Cumulative sum of blend values pre-multiplied with each value's weight. */
+	ValueType Total = 0.f;
+	/** Cumulative sum of weights. */
+	float Weight = 0.f;
+};
+
+/** Structure for holding the blend results of each blend type */
+template<typename ValueType>
+struct TAccumulationResult
+{
+	using FBlendResult = TBlendResult<ValueType>;
+
+	const FBlendResult* Absolutes = nullptr;
+	const FBlendResult* Relatives = nullptr;
+	const FBlendResult* Additives = nullptr;
+	const FBlendResult* AdditivesFromBase = nullptr;
+
+	bool IsValid() const
+	{
+		return Absolutes || Relatives || Additives || AdditivesFromBase;
+	}
+
+	FBlendResult GetAbsoluteResult(uint16 BlendID) const
+	{
+		return Absolutes ? Absolutes[BlendID] : FBlendResult{};
+	}
+	FBlendResult GetRelativeResult(uint16 BlendID) const
+	{
+		return Relatives ? Relatives[BlendID] : FBlendResult{};
+	}
+	FBlendResult GetAdditiveResult(uint16 BlendID) const
+	{
+		return Additives ? Additives[BlendID] : FBlendResult{};
+	}
+	FBlendResult GetAdditiveFromBaseResult(uint16 BlendID) const
+	{
+		return AdditivesFromBase ? AdditivesFromBase[BlendID] : FBlendResult{};
+	}
+};
+
+/** Buffer used for accumulating additive-from-base values */
+template<typename ValueType>
+struct TAdditiveFromBaseBuffer
+{
+	TArray<TBlendResult<ValueType>> Buffer;
+	TComponentTypeID<ValueType> BaseComponent;
+};
+
+/** Struct that maintains accumulation buffers for each blend type, one buffer per float result component type */
+template<typename ValueType>
+struct TAccumulationBuffers
+{
+	using FBlendResult = TBlendResult<ValueType>;
+	using FAccumulationResult = TAccumulationResult<ValueType>;
+	using FAdditiveFromBaseBuffer = TAdditiveFromBaseBuffer<ValueType>;
+
+	bool IsEmpty() const;
+
+	void Reset();
+
+	FAccumulationResult FindResults(FComponentTypeID InComponentType) const;
+
+	/** Map from value result component type -> Absolute blend accumulation buffer for that channel type */
+	TSortedMap<FComponentTypeID, TArray<FBlendResult>> Absolute;
+	/** Map from value result component type -> Relative blend accumulation buffer for that channel type */
+	TSortedMap<FComponentTypeID, TArray<FBlendResult>> Relative;
+	/** Map from value result component type -> Additive blend accumulation buffer for that channel type */
+	TSortedMap<FComponentTypeID, TArray<FBlendResult>> Additive;
+	/** Map from value result component type -> Additive From Base blend accumulation buffer for that channel type */
+	TSortedMap<FComponentTypeID, FAdditiveFromBaseBuffer> AdditiveFromBase;
+};
+
+/**
+ * Parameters for running the piecewise blender.
+ */
+struct FPiecewiseBlenderSystemImplRunParams
+{
+	int32 MaximumNumBlends = 0;
+
+	TStatId BlendValuesStatId;
+	TStatId CombineBlendsStatId;
+};
+
+/**
+ * Utility class for implementing piecewise blending given a floating precision type (float or double).
+ */
+template<typename ValueType>
+struct TPiecewiseBlenderSystemImpl
+{
+public:
+
+	/** Runs the blender system to blend all evaluated value channels of the given type */
+	void Run(FPiecewiseBlenderSystemImplRunParams& Params, FEntityManager& EntityManager, FSystemTaskPrerequisites& InPrerequisites, FSystemSubsequentTasks& Subsequents);
+
+	/** Decomposes the given value */
+	FGraphEventRef DispatchDecomposeTask(FEntityManager& EntityManager, const UE::MovieScene::FValueDecompositionParams& Params, UE::MovieScene::FAlignedDecomposedValue* Output);
+
+private:
+
+	using FBlendResult = TBlendResult<ValueType>;
+	using FAccumulationBuffers = TAccumulationBuffers<ValueType>;
+	using FAdditiveFromBaseBuffer = TAdditiveFromBaseBuffer<ValueType>;
+
+	void ReinitializeAccumulationBuffers(int32 MaximumNumBlends, FEntityManager& EntityManager);
+	void ZeroAccumulationBuffers();
+
+	/** Buffers that contain accumulated blend values, separated by blend type */
+	FAccumulationBuffers AccumulationBuffers;
+
+	/** Mask that contains value result components that have BlendChannelInput components */
+	FComponentMask BlendedResultMask;
+
+	/** Mask that contains property tags for any property type that has has at least one BlendChannelOutput */
+	FComponentMask BlendedPropertyMask;
+
+	/** Cache state that is used to invalidate and reset the accumulation buffers when the entity manager has structurally changed */
+	FCachedEntityManagerState ChannelRelevancyCache;
+
+	/** Bit array specifying FCompositePropertyTypeID's for properties contained within BlendedPropertyMask */
+	TBitArray<> CachedRelevantProperties;
+};
+
+} // namespace MovieScene
+} // namespace UE
diff --git a/Engine/Source/Runtime/MovieSceneTracks/Public/Systems/MovieScenePiecewiseDoubleBlenderSystem.h b/Engine/Source/Runtime/MovieSceneTracks/Public/Systems/MovieScenePiecewiseDoubleBlenderSystem.h
new file mode 100644
--- /dev/null
+++ b/Engine/Source/Runtime/MovieSceneTracks/Public/Systems/MovieScenePiecewiseDoubleBlenderSystem.h
@@ -0,0 +1,31 @@
+// Copyright Epic Games, Inc. All Rights Reserved.
+
+#pragma once
+
+#include "EntitySystem/MovieSceneBlenderSystem.h"
+#include "EntitySystem/MovieSceneDecompositionQuery.h"
+#include "Systems/MovieScenePiecewiseBlenderSystemImpl.h"
+
+#include "MovieScenePiecewiseDoubleBlenderSystem.generated.h"
+
+
+UCLASS()
+class MOVIESCENETRACKS_API UMovieScenePiecewiseDoubleBlenderSystem : public UMovieSceneBlenderSystem, public IMovieSceneValueDecomposer
+{
+public:
+	GENERATED_BODY()
+
+	UMovieScenePiecewiseDoubleBlenderSystem(const FObjectInitializer& ObjInit);
+
+	using FMovieSceneEntityID  = UE::MovieScene::FMovieSceneEntityID;
+	using FComponentTypeID     = UE::MovieScene::FComponentTypeID;
+
+	virtual void OnRun(FSystemTaskPrerequisites& InPrerequisites, FSystemSubsequentTasks& Subsequents) override;
+
+	virtual FGraphEventRef DispatchDecomposeTask(const UE::MovieScene::FValueDecompositionParams& Params, UE::MovieScene::FAlignedDecomposedValue* Output) override;
+
+private:
+
+	UE::MovieScene::TPiecewiseBlenderSystemImpl<double> Impl;
+};
+
diff --git a/Engine/Source/Runtime/MovieSceneTracks/Public/Systems/MovieScenePiecewiseFloatBlenderSystem.h b/Engine/Source/Runtime/MovieSceneTracks/Public/Systems/MovieScenePiecewiseFloatBlenderSystem.h
--- a/Engine/Source/Runtime/MovieSceneTracks/Public/Systems/MovieScenePiecewiseFloatBlenderSystem.h
+++ b/Engine/Source/Runtime/MovieSceneTracks/Public/Systems/MovieScenePiecewiseFloatBlenderSystem.h
@@ -2,62 +2,15 @@
 
 #pragma once
 
-#include "EntitySystem/MovieSceneEntitySystem.h"
 #include "EntitySystem/MovieSceneBlenderSystem.h"
 #include "EntitySystem/MovieSceneDecompositionQuery.h"
-#include "EntitySystem/MovieSceneCachedEntityFilterResult.h"
-#include "Containers/SortedMap.h"
+#include "Systems/MovieScenePiecewiseBlenderSystemImpl.h"
 
 #include "MovieScenePiecewiseFloatBlenderSystem.generated.h"
 
-namespace UE
-{
-namespace MovieScene
-{
-	struct FAccumulationResult;
-
-	/** Blend result struct that stores the cumulative sum of pre-weighted values, alongside the total weight */
-	struct FBlendResult
-	{
-		/** Cumulative sum of blend values pre-multiplied with each value's weight. */
-		float Total = 0.f;
-		/** Cumulative sum of weights. */
-		float Weight = 0.f;
-	};
-
-	/** Buffer used for accumulating additive-from-base values */
-	struct FAdditiveFromBaseBuffer
-	{
-		TArray<FBlendResult> Buffer;
-		TComponentTypeID<float> BaseComponent;
-	};
-
-	/** Struct that maintains accumulation buffers for each blend type, one buffer per float result component type */
-	struct FAccumulationBuffers
-	{
-		bool IsEmpty() const;
-
-		void Reset();
-
-		FAccumulationResult FindResults(FComponentTypeID InComponentType) const;
-
-		/** Map from Float Result component type -> Absolute blend accumulation buffer for that channel type */
-		TSortedMap<FComponentTypeID, TArray<FBlendResult>> Absolute;
-		/** Map from Float Result component type -> Relative blend accumulation buffer for that channel type */
-		TSortedMap<FComponentTypeID, TArray<FBlendResult>> Relative;
-		/** Map from Float Result component type -> Additive blend accumulation buffer for that channel type */
-		TSortedMap<FComponentTypeID, TArray<FBlendResult>> Additive;
-		/** Map from Float Result component type -> Additive From Base blend accumulation buffer for that channel type */
-		TSortedMap<FComponentTypeID, FAdditiveFromBaseBuffer> AdditiveFromBase;
-	};
-
-} // namespace MovieScene
-} // namespace UE
-
-
 
 UCLASS()
-class MOVIESCENETRACKS_API UMovieScenePiecewiseFloatBlenderSystem : public UMovieSceneBlenderSystem, public IMovieSceneFloatDecomposer
+class MOVIESCENETRACKS_API UMovieScenePiecewiseFloatBlenderSystem : public UMovieSceneBlenderSystem, public IMovieSceneValueDecomposer
 {
 public:
 	GENERATED_BODY()
@@ -67,29 +20,12 @@ public:
 	using FMovieSceneEntityID  = UE::MovieScene::FMovieSceneEntityID;
 	using FComponentTypeID     = UE::MovieScene::FComponentTypeID;
 
-	virtual void OnLink() override;
 	virtual void OnRun(FSystemTaskPrerequisites& InPrerequisites, FSystemSubsequentTasks& Subsequents) override;
 
-	virtual FGraphEventRef DispatchDecomposeTask(const UE::MovieScene::FFloatDecompositionParams& Params, UE::MovieScene::FAlignedDecomposedFloat* Output) override;
+	virtual FGraphEventRef DispatchDecomposeTask(const UE::MovieScene::FValueDecompositionParams& Params, UE::MovieScene::FAlignedDecomposedValue* Output) override;
 
 private:
 
-	void ReinitializeAccumulationBuffers();
-	void ZeroAccumulationBuffers();
-
-	/** Buffers that contain accumulated blend values, separated by blend type */
-	UE::MovieScene::FAccumulationBuffers AccumulationBuffers;
-
-	/** Mask that contains FloatResult components that have BlendChannelInput components */
-	UE::MovieScene::FComponentMask BlendedResultMask;
-
-	/** Mask that contains property tags for any property type that has has at least one BlendChannelOutput */
-	UE::MovieScene::FComponentMask BlendedPropertyMask;
-
-	/** Cache state that is used to invalidate and reset the accumulation buffers when the entity manager has structurally changed */
-	UE::MovieScene::FCachedEntityManagerState ChannelRelevancyCache;
-
-	/** Bit array specifying FCompositePropertyTypeID's for properties contained within BlendedPropertyMask */
-	TBitArray<> CachedRelevantProperties;
+	UE::MovieScene::TPiecewiseBlenderSystemImpl<float> Impl;
 };
 
diff --git a/Engine/Source/Runtime/MovieSceneTracks/Public/Systems/MovieSceneVectorPropertySystem.h b/Engine/Source/Runtime/MovieSceneTracks/Public/Systems/MovieSceneVectorPropertySystem.h
--- a/Engine/Source/Runtime/MovieSceneTracks/Public/Systems/MovieSceneVectorPropertySystem.h
+++ b/Engine/Source/Runtime/MovieSceneTracks/Public/Systems/MovieSceneVectorPropertySystem.h
@@ -7,12 +7,24 @@
 #include "MovieSceneVectorPropertySystem.generated.h"
 
 UCLASS()
-class MOVIESCENETRACKS_API UMovieSceneVectorPropertySystem : public UMovieScenePropertySystem
+class MOVIESCENETRACKS_API UMovieSceneFloatVectorPropertySystem : public UMovieScenePropertySystem
 {
 	GENERATED_BODY()
 
 private:
-	UMovieSceneVectorPropertySystem(const FObjectInitializer& ObjInit);
+	UMovieSceneFloatVectorPropertySystem(const FObjectInitializer& ObjInit);
 
 	virtual void OnRun(FSystemTaskPrerequisites& InPrerequisites, FSystemSubsequentTasks& Subsequents) override;
 };
+
+UCLASS()
+class MOVIESCENETRACKS_API UMovieSceneDoubleVectorPropertySystem : public UMovieScenePropertySystem
+{
+	GENERATED_BODY()
+
+private:
+	UMovieSceneDoubleVectorPropertySystem(const FObjectInitializer& ObjInit);
+
+	virtual void OnRun(FSystemTaskPrerequisites& InPrerequisites, FSystemSubsequentTasks& Subsequents) override;
+};
+
diff --git a/Engine/Source/Runtime/MovieSceneTracks/Public/Tracks/MovieSceneVectorTrack.h b/Engine/Source/Runtime/MovieSceneTracks/Public/Tracks/MovieSceneVectorTrack.h
--- a/Engine/Source/Runtime/MovieSceneTracks/Public/Tracks/MovieSceneVectorTrack.h
+++ b/Engine/Source/Runtime/MovieSceneTracks/Public/Tracks/MovieSceneVectorTrack.h
@@ -9,10 +9,35 @@
 #include "MovieSceneVectorTrack.generated.h"
 
 /**
- * Handles manipulation of component transforms in a movie scene
+ * Handles manipulation of float vectors in a movie scene
  */
 UCLASS(MinimalAPI)
-class UMovieSceneVectorTrack : public UMovieScenePropertyTrack
+class UMovieSceneFloatVectorTrack : public UMovieScenePropertyTrack
+{
+	GENERATED_UCLASS_BODY()
+
+public:
+	/** UMovieSceneTrack interface */
+	virtual bool SupportsType(TSubclassOf<UMovieSceneSection> SectionClass) const override;
+	virtual UMovieSceneSection* CreateNewSection() override;
+
+	/** @return Get the number of channels used by the vector */
+	int32 GetNumChannelsUsed() const { return NumChannelsUsed; }
+
+	/** Set the number of channels used by the vector */
+	void SetNumChannelsUsed( int32 InNumChannelsUsed ) { NumChannelsUsed = InNumChannelsUsed; }
+
+private:
+	/** The number of channels used by the vector (2,3, or 4) */
+	UPROPERTY()
+	int32 NumChannelsUsed;
+};
+
+/**
+ * Handles manipulation of double vectors in a movie scene
+ */
+UCLASS(MinimalAPI)
+class UMovieSceneDoubleVectorTrack : public UMovieScenePropertyTrack
 {
 	GENERATED_UCLASS_BODY()
 
diff --git a/Engine/Source/Runtime/UMG/Private/Animation/MovieSceneUMGComponentTypes.cpp b/Engine/Source/Runtime/UMG/Private/Animation/MovieSceneUMGComponentTypes.cpp
--- a/Engine/Source/Runtime/UMG/Private/Animation/MovieSceneUMGComponentTypes.cpp
+++ b/Engine/Source/Runtime/UMG/Private/Animation/MovieSceneUMGComponentTypes.cpp
@@ -1,11 +1,12 @@
 // Copyright Epic Games, Inc. All Rights Reserved.
 
 #include "Animation/MovieSceneUMGComponentTypes.h"
+#include "EntitySystem/BuiltInComponentTypes.h"
 #include "EntitySystem/MovieSceneComponentRegistry.h"
+#include "EntitySystem/MovieSceneEntityFactoryTemplates.h"
 #include "EntitySystem/MovieScenePropertyComponentHandler.h"
-#include "EntitySystem/BuiltInComponentTypes.h"
 #include "MovieSceneTracksComponentTypes.h"
-#include "EntitySystem/MovieSceneEntityFactoryTemplates.h"
+#include "Systems/MovieScenePiecewiseFloatBlenderSystem.h"
 
 #include "Components/Widget.h"
 
@@ -95,6 +96,7 @@ FMovieSceneUMGComponentTypes::FMovieSceneUMGComponentTypes()
 	.AddComposite(BuiltInComponents->FloatResult[4], &FIntermediateWidgetTransform::ScaleY)
 	.AddComposite(BuiltInComponents->FloatResult[5], &FIntermediateWidgetTransform::ShearX)
 	.AddComposite(BuiltInComponents->FloatResult[6], &FIntermediateWidgetTransform::ShearY)
+	.SetBlenderSystem<UMovieScenePiecewiseFloatBlenderSystem>()
 	.SetCustomAccessors(&CustomWidgetTransformAccessors)
 	.Commit();
 }
