Sequencer: Stop audio on scrubbing and stepping. This just moves the stop outside of the check whether there is a need for a state change (defined as (If the audio component is not playing we (may) need a state change. If the audio component is playing the wrong sound then we need a state change. If the audio playback time is significantly out of sync with the desired time then we need a state change.)

diff --git a/Engine/Source/Runtime/MovieSceneTracks/Private/Evaluation/MovieSceneAudioTemplate.cpp b/Engine/Source/Runtime/MovieSceneTracks/Private/Evaluation/MovieSceneAudioTemplate.cpp
--- a/Engine/Source/Runtime/MovieSceneTracks/Private/Evaluation/MovieSceneAudioTemplate.cpp
+++ b/Engine/Source/Runtime/MovieSceneTracks/Private/Evaluation/MovieSceneAudioTemplate.cpp
@@ -524,12 +524,12 @@ struct FAudioSectionExecutionToken : IMovieSceneExecutionToken
 					
 				}
 			}
+		}
 
-			if (Context.GetStatus() == EMovieScenePlayerStatus::Scrubbing || Context.GetStatus() == EMovieScenePlayerStatus::Stepping)
-			{
-				// While scrubbing, play the sound for a short time and then cut it.
-				AudioComponent.StopDelayed(AudioTrackConstants::ScrubDuration);
-			}
+		if (Context.GetStatus() == EMovieScenePlayerStatus::Scrubbing || Context.GetStatus() == EMovieScenePlayerStatus::Stepping)
+		{
+			// While scrubbing, play the sound for a short time and then cut it.
+			AudioComponent.StopDelayed(AudioTrackConstants::ScrubDuration);
 		}
 
 		if (bAllowSpatialization)
