Sequencer: Check for null child on motion vector simulation

diff --git a/Engine/Source/Runtime/MovieSceneTracks/Private/Evaluation/MovieSceneSkeletalAnimationTemplate.cpp b/Engine/Source/Runtime/MovieSceneTracks/Private/Evaluation/MovieSceneSkeletalAnimationTemplate.cpp
--- a/Engine/Source/Runtime/MovieSceneTracks/Private/Evaluation/MovieSceneSkeletalAnimationTemplate.cpp
+++ b/Engine/Source/Runtime/MovieSceneTracks/Private/Evaluation/MovieSceneSkeletalAnimationTemplate.cpp
@@ -398,6 +398,11 @@ namespace MovieScene
 		{
 			for (USceneComponent* Child : SkeletalMeshComponent->GetAttachChildren())
 			{
+				if (!Child)
+				{
+					continue;
+				}
+
 				FName SocketName = Child->GetAttachSocketName();
 				if (SocketName != NAME_None)
 				{
