Sequencer: Compute base values (for "additive from base" blending) for double-based curves.

Also, the blender systems should depend on the base value computing system, so that the base values are correctly computed by the time they're used for blending.

diff --git a/Engine/Source/Runtime/MovieSceneTracks/Private/Systems/MovieSceneBaseValueEvaluatorSystem.cpp b/Engine/Source/Runtime/MovieSceneTracks/Private/Systems/MovieSceneBaseValueEvaluatorSystem.cpp
--- a/Engine/Source/Runtime/MovieSceneTracks/Private/Systems/MovieSceneBaseValueEvaluatorSystem.cpp
+++ b/Engine/Source/Runtime/MovieSceneTracks/Private/Systems/MovieSceneBaseValueEvaluatorSystem.cpp
@@ -14,7 +14,7 @@ namespace UE
 namespace MovieScene
 {
 
-struct FEvaluateBaseValues
+struct FEvaluateBaseFloatValues
 {
 	void ForEachEntity(FSourceFloatChannel FloatChannel, FFrameTime FrameTime, float& OutResult)
 	{
@@ -25,6 +25,17 @@ struct FEvaluateBaseValues
 	}
 };
 
+struct FEvaluateBaseDoubleValues
+{
+	void ForEachEntity(FSourceDoubleChannel DoubleChannel, FFrameTime FrameTime, double& OutResult)
+	{
+		if (!DoubleChannel.Source->Evaluate(FrameTime, OutResult))
+		{
+			OutResult = MIN_dbl;
+		}
+	}
+};
+
 } // namespace MovieScene
 } // namespace UE
 
@@ -62,6 +73,23 @@ void UMovieSceneBaseValueEvaluatorSystem::OnRun(FSystemTaskPrerequisites& InPrer
 		.Write(BaseFloat)
 		.FilterAll({ BuiltInComponents->Tags.NeedsLink })
 		.FilterNone({ BuiltInComponents->Tags.Ignored })
-		.RunInline_PerEntity(&Linker->EntityManager, FEvaluateBaseValues());
+		.RunInline_PerEntity(&Linker->EntityManager, FEvaluateBaseFloatValues());
+	}
+
+	static_assert(
+			UE_ARRAY_COUNT(BuiltInComponents->BaseDouble) == UE_ARRAY_COUNT(BuiltInComponents->DoubleChannel),
+			"There should be a matching number of double channels and double base values.");
+	for (size_t Index = 0; Index < UE_ARRAY_COUNT(BuiltInComponents->BaseDouble); ++Index)
+	{
+		const TComponentTypeID<double> BaseDouble = BuiltInComponents->BaseDouble[Index];
+		const TComponentTypeID<FSourceDoubleChannel> DoubleChannel = BuiltInComponents->DoubleChannel[Index];
+
+		FEntityTaskBuilder()
+		.Read(DoubleChannel)
+		.Read(BuiltInComponents->BaseValueEvalTime)
+		.Write(BaseDouble)
+		.FilterAll({ BuiltInComponents->Tags.NeedsLink })
+		.FilterNone({ BuiltInComponents->Tags.Ignored })
+		.RunInline_PerEntity(&Linker->EntityManager, FEvaluateBaseDoubleValues());
 	}
 }
diff --git a/Engine/Source/Runtime/MovieSceneTracks/Private/Systems/MovieScenePiecewiseDoubleBlenderSystem.cpp b/Engine/Source/Runtime/MovieSceneTracks/Private/Systems/MovieScenePiecewiseDoubleBlenderSystem.cpp
--- a/Engine/Source/Runtime/MovieSceneTracks/Private/Systems/MovieScenePiecewiseDoubleBlenderSystem.cpp
+++ b/Engine/Source/Runtime/MovieSceneTracks/Private/Systems/MovieScenePiecewiseDoubleBlenderSystem.cpp
@@ -4,6 +4,7 @@
 #include "EntitySystem/BuiltInComponentTypes.h"
 #include "EntitySystem/MovieSceneEntitySystemLinker.h"
 #include "Systems/DoubleChannelEvaluatorSystem.h"
+#include "Systems/MovieSceneBaseValueEvaluatorSystem.h"
 #include "Systems/MovieSceneQuaternionInterpolationRotationSystem.h"
 #include "Systems/WeightAndEasingEvaluatorSystem.h"
 
@@ -19,6 +20,7 @@ UMovieScenePiecewiseDoubleBlenderSystem::UMovieScenePiecewiseDoubleBlenderSystem
 	if (HasAnyFlags(RF_ClassDefaultObject))
 	{
 		DefineImplicitPrerequisite(UDoubleChannelEvaluatorSystem::StaticClass(), GetClass());
+		DefineImplicitPrerequisite(UMovieSceneBaseValueEvaluatorSystem::StaticClass(), GetClass());
 		DefineImplicitPrerequisite(UMovieSceneQuaternionInterpolationRotationSystem::StaticClass(), GetClass());
 		DefineImplicitPrerequisite(UWeightAndEasingEvaluatorSystem::StaticClass(), GetClass());
 	}
diff --git a/Engine/Source/Runtime/MovieSceneTracks/Private/Systems/MovieScenePiecewiseFloatBlenderSystem.cpp b/Engine/Source/Runtime/MovieSceneTracks/Private/Systems/MovieScenePiecewiseFloatBlenderSystem.cpp
--- a/Engine/Source/Runtime/MovieSceneTracks/Private/Systems/MovieScenePiecewiseFloatBlenderSystem.cpp
+++ b/Engine/Source/Runtime/MovieSceneTracks/Private/Systems/MovieScenePiecewiseFloatBlenderSystem.cpp
@@ -5,6 +5,7 @@
 #include "EntitySystem/MovieSceneEntitySystemLinker.h"
 #include "EntitySystem/MovieSceneEntitySystemTask.h"
 #include "Systems/FloatChannelEvaluatorSystem.h"
+#include "Systems/MovieSceneBaseValueEvaluatorSystem.h"
 #include "Systems/MovieScenePiecewiseBlenderSystemImpl.h"
 #include "Systems/MovieSceneQuaternionInterpolationRotationSystem.h"
 #include "Systems/WeightAndEasingEvaluatorSystem.h"
@@ -21,6 +22,7 @@ UMovieScenePiecewiseFloatBlenderSystem::UMovieScenePiecewiseFloatBlenderSystem(c
 	if (HasAnyFlags(RF_ClassDefaultObject))
 	{
 		DefineImplicitPrerequisite(UFloatChannelEvaluatorSystem::StaticClass(), GetClass());
+		DefineImplicitPrerequisite(UMovieSceneBaseValueEvaluatorSystem::StaticClass(), GetClass());
 		DefineImplicitPrerequisite(UMovieSceneQuaternionInterpolationRotationSystem::StaticClass(), GetClass());
 		DefineImplicitPrerequisite(UWeightAndEasingEvaluatorSystem::StaticClass(), GetClass());
 	}
