Fix CIS non-unity build.

diff --git a/Engine/Source/Runtime/MovieSceneTracks/Private/Systems/MovieSceneBaseValueEvaluatorSystem.cpp b/Engine/Source/Runtime/MovieSceneTracks/Private/Systems/MovieSceneBaseValueEvaluatorSystem.cpp
--- a/Engine/Source/Runtime/MovieSceneTracks/Private/Systems/MovieSceneBaseValueEvaluatorSystem.cpp
+++ b/Engine/Source/Runtime/MovieSceneTracks/Private/Systems/MovieSceneBaseValueEvaluatorSystem.cpp
@@ -7,6 +7,7 @@
 #include "EntitySystem/MovieSceneEntitySystemTask.h"
 #include "EntitySystem/MovieSceneEntitySystemLinker.h"
 #include "EntitySystem/MovieSceneComponentRegistry.h"
+#include "Channels/MovieSceneDoubleChannel.h"
 #include "Channels/MovieSceneFloatChannel.h"
 
 namespace UE
