Sequencer: A trimmed, reversed animation clip in sequencer is playing correctly

diff --git a/Engine/Source/Runtime/MovieSceneTracks/Private/Evaluation/MovieSceneSkeletalAnimationTemplate.cpp b/Engine/Source/Runtime/MovieSceneTracks/Private/Evaluation/MovieSceneSkeletalAnimationTemplate.cpp
--- a/Engine/Source/Runtime/MovieSceneTracks/Private/Evaluation/MovieSceneSkeletalAnimationTemplate.cpp
+++ b/Engine/Source/Runtime/MovieSceneTracks/Private/Evaluation/MovieSceneSkeletalAnimationTemplate.cpp
@@ -635,30 +635,39 @@ void FMovieSceneSkeletalAnimationSectionTemplate::Evaluate(const FMovieSceneEval
 
 float FMovieSceneSkeletalAnimationSectionTemplateParameters::MapTimeToAnimation(FFrameTime InPosition, FFrameRate InFrameRate) const
 {
+	// Get Animation Length and frame time
 	const FFrameTime AnimationLength = GetSequenceLength() * InFrameRate;
 	const int32 LengthInFrames = AnimationLength.FrameNumber.Value + (int)(AnimationLength.GetSubFrame() + 0.5f) + 1;
-	//we only play end if we are not looping, and assuming we are looping if Length is greater than default length;
-	const bool bLooping = (SectionEndTime.Value - SectionStartTime.Value + StartFrameOffset + EndFrameOffset) > LengthInFrames;
 
-	InPosition = FMath::Clamp(InPosition, FFrameTime(SectionStartTime), FFrameTime(SectionEndTime -1));
+	// we only play end if we are not looping, and assuming we are looping if Length is greater than default length;
+	const bool bLooping = (SectionEndTime.Value - SectionStartTime.Value + StartFrameOffset + EndFrameOffset) > LengthInFrames;	
+
+	// Make sure InPosition FrameTime doesn't underflow SectionStartTime or overflow SectionEndTime
+	InPosition = FMath::Clamp(InPosition, FFrameTime(SectionStartTime), FFrameTime(SectionEndTime - 1));
 	
+	// Gather helper values
 	const float SectionPlayRate = PlayRate * Animation->RateScale;
 	const float AnimPlayRate = FMath::IsNearlyZero(SectionPlayRate) ? 1.0f : SectionPlayRate;
-
 	const float FirstLoopSeqLength = GetSequenceLength() - InFrameRate.AsSeconds(FirstLoopStartFrameOffset + StartFrameOffset + EndFrameOffset);
 	const float SeqLength = GetSequenceLength() - InFrameRate.AsSeconds(StartFrameOffset + EndFrameOffset);
 
-	float AnimPosition = FFrameTime::FromDecimal((InPosition - SectionStartTime).AsDecimal() * AnimPlayRate) / InFrameRate;
-	AnimPosition += InFrameRate.AsSeconds(FirstLoopStartFrameOffset);
-	if (SeqLength > 0.f && (bLooping || !FMath::IsNearlyEqual(AnimPosition, SeqLength,1e-4f)))
+	// The Time from the beginning of SectionStartTime to InPosition in seconds
+	float SecondsFromSectionStart = FFrameTime::FromDecimal((InPosition - SectionStartTime).AsDecimal() * AnimPlayRate) / InFrameRate;
+
+	// Logic for reversed animation
+	if (bReverse)
 	{
-		AnimPosition = FMath::Fmod(AnimPosition, SeqLength);
+		// Duration of this section 
+		float SectionDuration = (((SectionEndTime - SectionStartTime) * AnimPlayRate) / InFrameRate);
+		SecondsFromSectionStart = SectionDuration - SecondsFromSectionStart;
 	}
-	AnimPosition += InFrameRate.AsSeconds(StartFrameOffset);
-	if (bReverse)
+	// Make sure Seconds is in range
+	if (SeqLength > 0.f && (bLooping || !FMath::IsNearlyEqual(SecondsFromSectionStart, SeqLength,1e-4f)))
 	{
-		AnimPosition = GetSequenceLength() - AnimPosition;
+		SecondsFromSectionStart = FMath::Fmod(SecondsFromSectionStart, SeqLength);
 	}
 
-	return AnimPosition;
+	// Add the StartFrameOffset and FirstLoopStartFrameOffset to the current seconds in the section to get the right animation frame
+	SecondsFromSectionStart += InFrameRate.AsSeconds(StartFrameOffset) + InFrameRate.AsSeconds(FirstLoopStartFrameOffset);
+	return SecondsFromSectionStart;
 }
diff --git a/Engine/Source/Runtime/MovieSceneTracks/Private/Evaluation/MovieSceneSkeletalAnimationTemplate.h b/Engine/Source/Runtime/MovieSceneTracks/Private/Evaluation/MovieSceneSkeletalAnimationTemplate.h
--- a/Engine/Source/Runtime/MovieSceneTracks/Private/Evaluation/MovieSceneSkeletalAnimationTemplate.h
+++ b/Engine/Source/Runtime/MovieSceneTracks/Private/Evaluation/MovieSceneSkeletalAnimationTemplate.h
@@ -20,11 +20,19 @@ struct FMovieSceneSkeletalAnimationSectionTemplateParameters : public FMovieScen
 		, SectionEndTime(InSectionEndTime)
 	{}
 
+	/**
+	Calculates the animation's appropriate position in seconds from the begin by current frame time and frame rate
+	@param InPosition	Current frame time in sequencer
+	@param InFrameRate	Current frame time of this animation
+	@return A float indicates the seconds of the animation we should play at this frame time
+	*/
 	float MapTimeToAnimation(FFrameTime InPosition, FFrameRate InFrameRate) const;
 
+	/**The start time of this animation clip in this track in frames*/
 	UPROPERTY()
 	FFrameNumber SectionStartTime;
 
+	/**The end time of this animation clip in this track in frames*/
 	UPROPERTY()
 	FFrameNumber SectionEndTime;
 };
