Sequencer: Turn blend channel IDs into structs that encode the blender system they relate to

This solves the problem of multiple channel IDs of the same value being mixed together when they really are for different blender systems. The hierarchical bias system for example was mixing them, which means it could disable any random thing that had the same channel ID as something else of a different type that was legitimately disabled.

diff --git a/Engine/Source/Runtime/MovieScene/Private/EntitySystem/MovieSceneBlenderSystem.cpp b/Engine/Source/Runtime/MovieScene/Private/EntitySystem/MovieSceneBlenderSystem.cpp
--- a/Engine/Source/Runtime/MovieScene/Private/EntitySystem/MovieSceneBlenderSystem.cpp
+++ b/Engine/Source/Runtime/MovieScene/Private/EntitySystem/MovieSceneBlenderSystem.cpp
@@ -4,9 +4,48 @@
 #include "EntitySystem/MovieSceneEntitySystemLinker.h"
 #include "EntitySystem/BuiltInComponentTypes.h"
 
+namespace UE
+{
+namespace MovieScene
+{
+
+static TMap<FMovieSceneBlenderSystemID, TSubclassOf<UMovieSceneBlenderSystem>> GBlenderSystemRegistry;
+
+} // namespace MovieScene
+} // namespace UE
+
+TSubclassOf<UMovieSceneBlenderSystem> UMovieSceneBlenderSystem::GetBlenderSystemClass(FMovieSceneBlenderSystemID InSystemID)
+{
+	using namespace UE::MovieScene;
+
+	TSubclassOf<UMovieSceneBlenderSystem>* BlenderClass = GBlenderSystemRegistry.Find(InSystemID);
+	return BlenderClass ? *BlenderClass : TSubclassOf<UMovieSceneBlenderSystem>();
+}
+
+UMovieSceneBlenderSystem::UMovieSceneBlenderSystem(const FObjectInitializer& ObjInit)
+	: Super(ObjInit)
+{
+	using namespace UE::MovieScene;
 
+	if (HasAnyFlags(RF_ClassDefaultObject))
+	{
+		// You can only ever register blender systems, never unregister them.
+		SystemID = FMovieSceneBlenderSystemID(GBlenderSystemRegistry.Num());
+		GBlenderSystemRegistry.Add(SystemID, GetClass());
+	}
+	else 
+	{
+		SystemID = GetClass()->GetDefaultObject<UMovieSceneBlenderSystem>()->SystemID;
+		checkf(SystemID.IsValid(), TEXT("Blender system wasn't registered correctly on init!"));
+	}
+}
+
+FMovieSceneBlenderSystemID UMovieSceneBlenderSystem::GetBlenderSystemID() const
+{
+	return SystemID;
+}
 
-uint16 UMovieSceneBlenderSystem::AllocateBlendChannel()
+FMovieSceneBlendChannelID UMovieSceneBlenderSystem::AllocateBlendChannel()
 {
 	int32 NewBlendChannel = AllocatedBlendChannels.FindAndSetFirstZeroBit();
 	if (NewBlendChannel == INDEX_NONE)
@@ -15,13 +54,14 @@ uint16 UMovieSceneBlenderSystem::AllocateBlendChannel()
 	}
 
 	checkf(NewBlendChannel < TNumericLimits<uint16>::Max(), TEXT("Maximum number of active blends reached - this indicates either a leak, or more than 65535 blend channels are genuinely required"));
-	return static_cast<uint16>(NewBlendChannel);
+	return FMovieSceneBlendChannelID { SystemID, static_cast<uint16>(NewBlendChannel) };
 }
 
 
-void UMovieSceneBlenderSystem::ReleaseBlendChannel(uint16 BlendID)
+void UMovieSceneBlenderSystem::ReleaseBlendChannel(FMovieSceneBlendChannelID BlendID)
 {
-	AllocatedBlendChannels[BlendID] = false;
+	ensureMsgf(BlendID.SystemID == SystemID, TEXT("This given blend channel wasn't allocated by this blender system!"));
+	AllocatedBlendChannels[BlendID.ChannelID] = false;
 }
 
 
diff --git a/Engine/Source/Runtime/MovieScene/Private/EntitySystem/MovieSceneBlenderSystemTypes.cpp b/Engine/Source/Runtime/MovieScene/Private/EntitySystem/MovieSceneBlenderSystemTypes.cpp
new file mode 100644
--- /dev/null
+++ b/Engine/Source/Runtime/MovieScene/Private/EntitySystem/MovieSceneBlenderSystemTypes.cpp
@@ -0,0 +1,19 @@
+// Copyright Epic Games, Inc. All Rights Reserved.
+
+#include "EntitySystem/MovieSceneBlenderSystemTypes.h"
+#include "EntitySystem/MovieSceneBlenderSystem.h"
+#include "EntitySystem/MovieSceneEntitySystemLinker.h"
+
+FMovieSceneBlenderSystemID FMovieSceneBlenderSystemID::Invalid;
+
+TSubclassOf<UMovieSceneBlenderSystem> FMovieSceneBlendChannelID::GetSystemClass() const
+{
+	return UMovieSceneBlenderSystem::GetBlenderSystemClass(SystemID);
+}
+
+UMovieSceneBlenderSystem* FMovieSceneBlendChannelID::FindSystem(const UMovieSceneEntitySystemLinker* Linker) const
+{
+	TSubclassOf<UMovieSceneBlenderSystem> SystemClass = GetSystemClass();
+	return CastChecked<UMovieSceneBlenderSystem>(Linker->FindSystem(SystemClass));
+}
+
diff --git a/Engine/Source/Runtime/MovieScene/Public/EntitySystem/BuiltInComponentTypes.h b/Engine/Source/Runtime/MovieScene/Public/EntitySystem/BuiltInComponentTypes.h
--- a/Engine/Source/Runtime/MovieScene/Public/EntitySystem/BuiltInComponentTypes.h
+++ b/Engine/Source/Runtime/MovieScene/Public/EntitySystem/BuiltInComponentTypes.h
@@ -4,6 +4,7 @@
 
 #include "EntitySystem/MovieSceneEntityManager.h"
 #include "UObject/GCObjectScopeGuard.h"
+#include "EntitySystem/MovieSceneBlenderSystemTypes.h"
 #include "EntitySystem/MovieSceneEntityIDs.h"
 #include "EntitySystem/MovieSceneSequenceInstanceHandle.h"
 #include "Evaluation/Blending/MovieSceneBlendType.h"
@@ -71,7 +72,6 @@ struct FMovieSceneEvaluationHookComponent
 };
 
 
-
 USTRUCT()
 struct FTrackInstanceInputComponent
 {
@@ -174,11 +174,11 @@ public:
 
 public:
 
-	TComponentTypeID<uint16>              BlendChannelInput;
+	TComponentTypeID<FMovieSceneBlendChannelID> BlendChannelInput;
 
-	TComponentTypeID<int16>               HierarchicalBias;
+	TComponentTypeID<FMovieSceneBlendChannelID> BlendChannelOutput;
 
-	TComponentTypeID<uint16>              BlendChannelOutput;
+	TComponentTypeID<int16>               HierarchicalBias;
 
 	TComponentTypeID<FInitialValueIndex>  InitialValueIndex;
 public:
diff --git a/Engine/Source/Runtime/MovieScene/Public/EntitySystem/MovieSceneBlenderSystem.h b/Engine/Source/Runtime/MovieScene/Public/EntitySystem/MovieSceneBlenderSystem.h
--- a/Engine/Source/Runtime/MovieScene/Public/EntitySystem/MovieSceneBlenderSystem.h
+++ b/Engine/Source/Runtime/MovieScene/Public/EntitySystem/MovieSceneBlenderSystem.h
@@ -3,6 +3,8 @@
 #pragma once
 
 #include "EntitySystem/MovieSceneEntitySystem.h"
+#include "EntitySystem/MovieSceneBlenderSystemTypes.h"
+#include "Math/NumericLimits.h"
 #include "Templates/SubclassOf.h"
 
 #include "MovieSceneBlenderSystem.generated.h"
@@ -38,25 +40,52 @@ public:
 
 	GENERATED_BODY()
 
+	/**
+	 * Get a blender system class, given a blender system ID.
+	 */
+	static TSubclassOf<UMovieSceneBlenderSystem> GetBlenderSystemClass(FMovieSceneBlenderSystemID InSystemID);
+
+	/**
+	 * Get the unique ID for a given blender system type.
+	 */
+	template<typename BlenderSystemClass>
+	static FMovieSceneBlenderSystemID GetBlenderSystemID()
+	{
+		UMovieSceneBlenderSystem* DefaultObject = BlenderSystemClass::StaticClass()->GetDefaultObject<UMovieSceneBlenderSystem>();
+		check(DefaultObject);
+		return DefaultObject->SystemID;
+	}
+
+	/**
+	 * Get the unique ID for this blender system's underlying type.
+	 */
+	FMovieSceneBlenderSystemID GetBlenderSystemID() const;
+
 	/**
 	 * Allocate a new blend channel.
 	 * @note Must be released when it is no longer needed in order to prevent leaking channels.
 	 */
-	uint16 AllocateBlendChannel();
-
+	FMovieSceneBlendChannelID AllocateBlendChannel();
 
 	/**
 	 * Release a previously allocated blend channel.
 	 */
-	void ReleaseBlendChannel(uint16 BlendChannelID);
+	void ReleaseBlendChannel(FMovieSceneBlendChannelID BlendChannelID);
 
 
 protected:
 
+	UMovieSceneBlenderSystem(const FObjectInitializer& ObjInit);
+
 	virtual bool IsRelevantImpl(UMovieSceneEntitySystemLinker* InLinker) const override;
 
 	void CompactBlendChannels();
 
 	/** Bit array specifying currently allocated blend channels */
 	TBitArray<> AllocatedBlendChannels;
+
+private:
+
+	/** Cached blender system ID */
+	FMovieSceneBlenderSystemID SystemID;
 };
diff --git a/Engine/Source/Runtime/MovieScene/Public/EntitySystem/MovieSceneBlenderSystemTypes.h b/Engine/Source/Runtime/MovieScene/Public/EntitySystem/MovieSceneBlenderSystemTypes.h
new file mode 100644
--- /dev/null
+++ b/Engine/Source/Runtime/MovieScene/Public/EntitySystem/MovieSceneBlenderSystemTypes.h
@@ -0,0 +1,89 @@
+// Copyright Epic Games, Inc. All Rights Reserved.
+
+#pragma once
+
+#include "CoreMinimal.h"
+#include "Math/NumericLimits.h"
+#include "Templates/Casts.h"
+#include "Templates/SubclassOf.h"
+
+class UMovieSceneBlenderSystem;
+class UMovieSceneEntitySystemLinker;
+
+/**
+ * Identifier for a type of blender system (e.g. float blender, integer blender, etc.)
+ */
+struct FMovieSceneBlenderSystemID
+{
+	static FMovieSceneBlenderSystemID Invalid;
+
+	FMovieSceneBlenderSystemID() : Value(TNumericLimits<uint8>::Max()) {}
+	explicit FMovieSceneBlenderSystemID(uint8 InValue) : Value(InValue) {}
+
+	bool IsValid() const { return Value != Invalid.Value; }
+
+	friend uint32 GetTypeHash(FMovieSceneBlenderSystemID SystemID)
+	{
+		return GetTypeHash(SystemID.Value);
+	}
+
+	friend bool operator==(FMovieSceneBlenderSystemID A, FMovieSceneBlenderSystemID B)
+	{
+		return A.Value == B.Value;
+	}
+
+	friend bool operator<(FMovieSceneBlenderSystemID A, FMovieSceneBlenderSystemID B)
+	{
+		return A.Value < B.Value;
+	}
+
+private:
+	uint8 Value;
+};
+
+/**
+ * A blend channel ID
+ */
+struct FMovieSceneBlendChannelID
+{
+	/** Invalid blend channel ID */
+	static constexpr uint16 INVALID_BLEND_CHANNEL = uint16(-1);
+
+	/** Builds a default, invalid channel ID */
+	FMovieSceneBlendChannelID() : ChannelID(INVALID_BLEND_CHANNEL) {}
+	/** Builds a channel ID */
+	FMovieSceneBlendChannelID(FMovieSceneBlenderSystemID InSystemID, uint16 InChannelID) : SystemID(InSystemID), ChannelID(InChannelID) {}
+
+	/** Returns whether this ID is valid */
+	bool IsValid() const { return ChannelID != INVALID_BLEND_CHANNEL; }
+
+	/** Returns the blender system type */
+	MOVIESCENE_API TSubclassOf<UMovieSceneBlenderSystem> GetSystemClass() const;
+
+	/** Returns the blender system instance found in a given linker */
+	MOVIESCENE_API UMovieSceneBlenderSystem* FindSystem(const UMovieSceneEntitySystemLinker* Linker) const;
+
+	/** Returns the blender system instance found in a given linker */
+	template<typename BlenderSystemClass>
+	BlenderSystemClass* FindSystem(const UMovieSceneEntitySystemLinker* Linker) const
+	{
+		return CastChecked<BlenderSystemClass>(FindSystem(Linker));
+	}
+
+	friend bool operator==(const FMovieSceneBlendChannelID A, const FMovieSceneBlendChannelID B)
+	{
+		return (A.SystemID == B.SystemID && A.ChannelID == B.ChannelID);
+	}
+
+	friend uint32 GetTypeHash(const FMovieSceneBlendChannelID A)
+	{
+		return HashCombine(GetTypeHash(A.SystemID), GetTypeHash(A.ChannelID));
+	}
+
+	/** The blender system associated with the blend channel */
+	FMovieSceneBlenderSystemID SystemID;
+
+	/** The blend channel ID for the given blender system */
+	uint16 ChannelID;
+};
+
diff --git a/Engine/Source/Runtime/MovieSceneTracks/Private/EntitySystem/Interrogation/MovieSceneInterrogatedPropertyInstantiator.cpp b/Engine/Source/Runtime/MovieSceneTracks/Private/EntitySystem/Interrogation/MovieSceneInterrogatedPropertyInstantiator.cpp
--- a/Engine/Source/Runtime/MovieSceneTracks/Private/EntitySystem/Interrogation/MovieSceneInterrogatedPropertyInstantiator.cpp
+++ b/Engine/Source/Runtime/MovieSceneTracks/Private/EntitySystem/Interrogation/MovieSceneInterrogatedPropertyInstantiator.cpp
@@ -181,7 +181,8 @@ void UMovieSceneInterrogatedPropertyInstantiatorSystem::UpdateOutput(UE::MovieSc
 	UMovieSceneBlenderSystem* ExistingBlender = Output->Blender.Get();
 	if (ExistingBlender && BlenderClass != ExistingBlender->GetClass())
 	{
-		ExistingBlender->ReleaseBlendChannel(Output->BlendChannel);
+		const FMovieSceneBlendChannelID BlendChannel(ExistingBlender->GetBlenderSystemID(), Output->BlendChannel);
+		ExistingBlender->ReleaseBlendChannel(BlendChannel);
 		Output->BlendChannel = INVALID_BLEND_CHANNEL;
 	}
 
@@ -190,7 +191,8 @@ void UMovieSceneInterrogatedPropertyInstantiatorSystem::UpdateOutput(UE::MovieSc
 	const bool bWasAlreadyBlended = Output->BlendChannel != INVALID_BLEND_CHANNEL;
 	if (!bWasAlreadyBlended)
 	{
-		Output->BlendChannel = Output->Blender->AllocateBlendChannel();
+		const FMovieSceneBlendChannelID BlendChannel = Output->Blender->AllocateBlendChannel();
+		Output->BlendChannel = BlendChannel.ChannelID;
 	}
 
 	FComponentMask NewMask;
@@ -220,10 +222,12 @@ void UMovieSceneInterrogatedPropertyInstantiatorSystem::UpdateOutput(UE::MovieSc
 		}
 		NewMask.Set(PropertyDefinition->PropertyType);
 
+		const FMovieSceneBlendChannelID BlendChannel(Output->Blender->GetBlenderSystemID(), Output->BlendChannel);
+
 		// Never seen this property before
 		FMovieSceneEntityID NewEntityID = FEntityBuilder()
 		.Add(BuiltInComponents->Interrogation.OutputKey, Key)
-		.Add(BuiltInComponents->BlendChannelOutput, Output->BlendChannel)
+		.Add(BuiltInComponents->BlendChannelOutput, BlendChannel)
 		.AddTagConditional(BuiltInComponents->Tags.MigratedFromFastPath, Output->PropertyEntityID.IsValid())
 		.AddTag(BuiltInComponents->Tags.NeedsLink)
 		.AddMutualComponents()
@@ -235,7 +239,7 @@ void UMovieSceneInterrogatedPropertyInstantiatorSystem::UpdateOutput(UE::MovieSc
 			Linker->EntityManager.CopyComponents(Output->PropertyEntityID, NewEntityID, Linker->EntityManager.GetComponents()->GetMigrationMask());
 
 			// Add blend inputs on the first contributor, which was using the fast-path
-			Linker->EntityManager.AddComponent(Output->PropertyEntityID, BuiltInComponents->BlendChannelInput, Output->BlendChannel);
+			Linker->EntityManager.AddComponent(Output->PropertyEntityID, BuiltInComponents->BlendChannelInput, BlendChannel);
 			Linker->EntityManager.RemoveComponents(Output->PropertyEntityID, CleanFastPathMask);
 		}
 
@@ -262,7 +266,8 @@ void UMovieSceneInterrogatedPropertyInstantiatorSystem::UpdateOutput(UE::MovieSc
 	// Ensure contributors all have the necessary blend inputs and tags
 	for (FMovieSceneEntityID Input : Inputs)
 	{
-		Linker->EntityManager.AddComponent(Input, BuiltInComponents->BlendChannelInput, Output->BlendChannel);
+		const FMovieSceneBlendChannelID BlendChannel(Output->Blender->GetBlenderSystemID(), Output->BlendChannel);
+		Linker->EntityManager.AddComponent(Input, BuiltInComponents->BlendChannelInput, BlendChannel);
 		Linker->EntityManager.RemoveComponents(Input, CleanFastPathMask);
 	}
 }
@@ -273,7 +278,8 @@ void UMovieSceneInterrogatedPropertyInstantiatorSystem::DestroyOutput(UE::MovieS
 	{
 		if (UMovieSceneBlenderSystem* Blender = Output->Blender.Get())
 		{
-			Blender->ReleaseBlendChannel(Output->BlendChannel);
+			const FMovieSceneBlendChannelID BlendChannel(Blender->GetBlenderSystemID(), Output->BlendChannel);
+			Blender->ReleaseBlendChannel(BlendChannel);
 		}
 		Linker->EntityManager.AddComponents(Output->PropertyEntityID, BuiltInComponents->FinishedMask);
 	}
diff --git a/Engine/Source/Runtime/MovieSceneTracks/Private/Systems/MovieSceneHierarchicalBiasSystem.cpp b/Engine/Source/Runtime/MovieSceneTracks/Private/Systems/MovieSceneHierarchicalBiasSystem.cpp
--- a/Engine/Source/Runtime/MovieSceneTracks/Private/Systems/MovieSceneHierarchicalBiasSystem.cpp
+++ b/Engine/Source/Runtime/MovieSceneTracks/Private/Systems/MovieSceneHierarchicalBiasSystem.cpp
@@ -18,7 +18,7 @@ struct FHierarchicalBiasTask
 		: Linker(InLinker)
 	{}
 
-	void InitializeChannel(uint16 BlendChannel)
+	void InitializeChannel(FMovieSceneBlendChannelID BlendChannel)
 	{
 		MaxBiasByChannel.FindOrAdd(BlendChannel, MIN_int16);
 	}
@@ -28,7 +28,7 @@ struct FHierarchicalBiasTask
 		return MaxBiasByChannel.Num() != 0;
 	}
 
-	void ForEachAllocation(const FEntityAllocation* Allocation, TRead<FMovieSceneEntityID> EntityIDs, TRead<uint16> BlendChannels, TReadOptional<int16> OptHBiases)
+	void ForEachAllocation(const FEntityAllocation* Allocation, TRead<FMovieSceneEntityID> EntityIDs, TRead<FMovieSceneBlendChannelID> BlendChannels, TReadOptional<int16> OptHBiases)
 	{
 		const int32 Num = Allocation->Num();
 		if (OptHBiases)
@@ -64,7 +64,7 @@ struct FHierarchicalBiasTask
 
 private:
 
-	void VisitChannel(FMovieSceneEntityID EntityID, uint16 BlendChannel, int16 HBias)
+	void VisitChannel(FMovieSceneEntityID EntityID, FMovieSceneBlendChannelID BlendChannel, int16 HBias)
 	{
 		// If this channel hasn't changed at all (ie InitializeChannel was not called for it) do nothing
 		if (int16* ExistingBias = MaxBiasByChannel.Find(BlendChannel))
@@ -91,11 +91,11 @@ private:
 		}
 	}
 
-	TMap<uint16, int16> MaxBiasByChannel;
+	TMap<FMovieSceneBlendChannelID, int16> MaxBiasByChannel;
 
-	TMultiMap<uint16, FMovieSceneEntityID> InactiveContributorsByChannel;
+	TMultiMap<FMovieSceneBlendChannelID, FMovieSceneEntityID> InactiveContributorsByChannel;
 
-	TMultiMap<uint16, FMovieSceneEntityID> ActiveContributorsByChannel;
+	TMultiMap<FMovieSceneBlendChannelID, FMovieSceneEntityID> ActiveContributorsByChannel;
 
 	UMovieSceneEntitySystemLinker* Linker;
 };
@@ -135,7 +135,7 @@ void UMovieSceneHierarchicalBiasSystem::OnRun(FSystemTaskPrerequisites& InPrereq
 	FEntityTaskBuilder()
 	.Read(Components->BlendChannelInput)
 	.FilterAny({ Components->Tags.NeedsLink, Components->Tags.NeedsUnlink })
-	.Iterate_PerEntity(&Linker->EntityManager, [&Task](uint16 BlendChannel){ Task.InitializeChannel(BlendChannel); });
+	.Iterate_PerEntity(&Linker->EntityManager, [&Task](FMovieSceneBlendChannelID BlendChannel){ Task.InitializeChannel(BlendChannel); });
 
 	if (Task.HasAnyWork())
 	{
@@ -146,4 +146,4 @@ void UMovieSceneHierarchicalBiasSystem::OnRun(FSystemTaskPrerequisites& InPrereq
 		.FilterNone({ Components->Tags.NeedsUnlink })
 		.RunInline_PerAllocation(&Linker->EntityManager, Task);
 	}
-}
\ No newline at end of file
+}
diff --git a/Engine/Source/Runtime/MovieSceneTracks/Private/Systems/MovieScenePiecewiseFloatBlenderSystem.cpp b/Engine/Source/Runtime/MovieSceneTracks/Private/Systems/MovieScenePiecewiseFloatBlenderSystem.cpp
--- a/Engine/Source/Runtime/MovieSceneTracks/Private/Systems/MovieScenePiecewiseFloatBlenderSystem.cpp
+++ b/Engine/Source/Runtime/MovieSceneTracks/Private/Systems/MovieScenePiecewiseFloatBlenderSystem.cpp
@@ -65,7 +65,7 @@ struct FAccumulationTask
 	{}
 
 	/** Task entry point - iterates the allocation's headers and accumulates float results for any required components */
-	void ForEachAllocation(FEntityAllocationIteratorItem InItem, TRead<uint16> BlendIDs, TReadOptional<float> OptionalEasingAndWeights)
+	void ForEachAllocation(FEntityAllocationIteratorItem InItem, TRead<FMovieSceneBlendChannelID> BlendIDs, TReadOptional<float> OptionalEasingAndWeights)
 	{
 		const FEntityAllocation* Allocation = InItem;
 		const FComponentMask& AllocationType = InItem;
@@ -86,15 +86,20 @@ struct FAccumulationTask
 
 private:
 
-	void AccumulateResults(const FEntityAllocation* InAllocation, const float* InFloatResults, const uint16* BlendIDs, const float* OptionalEasingAndWeights, TArray<FBlendResult>& OutBlendResults)
+	void AccumulateResults(const FEntityAllocation* InAllocation, const float* InFloatResults, const FMovieSceneBlendChannelID* BlendIDs, const float* OptionalEasingAndWeights, TArray<FBlendResult>& OutBlendResults)
 	{
+		static const FMovieSceneBlenderSystemID FloatBlenderSystemID = UMovieSceneBlenderSystem::GetBlenderSystemID<UMovieScenePiecewiseFloatBlenderSystem>();
+
 		const int32 Num = InAllocation->Num();
 		if (OptionalEasingAndWeights)
 		{
 			// We have some easing/weight factors to multiply values with.
 			for (int32 Index = 0; Index < Num; ++Index)
 			{
-				FBlendResult& Result = OutBlendResults[BlendIDs[Index]];
+				const FMovieSceneBlendChannelID& BlendID(BlendIDs[Index]);
+				ensureMsgf(BlendID.SystemID == FloatBlenderSystemID, TEXT("Overriding the standard blender system of standard types isn't supported."));
+
+				FBlendResult& Result = OutBlendResults[BlendID.ChannelID];
 
 				const float Weight = OptionalEasingAndWeights[Index];
 				Result.Total  += InFloatResults[Index] * Weight;
@@ -106,7 +111,10 @@ private:
 			// Faster path for when there's no weight to multiply values with.
 			for (int32 Index = 0; Index < Num; ++Index)
 			{
-				FBlendResult& Result = OutBlendResults[BlendIDs[Index]];
+				const FMovieSceneBlendChannelID& BlendID(BlendIDs[Index]);
+				ensureMsgf(BlendID.SystemID == FloatBlenderSystemID, TEXT("Overriding the standard blender system of standard types isn't supported."));
+
+				FBlendResult& Result = OutBlendResults[BlendID.ChannelID];
 				Result.Total  += InFloatResults[Index];
 				Result.Weight += 1.f;
 			}
@@ -124,7 +132,7 @@ struct FAdditiveFromBaseBlendTask
 {
 	TSortedMap<FComponentTypeID, FAdditiveFromBaseBuffer>* AccumulationBuffers;
 
-	void ForEachAllocation(FEntityAllocationIteratorItem InItem, TRead<uint16> BlendIDs, TReadOptional<float> EasingAndWeightResults)
+	void ForEachAllocation(FEntityAllocationIteratorItem InItem, TRead<FMovieSceneBlendChannelID> BlendIDs, TReadOptional<float> EasingAndWeightResults)
 	{
 		FEntityAllocation* Allocation = InItem;
 		const FComponentMask& AllocationType = InItem;
@@ -143,15 +151,20 @@ struct FAdditiveFromBaseBlendTask
 
 private:
 
-	void AccumulateResults(const FEntityAllocation* InAllocation, const float* FloatResults, const float* BaseValues, const uint16* BlendIDs, const float* OptionalEasingAndWeights, TArray<FBlendResult>& OutBlendResults)
+	void AccumulateResults(const FEntityAllocation* InAllocation, const float* FloatResults, const float* BaseValues, const FMovieSceneBlendChannelID* BlendIDs, const float* OptionalEasingAndWeights, TArray<FBlendResult>& OutBlendResults)
 	{
+		static const FMovieSceneBlenderSystemID FloatBlenderSystemID = UMovieSceneBlenderSystem::GetBlenderSystemID<UMovieScenePiecewiseFloatBlenderSystem>();
+		
 		const int32 Num = InAllocation->Num();
 
 		if (OptionalEasingAndWeights)
 		{
 			for (int32 Index = 0; Index < Num; ++Index)
 			{
-				FBlendResult& Result = OutBlendResults[BlendIDs[Index]];
+				const FMovieSceneBlendChannelID& BlendID(BlendIDs[Index]);
+				ensureMsgf(BlendID.SystemID == FloatBlenderSystemID, TEXT("Overriding the standard blender system of standard types isn't supported."));
+
+				FBlendResult& Result = OutBlendResults[BlendID.ChannelID];
 
 				const float Weight = OptionalEasingAndWeights[Index];
 				Result.Total  += (FloatResults[Index] - BaseValues[Index]) * Weight;
@@ -163,7 +176,10 @@ private:
 			// Faster path for when there's no weight to multiply values with.
 			for (int32 Index = 0; Index < Num; ++Index)
 			{
-				FBlendResult& Result = OutBlendResults[BlendIDs[Index]];
+				const FMovieSceneBlendChannelID& BlendID(BlendIDs[Index]);
+				ensureMsgf(BlendID.SystemID == FloatBlenderSystemID, TEXT("Overriding the standard blender system of standard types isn't supported."));
+
+				FBlendResult& Result = OutBlendResults[BlendID.ChannelID];
 				Result.Total  += (FloatResults[Index] - BaseValues[Index]);
 				Result.Weight += 1.f;
 			}
@@ -181,7 +197,7 @@ struct FCombineBlends
 		, WriteContext(InWriteContext)
 	{}
 
-	void ForEachAllocation(FEntityAllocationIteratorItem InItem, TRead<uint16> BlendIDs)
+	void ForEachAllocation(FEntityAllocationIteratorItem InItem, TRead<FMovieSceneBlendChannelID> BlendIDs)
 	{
 		FEntityAllocation* Allocation = InItem;
 		const FComponentMask& AllocationType = InItem;
@@ -200,10 +216,11 @@ struct FCombineBlends
 
 private:
 
-	void ProcessPropertyType(FEntityAllocation* Allocation, const FComponentMask& AllocationType, const FPropertyDefinition& PropertyDefinition, const uint16* BlendIDs)
+	void ProcessPropertyType(FEntityAllocation* Allocation, const FComponentMask& AllocationType, const FPropertyDefinition& PropertyDefinition, const FMovieSceneBlendChannelID* BlendIDs)
 	{
-		TArrayView<const FPropertyCompositeDefinition> Composites = PropertyRegistry->GetComposites(PropertyDefinition);
+		static const FMovieSceneBlenderSystemID FloatBlenderSystemID = UMovieSceneBlenderSystem::GetBlenderSystemID<UMovieScenePiecewiseFloatBlenderSystem>();
 
+		TArrayView<const FPropertyCompositeDefinition> Composites = PropertyRegistry->GetComposites(PropertyDefinition);
 
 		FOptionalComponentReader OptInitialValues = Allocation->TryReadComponentsErased(PropertyDefinition.InitialValueType);
 
@@ -232,15 +249,17 @@ private:
 				{
 					for (int32 Index = 0; Index < Allocation->Num(); ++Index)
 					{
+						ensureMsgf(BlendIDs[Index].SystemID == FloatBlenderSystemID, TEXT("Overriding the standard blender system of standard types isn't supported."));
 						const float InitialValue = *reinterpret_cast<const float*>(static_cast<const uint8*>(OptInitialValues[Index]) + InitialValueProjectionOffset);
-						BlendResultsWithInitial(Results, BlendIDs[Index], InitialValue, FloatResults[Index]);
+						BlendResultsWithInitial(Results, BlendIDs[Index].ChannelID, InitialValue, FloatResults[Index]);
 					}
 				}
 				else
 				{
 					for (int32 Index = 0; Index < Allocation->Num(); ++Index)
 					{
-						BlendResults(Results, BlendIDs[Index], FloatResults[Index]);
+						ensureMsgf(BlendIDs[Index].SystemID == FloatBlenderSystemID, TEXT("Overriding the standard blender system of standard types isn't supported."));
+						BlendResults(Results, BlendIDs[Index].ChannelID, FloatResults[Index]);
 					}
 				}
 			}
@@ -608,14 +627,19 @@ FGraphEventRef UMovieScenePiecewiseFloatBlenderSystem::DispatchDecomposeTask(con
 			EntitiesToDecompose.Append(Params.Query.Entities.GetData(), Params.Query.Entities.Num());
 		}
 
-		void ForEachAllocation(const FEntityAllocation* Allocation, TRead<FMovieSceneEntityID> EntityToDecomposeIDs, TRead<uint16> BlendChannels, TRead<float> FloatResultComponent, TReadOptional<float> OptionalWeightComponent)
+		void ForEachAllocation(const FEntityAllocation* Allocation, TRead<FMovieSceneEntityID> EntityToDecomposeIDs, TRead<FMovieSceneBlendChannelID> BlendChannels, TRead<float> FloatResultComponent, TReadOptional<float> OptionalWeightComponent)
 		{
+			static const FMovieSceneBlenderSystemID FloatBlenderSystemID = UMovieSceneBlenderSystem::GetBlenderSystemID<UMovieScenePiecewiseFloatBlenderSystem>();
+
 			const bool bAdditive = Allocation->HasComponent(AdditiveBlendTag);
 
 			const int32 Num = Allocation->Num();
 			for (int32 EntityIndex = 0; EntityIndex < Num; ++EntityIndex)
 			{
-				if (BlendChannels[EntityIndex] != DecomposeBlendChannel)
+				const FMovieSceneBlendChannelID& BlendChannel(BlendChannels[EntityIndex]);
+				ensureMsgf(BlendChannel.SystemID == FloatBlenderSystemID, TEXT("Overriding the standard blender system of standard types isn't supported."));
+
+				if (BlendChannel.ChannelID != DecomposeBlendChannel)
 				{
 					continue;
 				}
diff --git a/Engine/Source/Runtime/MovieSceneTracks/Private/Systems/MovieScenePiecewiseIntegerBlenderSystem.cpp b/Engine/Source/Runtime/MovieSceneTracks/Private/Systems/MovieScenePiecewiseIntegerBlenderSystem.cpp
--- a/Engine/Source/Runtime/MovieSceneTracks/Private/Systems/MovieScenePiecewiseIntegerBlenderSystem.cpp
+++ b/Engine/Source/Runtime/MovieSceneTracks/Private/Systems/MovieScenePiecewiseIntegerBlenderSystem.cpp
@@ -25,8 +25,10 @@ struct FIntegerAccumulationTask
 	{}
 
 	/** Task entry point - iterates the allocation's headers and accumulates int32 results for any required components */
-	void ForEachAllocation(const FEntityAllocation* Allocation, TRead<uint16> BlendIDs, TRead<int32> Integers, TReadOptional<float> EasingAndWeights)
+	void ForEachAllocation(const FEntityAllocation* Allocation, TRead<FMovieSceneBlendChannelID> BlendIDs, TRead<int32> Integers, TReadOptional<float> EasingAndWeights)
 	{
+		static const FMovieSceneBlenderSystemID IntegerBlenderSystemID = UMovieSceneBlenderSystem::GetBlenderSystemID<UMovieScenePiecewiseIntegerBlenderSystem>();
+
 		const int32 Num = Allocation->Num();
 
 		if (EasingAndWeights)
@@ -34,7 +36,10 @@ struct FIntegerAccumulationTask
 			// We have some easing/weight factors to multiply values with.
 			for (int32 Index = 0; Index < Num; ++Index)
 			{
-				FIntegerBlendResult& Result = AccumulationBuffer[BlendIDs[Index]];
+				const FMovieSceneBlendChannelID& BlendID(BlendIDs[Index]);
+				ensureMsgf(BlendID.SystemID == IntegerBlenderSystemID, TEXT("Overriding the standard blender system of standard types isn't supported."));
+
+				FIntegerBlendResult& Result = AccumulationBuffer[BlendID.ChannelID];
 
 				const float Weight = EasingAndWeights[Index];
 				Result.Total  += (int32)(Integers[Index] * Weight);
@@ -46,7 +51,11 @@ struct FIntegerAccumulationTask
 			// Faster path for when there's no weight to multiply values with.
 			for (int32 Index = 0; Index < Num; ++Index)
 			{
-				FIntegerBlendResult& Result = AccumulationBuffer[BlendIDs[Index]];
+				const FMovieSceneBlendChannelID& BlendID(BlendIDs[Index]);
+				ensureMsgf(BlendID.SystemID == IntegerBlenderSystemID, TEXT("Overriding the standard blender system of standard types isn't supported."));
+
+				FIntegerBlendResult& Result = AccumulationBuffer[BlendID.ChannelID];
+
 				Result.Total  += Integers[Index];
 				Result.Weight += 1.f;
 			}
@@ -66,15 +75,20 @@ struct FIntegerAdditiveFromBaseBlendTask
 		: AccumulationBuffer(*InAccumulationBuffer)
 	{}
 
-	void ForEachAllocation(const FEntityAllocation* Allocation, TRead<uint16> BlendIDs, TRead<int32> Integers, TRead<int32> BaseValues, TReadOptional<float> EasingAndWeights)
+	void ForEachAllocation(const FEntityAllocation* Allocation, TRead<FMovieSceneBlendChannelID> BlendIDs, TRead<int32> Integers, TRead<int32> BaseValues, TReadOptional<float> EasingAndWeights)
 	{
+		static const FMovieSceneBlenderSystemID IntegerBlenderSystemID = UMovieSceneBlenderSystem::GetBlenderSystemID<UMovieScenePiecewiseIntegerBlenderSystem>();
+
 		const int32 Num = Allocation->Num();
 
 		if (EasingAndWeights)
 		{
 			for (int32 Index = 0; Index < Num; ++Index)
 			{
-				FIntegerBlendResult& Result = AccumulationBuffer[BlendIDs[Index]];
+				const FMovieSceneBlendChannelID& BlendID(BlendIDs[Index]);
+				ensureMsgf(BlendID.SystemID == IntegerBlenderSystemID, TEXT("Overriding the standard blender system of standard types isn't supported."));
+
+				FIntegerBlendResult& Result = AccumulationBuffer[BlendID.ChannelID];
 
 				const float Weight = EasingAndWeights[Index];
 				Result.Total  += (int32)((Integers[Index] - BaseValues[Index]) * Weight);
@@ -86,7 +100,11 @@ struct FIntegerAdditiveFromBaseBlendTask
 			// Faster path for when there's no weight to multiply values with.
 			for (int32 Index = 0; Index < Num; ++Index)
 			{
-				FIntegerBlendResult& Result = AccumulationBuffer[BlendIDs[Index]];
+				const FMovieSceneBlendChannelID& BlendID(BlendIDs[Index]);
+				ensureMsgf(BlendID.SystemID == IntegerBlenderSystemID, TEXT("Overriding the standard blender system of standard types isn't supported."));
+
+				FIntegerBlendResult& Result = AccumulationBuffer[BlendID.ChannelID];
+
 				Result.Total  += (Integers[Index] - BaseValues[Index]);
 				Result.Weight += 1.f;
 			}
@@ -103,20 +121,24 @@ struct FIntegerCombineBlends
 		: AccumulationBuffers(*InAccumulationBuffers)
 	{}
 
-	void ForEachAllocation(FEntityAllocation* Allocation, TRead<uint16> BlendIDs, TReadOptional<int32> InitialValues, TWrite<int32> IntegerResults)
+	void ForEachAllocation(FEntityAllocation* Allocation, TRead<FMovieSceneBlendChannelID> BlendIDs, TReadOptional<int32> InitialValues, TWrite<int32> IntegerResults)
 	{
+		static const FMovieSceneBlenderSystemID IntegerBlenderSystemID = UMovieSceneBlenderSystem::GetBlenderSystemID<UMovieScenePiecewiseIntegerBlenderSystem>();
+
 		if (InitialValues)
 		{
 			for (int32 Index = 0; Index < Allocation->Num(); ++Index)
 			{
-				BlendResultsWithInitial(BlendIDs[Index], InitialValues[Index], IntegerResults[Index]);
+				ensureMsgf(BlendIDs[Index].SystemID == IntegerBlenderSystemID, TEXT("Overriding the standard blender system of standard types isn't supported."));
+				BlendResultsWithInitial(BlendIDs[Index].ChannelID, InitialValues[Index], IntegerResults[Index]);
 			}
 		}
 		else
 		{
 			for (int32 Index = 0; Index < Allocation->Num(); ++Index)
 			{
-				BlendResults(BlendIDs[Index], IntegerResults[Index]);
+				ensureMsgf(BlendIDs[Index].SystemID == IntegerBlenderSystemID, TEXT("Overriding the standard blender system of standard types isn't supported."));
+				BlendResults(BlendIDs[Index].ChannelID, IntegerResults[Index]);
 			}
 		}
 	}
diff --git a/Engine/Source/Runtime/MovieSceneTracks/Private/Systems/MovieScenePropertyInstantiator.cpp b/Engine/Source/Runtime/MovieSceneTracks/Private/Systems/MovieScenePropertyInstantiator.cpp
--- a/Engine/Source/Runtime/MovieSceneTracks/Private/Systems/MovieScenePropertyInstantiator.cpp
+++ b/Engine/Source/Runtime/MovieSceneTracks/Private/Systems/MovieScenePropertyInstantiator.cpp
@@ -222,7 +222,8 @@ void UMovieScenePropertyInstantiatorSystem::ProcessInvalidatedProperties(const T
 			{
 				if (UMovieSceneBlenderSystem* Blender = PropertyInfo->Blender.Get())
 				{
-					Blender->ReleaseBlendChannel(PropertyInfo->BlendChannel);
+					const FMovieSceneBlendChannelID BlendChannelID(Blender->GetBlenderSystemID(), PropertyInfo->BlendChannel);
+					Blender->ReleaseBlendChannel(BlendChannelID);
 				}
 				Linker->EntityManager.AddComponents(PropertyInfo->PropertyEntityID, BuiltInComponents->FinishedMask);
 			}
@@ -405,18 +406,23 @@ void UMovieScenePropertyInstantiatorSystem::InitializeBlendPath(const FPropertyP
 	UMovieSceneBlenderSystem* ExistingBlender = Params.PropertyInfo->Blender.Get();
 	if (ExistingBlender && BlenderClass != ExistingBlender->GetClass())
 	{
-		ExistingBlender->ReleaseBlendChannel(Params.PropertyInfo->BlendChannel);
+		const FMovieSceneBlendChannelID ExistingBlendChannel(ExistingBlender->GetBlenderSystemID(), Params.PropertyInfo->BlendChannel);
+		ExistingBlender->ReleaseBlendChannel(ExistingBlendChannel);
 		Params.PropertyInfo->BlendChannel = INVALID_BLEND_CHANNEL;
 	}
 
-	Params.PropertyInfo->Blender = CastChecked<UMovieSceneBlenderSystem>(Linker->LinkSystem(BlenderClass));
+	UMovieSceneBlenderSystem* const Blender = CastChecked<UMovieSceneBlenderSystem>(Linker->LinkSystem(BlenderClass));
+	Params.PropertyInfo->Blender = Blender;
 
 	const bool bWasAlreadyBlended = Params.PropertyInfo->BlendChannel != INVALID_BLEND_CHANNEL;
 	if (!bWasAlreadyBlended)
 	{
-		Params.PropertyInfo->BlendChannel = Params.PropertyInfo->Blender->AllocateBlendChannel();
+		const FMovieSceneBlendChannelID NewBlendChannel = Blender->AllocateBlendChannel();
+		Params.PropertyInfo->BlendChannel = NewBlendChannel.ChannelID;
 	}
 
+	const FMovieSceneBlendChannelID BlendChannel(Blender->GetBlenderSystemID(), Params.PropertyInfo->BlendChannel);
+
 	if (!bWasAlreadyBlended)
 	{
 		FComponentMask NewMask;
@@ -451,7 +457,7 @@ void UMovieScenePropertyInstantiatorSystem::InitializeBlendPath(const FPropertyP
 			.Add(BuiltInComponents->FastPropertyOffset,      Params.PropertyInfo->Property.template Get<uint16>())
 			.Add(BuiltInComponents->PropertyBinding,         Params.PropertyInfo->PropertyBinding)
 			.Add(BuiltInComponents->BoundObject,             Params.PropertyInfo->BoundObject)
-			.Add(BuiltInComponents->BlendChannelOutput,      Params.PropertyInfo->BlendChannel)
+			.Add(BuiltInComponents->BlendChannelOutput,      BlendChannel)
 			.AddTagConditional(BuiltInComponents->Tags.MigratedFromFastPath, Params.PropertyInfo->PropertyEntityID.IsValid())
 			.AddTagConditional(BuiltInComponents->Tags.RestoreState, Params.PropertyInfo->bWantsRestoreState)
 			.AddTag(BuiltInComponents->Tags.NeedsLink)
@@ -464,7 +470,7 @@ void UMovieScenePropertyInstantiatorSystem::InitializeBlendPath(const FPropertyP
 			.Add(BuiltInComponents->CustomPropertyIndex, Params.PropertyInfo->Property.template Get<FCustomPropertyIndex>())
 			.Add(BuiltInComponents->PropertyBinding,     Params.PropertyInfo->PropertyBinding)
 			.Add(BuiltInComponents->BoundObject,         Params.PropertyInfo->BoundObject)
-			.Add(BuiltInComponents->BlendChannelOutput,  Params.PropertyInfo->BlendChannel)
+			.Add(BuiltInComponents->BlendChannelOutput,  BlendChannel)
 			.AddTagConditional(BuiltInComponents->Tags.MigratedFromFastPath, Params.PropertyInfo->PropertyEntityID.IsValid())
 			.AddTagConditional(BuiltInComponents->Tags.RestoreState, Params.PropertyInfo->bWantsRestoreState)
 			.AddTag(BuiltInComponents->Tags.NeedsLink)
@@ -477,7 +483,7 @@ void UMovieScenePropertyInstantiatorSystem::InitializeBlendPath(const FPropertyP
 			.Add(BuiltInComponents->SlowProperty,            Params.PropertyInfo->Property.template Get<FSlowPropertyPtr>())
 			.Add(BuiltInComponents->PropertyBinding,         Params.PropertyInfo->PropertyBinding)
 			.Add(BuiltInComponents->BoundObject,             Params.PropertyInfo->BoundObject)
-			.Add(BuiltInComponents->BlendChannelOutput,      Params.PropertyInfo->BlendChannel)
+			.Add(BuiltInComponents->BlendChannelOutput,      BlendChannel)
 			.AddTagConditional(BuiltInComponents->Tags.MigratedFromFastPath, Params.PropertyInfo->PropertyEntityID.IsValid())
 			.AddTagConditional(BuiltInComponents->Tags.RestoreState, Params.PropertyInfo->bWantsRestoreState)
 			.AddTag(BuiltInComponents->Tags.NeedsLink)
@@ -492,7 +498,7 @@ void UMovieScenePropertyInstantiatorSystem::InitializeBlendPath(const FPropertyP
 			Linker->EntityManager.CopyComponents(Params.PropertyInfo->PropertyEntityID, NewEntityID, Linker->EntityManager.GetComponents()->GetCopyAndMigrationMask());
 
 			// Add blend inputs on the first contributor, which was using the fast-path
-			Linker->EntityManager.AddComponent(Params.PropertyInfo->PropertyEntityID, BuiltInComponents->BlendChannelInput, Params.PropertyInfo->BlendChannel);
+			Linker->EntityManager.AddComponent(Params.PropertyInfo->PropertyEntityID, BuiltInComponents->BlendChannelInput, BlendChannel);
 			Linker->EntityManager.RemoveComponents(Params.PropertyInfo->PropertyEntityID, CleanFastPathMask);
 		}
 
@@ -525,8 +531,8 @@ void UMovieScenePropertyInstantiatorSystem::InitializeBlendPath(const FPropertyP
 	// Ensure contributors all have the necessary blend inputs and tags
 	for (auto ContributorIt = NewContributors.CreateConstKeyIterator(Params.PropertyInfoIndex); ContributorIt; ++ContributorIt)
 	{
-		FMovieSceneEntityID Contributor = ContributorIt.Value();
-		Linker->EntityManager.AddComponent(Contributor, BuiltInComponents->BlendChannelInput, Params.PropertyInfo->BlendChannel);
+		const FMovieSceneEntityID Contributor = ContributorIt.Value();
+		Linker->EntityManager.AddComponent(Contributor, BuiltInComponents->BlendChannelInput, BlendChannel);
 		Linker->EntityManager.RemoveComponents(Contributor, CleanFastPathMask);
 	}
 }
diff --git a/Engine/Source/Runtime/MovieSceneTracks/Public/Systems/MovieSceneBlenderSystemHelper.h b/Engine/Source/Runtime/MovieSceneTracks/Public/Systems/MovieSceneBlenderSystemHelper.h
--- a/Engine/Source/Runtime/MovieSceneTracks/Public/Systems/MovieSceneBlenderSystemHelper.h
+++ b/Engine/Source/Runtime/MovieSceneTracks/Public/Systems/MovieSceneBlenderSystemHelper.h
@@ -58,19 +58,22 @@ struct TSimpleBlendResultTraits
 template<typename PropertyType>
 struct TSimpleBlenderGatherResults
 {
+	FMovieSceneBlenderSystemID SystemID;
 	TArray<TSimpleBlendResult<PropertyType>>& BlendChannelResults;
 
-	TSimpleBlenderGatherResults(TArray<TSimpleBlendResult<PropertyType>>& InBlendChannelResults)
-		: BlendChannelResults(InBlendChannelResults)
+	TSimpleBlenderGatherResults(FMovieSceneBlenderSystemID InSystemID, TArray<TSimpleBlendResult<PropertyType>>& InBlendChannelResults)
+		: SystemID(InSystemID)
+		, BlendChannelResults(InBlendChannelResults)
 	{}
 
-	void ForEachEntity(uint16 BlendChannelInput, PropertyType Value)
+	void ForEachEntity(FMovieSceneBlendChannelID BlendChannelInput, PropertyType Value)
 	{
 		using FResultTraits = TSimpleBlendResultTraits<PropertyType>;
 
-		if (ensure(BlendChannelResults.IsValidIndex(BlendChannelInput)))
+		ensureMsgf(BlendChannelInput.SystemID == SystemID, TEXT("Overriding the standard blender system of standard types isn't supported."));
+		if (ensure(BlendChannelResults.IsValidIndex(BlendChannelInput.ChannelID)))
 		{
-			FResultTraits::AccumulateResult(BlendChannelResults[BlendChannelInput], Value);
+			FResultTraits::AccumulateResult(BlendChannelResults[BlendChannelInput.ChannelID], Value);
 		}
 	}
 };
@@ -78,19 +81,22 @@ struct TSimpleBlenderGatherResults
 template<typename PropertyType>
 struct TSimpleBlenderCombineResults
 {
+	FMovieSceneBlenderSystemID SystemID;
 	TArray<TSimpleBlendResult<PropertyType>>& BlendChannelResults;
 
-	TSimpleBlenderCombineResults(TArray<TSimpleBlendResult<PropertyType>>& InBlendChannelResults)
-		: BlendChannelResults(InBlendChannelResults)
+	TSimpleBlenderCombineResults(FMovieSceneBlenderSystemID InSystemID, TArray<TSimpleBlendResult<PropertyType>>& InBlendChannelResults)
+		: SystemID(InSystemID)
+		, BlendChannelResults(InBlendChannelResults)
 	{}
 
-	void ForEachEntity(uint16 BlendChannelOutput, PropertyType& OutValue)
+	void ForEachEntity(FMovieSceneBlendChannelID BlendChannelOutput, PropertyType& OutValue)
 	{
 		using FResultTraits = TSimpleBlendResultTraits<PropertyType>;
 
-		if (ensure(BlendChannelResults.IsValidIndex(BlendChannelOutput)))
+		ensureMsgf(BlendChannelOutput.SystemID == SystemID, TEXT("Overriding the standard blender system of standard types isn't supported."));
+		if (ensure(BlendChannelResults.IsValidIndex(BlendChannelOutput.ChannelID)))
 		{
-			OutValue = FResultTraits::BlendResult(BlendChannelResults[BlendChannelOutput]);
+			OutValue = FResultTraits::BlendResult(BlendChannelResults[BlendChannelOutput.ChannelID]);
 		}
 	}
 };
@@ -170,7 +176,7 @@ public:
 			.Read(ResultComponentID)
 			.FilterNone({ BuiltInComponents->Tags.Ignored })
 			.template Dispatch_PerEntity<TSimpleBlenderGatherResults<PropertyType>>(
-					&Linker->EntityManager, InPrerequisites, nullptr, BlendChannelResults);
+					&Linker->EntityManager, InPrerequisites, nullptr, BlenderSystem->GetBlenderSystemID(), BlendChannelResults);
 		if (Task)
 		{
 			Prereqs.AddMasterTask(Task);
@@ -182,7 +188,7 @@ public:
 			.Write(ResultComponentID)
 			.FilterNone({ BuiltInComponents->Tags.Ignored })
 			.template Dispatch_PerEntity<TSimpleBlenderCombineResults<PropertyType>>(
-					&Linker->EntityManager, Prereqs, &Subsequents, BlendChannelResults);
+					&Linker->EntityManager, Prereqs, &Subsequents, BlenderSystem->GetBlenderSystemID(), BlendChannelResults);
 	}
 
 private:
