Sequencer: Deleted old 3DTransform Template code

diff --git a/Engine/Source/Runtime/MovieSceneTracks/Private/Evaluation/MovieScene3DTransformTemplate.cpp b/Engine/Source/Runtime/MovieSceneTracks/Private/Evaluation/MovieScene3DTransformTemplate.cpp
deleted file mode 100644
diff --git a/Engine/Source/Runtime/MovieSceneTracks/Private/Evaluation/MovieScenePropertyTemplates.h b/Engine/Source/Runtime/MovieSceneTracks/Private/Evaluation/MovieScenePropertyTemplates.h
--- a/Engine/Source/Runtime/MovieSceneTracks/Private/Evaluation/MovieScenePropertyTemplates.h
+++ b/Engine/Source/Runtime/MovieSceneTracks/Private/Evaluation/MovieScenePropertyTemplates.h
@@ -7,7 +7,6 @@
 #include "Evaluation/MovieSceneEvalTemplate.h"
 #include "Evaluation/MovieScenePropertyTemplate.h"
 #include "Sections/MovieSceneStringSection.h"
-#include "Evaluation/MovieScene3DTransformTemplate.h"
 #include "Evaluation/Blending/MovieSceneBlendType.h"
 #include "Channels/MovieSceneBoolChannel.h"
 #include "Channels/MovieSceneByteChannel.h"
diff --git a/Engine/Source/Runtime/MovieSceneTracks/Public/Evaluation/MovieScene3DTransformTemplate.h b/Engine/Source/Runtime/MovieSceneTracks/Public/Evaluation/MovieScene3DTransformTemplate.h
deleted file mode 100644
