Sequencer: Fixed crash when encountering SceneComponent impersonators with pre-animated state

#tests Ran all engine tests, tested sequence camera shakes playing inside sequencer

diff --git a/Engine/Source/Runtime/MovieScene/Public/Evaluation/PreAnimatedState/MovieScenePreAnimatedPropertyStorage.h b/Engine/Source/Runtime/MovieScene/Public/Evaluation/PreAnimatedState/MovieScenePreAnimatedPropertyStorage.h
--- a/Engine/Source/Runtime/MovieScene/Public/Evaluation/PreAnimatedState/MovieScenePreAnimatedPropertyStorage.h
+++ b/Engine/Source/Runtime/MovieScene/Public/Evaluation/PreAnimatedState/MovieScenePreAnimatedPropertyStorage.h
@@ -168,7 +168,7 @@ struct TPreAnimatedPropertyStorageImpl<PropertyTraits, TPropertyMetaData<MetaDat
 		}
 	}
 
-private:
+protected:
 
 	struct FAnimatedPropertyKey
 	{
@@ -226,8 +226,6 @@ private:
 	FPreAnimatedStorageID StorageID;
 	FPreAnimatedStateExtension* ParentExtension;
 
-private:
-
 	TSharedPtr<FPreAnimatedObjectGroupManager> ObjectGroupManager;
 };
 
diff --git a/Engine/Source/Runtime/MovieSceneTracks/Private/PreAnimatedState/MovieScenePreAnimatedComponentTransformStorage.cpp b/Engine/Source/Runtime/MovieSceneTracks/Private/PreAnimatedState/MovieScenePreAnimatedComponentTransformStorage.cpp
--- a/Engine/Source/Runtime/MovieSceneTracks/Private/PreAnimatedState/MovieScenePreAnimatedComponentTransformStorage.cpp
+++ b/Engine/Source/Runtime/MovieSceneTracks/Private/PreAnimatedState/MovieScenePreAnimatedComponentTransformStorage.cpp
@@ -12,35 +12,118 @@ namespace MovieScene
 
 TAutoRegisterPreAnimatedStorageID<FPreAnimatedComponentTransformStorage> FPreAnimatedComponentTransformStorage::StorageID;
 
-void FComponentTransformPreAnimatedTraits::CachePreAnimatedValue(UObject* InObject, FIntermediate3DTransform& OutCachedTransform)
-{
-	OutCachedTransform = GetComponentTransform(InObject);
-}
 
-void FComponentTransformPreAnimatedTraits::RestorePreAnimatedValue(const FObjectKey& InKey, FIntermediate3DTransform& CachedTransform, const FRestoreStateParams& Params)
+struct FTemporaryMobilityScope
 {
-	USceneComponent* SceneComponent = Cast<USceneComponent>(InKey.ResolveObjectPtr());
-	if (SceneComponent)
-	{
-		// Ideally we would not be temporarily changing mobility here, but there are some very specific
-		// edge cases where mobility can be legitimately restored whilst pre-animated transforms are still
-		// maintained. One example is where an attach track has previously been run and since restored - 
-		// thus detatching and resetting the transform. If nothing else animates the mobility, this will also
-		// be reset, but the object's global transform may have been captured.
+	// Ideally we would not be temporarily changing mobility here, but there are some very specific
+	// edge cases where mobility can be legitimately restored whilst pre-animated transforms are still
+	// maintained. One example is where an attach track has previously been run and since restored - 
+	// thus detatching and resetting the transform. If nothing else animates the mobility, this will also
+	// be reset, but the object's global transform may have been captured.
 
-		const EComponentMobility::Type PreviousMobility = SceneComponent->Mobility;
+	FTemporaryMobilityScope(UObject* Object)
+		: SceneComponent(Cast<USceneComponent>(Object))
+		, PreviousMobility(SceneComponent ? SceneComponent->Mobility.GetValue() : EComponentMobility::Movable)
+	{
 		if (PreviousMobility != EComponentMobility::Movable)
 		{
 			SceneComponent->SetMobility(EComponentMobility::Movable);
 		}
-
-		SetComponentTransform(SceneComponent, CachedTransform);
-
+	}
+	~FTemporaryMobilityScope()
+	{
 		if (PreviousMobility != EComponentMobility::Movable)
 		{
 			SceneComponent->SetMobility(PreviousMobility);
 		}
 	}
+
+	USceneComponent* SceneComponent;
+	const EComponentMobility::Type PreviousMobility;
+};
+
+
+void FComponentTransformPreAnimatedTraits::SetObjectPropertyValue(UObject* InObject, const FCustomPropertyAccessor& BaseCustomAccessor, const FIntermediate3DTransform& CachedTransform)
+{
+	FTemporaryMobilityScope TemporaryMobilityScope(InObject);
+
+	const TCustomPropertyAccessor<TraitsType>& CustomAccessor = static_cast<const TCustomPropertyAccessor<TraitsType>&>(BaseCustomAccessor);
+	(*CustomAccessor.Functions.Setter)(InObject, CachedTransform);
+}
+
+void FComponentTransformPreAnimatedTraits::SetObjectPropertyValue(UObject* InObject, uint16 PropertyOffset, const FIntermediate3DTransform& CachedTransform)
+{
+	FTemporaryMobilityScope TemporaryMobilityScope(InObject);
+
+	StorageType* PropertyAddress = reinterpret_cast<StorageType*>( reinterpret_cast<uint8*>(InObject) + PropertyOffset );
+	*PropertyAddress = CachedTransform;
+}
+
+void FComponentTransformPreAnimatedTraits::SetObjectPropertyValue(UObject* InObject, FTrackInstancePropertyBindings* PropertyBindings, const FIntermediate3DTransform& CachedTransform)
+{
+	FTemporaryMobilityScope TemporaryMobilityScope(InObject);
+
+	PropertyBindings->CallFunction<StorageType>(*InObject, CachedTransform);
+}
+
+FPreAnimatedComponentTransformStorage::FPreAnimatedComponentTransformStorage()
+	: TPreAnimatedPropertyStorage<FComponentTransformPreAnimatedTraits>(FBuiltInComponentTypes::Get()->PropertyRegistry.GetDefinition(FMovieSceneTracksComponentTypes::Get()->ComponentTransform.CompositeID))
+{}
+
+void FPreAnimatedComponentTransformStorage::CachePreAnimatedTransforms(const FCachePreAnimatedValueParams& Params, TArrayView<UObject* const> BoundObjects)
+{
+	static FMovieScenePropertyBinding PropertyBinding("Transform", TEXT("Transform"));
+
+	for (UObject* BoundObject : BoundObjects)
+	{
+		TOptional<FResolvedProperty> Property = FPropertyRegistry::ResolveProperty(BoundObject, PropertyBinding, CustomAccessors);
+		if (!Property)
+		{
+			continue;
+		}
+
+		FAnimatedPropertyKey Key{ BoundObject, PropertyBinding.PropertyPath };
+
+		FPreAnimatedStorageGroupHandle GroupHandle  = ObjectGroupManager->MakeGroupForObject(BoundObject);
+		FPreAnimatedStorageIndex       StorageIndex = Storage.GetOrCreateStorageIndex(Key);
+
+		FPreAnimatedStateEntry Entry{ GroupHandle, FPreAnimatedStateCachedValueHandle{ StorageID, StorageIndex } };
+
+		ParentExtension->EnsureMetaData(Entry);
+
+		EPreAnimatedStorageRequirement StorageRequirement = ParentExtension->GetStorageRequirement(Entry);
+		if (!Storage.IsStorageRequirementSatisfied(StorageIndex, StorageRequirement))
+		{
+			FPreAnimatedProperty NewValue;
+
+			if (const uint16* Fast = Property->TryGet<uint16>())
+			{
+				NewValue.Binding.template Set<uint16>(*Fast);
+				FComponentTransformPropertyTraits::GetObjectPropertyValue(BoundObject, *Fast, NewValue.Data);
+			}
+			else if (const FCustomPropertyIndex* CustomIndex = Property->TryGet<FCustomPropertyIndex>())
+			{
+				const FCustomPropertyAccessor& Accessor = CustomAccessors[CustomIndex->Value];
+
+				NewValue.Binding.template Set<const FCustomPropertyAccessor*>(&Accessor);
+				FComponentTransformPropertyTraits::GetObjectPropertyValue(BoundObject, Accessor, NewValue.Data);
+			}
+			else
+			{
+				const TSharedPtr<FTrackInstancePropertyBindings>& Bindings = Property->Get<TSharedPtr<FTrackInstancePropertyBindings>>();
+
+				NewValue.Binding.template Set<TSharedPtr<FTrackInstancePropertyBindings>>(Bindings);
+				FComponentTransformPropertyTraits::GetObjectPropertyValue(BoundObject, Bindings.Get(), NewValue.Data);
+			}
+
+			Storage.AssignPreAnimatedValue(StorageIndex, StorageRequirement, MoveTemp(NewValue));
+		}
+
+		if (Params.bForcePersist)
+		{
+			Storage.ForciblyPersistStorage(StorageIndex);
+		}
+	}
 }
 
 } // namespace MovieScene
diff --git a/Engine/Source/Runtime/MovieSceneTracks/Private/PreAnimatedState/MovieScenePreAnimatedComponentTransformStorage.h b/Engine/Source/Runtime/MovieSceneTracks/Private/PreAnimatedState/MovieScenePreAnimatedComponentTransformStorage.h
--- a/Engine/Source/Runtime/MovieSceneTracks/Private/PreAnimatedState/MovieScenePreAnimatedComponentTransformStorage.h
+++ b/Engine/Source/Runtime/MovieSceneTracks/Private/PreAnimatedState/MovieScenePreAnimatedComponentTransformStorage.h
@@ -4,33 +4,33 @@
 
 #include "CoreTypes.h"
 #include "MovieSceneTracksPropertyTypes.h"
-#include "Evaluation/PreAnimatedState/MovieScenePreAnimatedObjectStorage.h"
+#include "MovieSceneTracksComponentTypes.h"
+#include "Evaluation/PreAnimatedState/MovieScenePreAnimatedPropertyStorage.h"
 
 namespace UE
 {
 namespace MovieScene
 {
 
-struct MOVIESCENETRACKS_API FComponentTransformPreAnimatedTraits
+struct FComponentTransformPreAnimatedTraits : FComponentTransformPropertyTraits
 {
 	using KeyType     = FObjectKey;
 	using StorageType = FIntermediate3DTransform;
 
-	static void CachePreAnimatedValue(UObject* InObject, FIntermediate3DTransform& OutCachedTransform);
-	static void RestorePreAnimatedValue(const FObjectKey& InKey, FIntermediate3DTransform& CachedTransform, const FRestoreStateParams& Params);
+	/** These override the functions in FComponentTransformPropertyTraits in order to wrap them with a temporary mobility assignment */
+	static void SetObjectPropertyValue(UObject* InObject, const FCustomPropertyAccessor& BaseCustomAccessor, const FIntermediate3DTransform& CachedTransform);
+	static void SetObjectPropertyValue(UObject* InObject, uint16 PropertyOffset, const FIntermediate3DTransform& CachedTransform);
+	static void SetObjectPropertyValue(UObject* InObject, FTrackInstancePropertyBindings* PropertyBindings, const FIntermediate3DTransform& CachedTransform);
 };
 
 struct MOVIESCENETRACKS_API FPreAnimatedComponentTransformStorage
-	: TPreAnimatedStateStorage_ObjectTraits<FComponentTransformPreAnimatedTraits>
+	: TPreAnimatedPropertyStorage<FComponentTransformPreAnimatedTraits>
 {
-	static TAutoRegisterPreAnimatedStorageID<FPreAnimatedComponentTransformStorage> StorageID;
+	FPreAnimatedComponentTransformStorage();
 
-	FPreAnimatedStorageID GetStorageType() const override { return StorageID; }
+	static TAutoRegisterPreAnimatedStorageID<FPreAnimatedComponentTransformStorage> StorageID;
 
-	void OnObjectReplaced(FPreAnimatedStorageIndex StorageIndex, const FObjectKey& OldObject, const FObjectKey& NewObject) override
-	{
-		ReplaceKey(StorageIndex, NewObject);
-	}
+	void CachePreAnimatedTransforms(const FCachePreAnimatedValueParams& Params, TArrayView<UObject* const> BoundObjects);
 };
 
 
diff --git a/Engine/Source/Runtime/MovieSceneTracks/Private/Systems/MovieSceneComponentAttachmentSystem.cpp b/Engine/Source/Runtime/MovieSceneTracks/Private/Systems/MovieSceneComponentAttachmentSystem.cpp
--- a/Engine/Source/Runtime/MovieSceneTracks/Private/Systems/MovieSceneComponentAttachmentSystem.cpp
+++ b/Engine/Source/Runtime/MovieSceneTracks/Private/Systems/MovieSceneComponentAttachmentSystem.cpp
@@ -288,7 +288,7 @@ void UMovieSceneComponentAttachmentSystem::SavePreAnimatedState(const FPreAnimat
 			// state regardless of the detach rules used during normal playback.
 			FCachePreAnimatedValueParams ForcePersistParams;
 			ForcePersistParams.bForcePersist = true;
-			ComponentTransformStorage->CachePreAnimatedValues(ForcePersistParams, BoundObjectArray);
+			ComponentTransformStorage->CachePreAnimatedTransforms(ForcePersistParams, BoundObjectArray);
 
 			FPreAnimatedTrackerParams AttachmentParams(Item);
 			AttachmentParams.bWantsRestoreState = false;
