Sequencer: Re-enabled motion vector simulation

#tests Ran all engine tests

diff --git a/Engine/Source/Editor/MovieSceneTools/Private/AutomatedLevelSequenceCapture.cpp b/Engine/Source/Editor/MovieSceneTools/Private/AutomatedLevelSequenceCapture.cpp
--- a/Engine/Source/Editor/MovieSceneTools/Private/AutomatedLevelSequenceCapture.cpp
+++ b/Engine/Source/Editor/MovieSceneTools/Private/AutomatedLevelSequenceCapture.cpp
@@ -20,10 +20,10 @@
 #include "MovieSceneTimeHelpers.h"
 #include "MovieSceneToolHelpers.h"
 #include "Protocols/AudioCaptureProtocol.h"
-#include "Evaluation/IMovieSceneMotionVectorSimulation.h"
-#include "Rendering/MotionVectorSimulation.h"
 #include "ShaderCompiler.h"
 #include "DistanceFieldAtlas.h"
+#include "EntitySystem/MovieSceneEntitySystemLinker.h"
+#include "Systems/MovieSceneMotionVectorSimulationSystem.h"
 
 const FName UAutomatedLevelSequenceCapture::AutomatedLevelSequenceCaptureUIName = FName(TEXT("AutomatedLevelSequenceCaptureUIInstance"));
 
@@ -102,6 +102,11 @@ UMovieSceneCinematicShotTrack* GetCinematicShotTrack(TWeakObjectPtr<ALevelSequen
 	return MovieScene->FindMasterTrack<UMovieSceneCinematicShotTrack>();
 }
 
+UMovieSceneMotionVectorSimulationSystem* FindMotionVectorSimulation(ALevelSequenceActor* LevelSequenceActor)
+{
+	return static_cast<IMovieScenePlayer*>(LevelSequenceActor->SequencePlayer)->GetEvaluationTemplate().GetEntitySystemLinker()->FindSystem<UMovieSceneMotionVectorSimulationSystem>();
+}
+
 UAutomatedLevelSequenceCapture::UAutomatedLevelSequenceCapture(const FObjectInitializer& Init)
 	: Super(Init)
 {
@@ -591,13 +596,11 @@ void UAutomatedLevelSequenceCapture::SetupFrameRange()
 				 	PlaybackStartFrame -= RemainingWarmUpFrames;
 				}
 
-				if (Actor->SequencePlayer->MotionVectorSimulation.IsValid())
+				if (UMovieSceneMotionVectorSimulationSystem* MotionVectorSim = FindMotionVectorSimulation(Actor))
 				{
-					Actor->SequencePlayer->MotionVectorSimulation->PreserveSimulatedMotion(true);
+					MotionVectorSim->PreserveSimulatedMotion(true);
 				}
 
-
-
 				// Override the movie scene's playback range
 				Actor->SequencePlayer->SetFrameRate(Settings.GetFrameRate());
 				Actor->SequencePlayer->SetFrameRange(PlaybackStartFrame.Value, (PlaybackEndFrame - PlaybackStartFrame).Value);
@@ -848,9 +851,9 @@ void UAutomatedLevelSequenceCapture::SequenceUpdated(const UMovieSceneSequencePl
 				
 				CaptureState = ELevelSequenceCaptureState::Paused;
 
-				if (Actor->SequencePlayer->MotionVectorSimulation.IsValid())
+				if (UMovieSceneMotionVectorSimulationSystem* MotionVectorSim = FindMotionVectorSimulation(Actor))
 				{
-					Actor->SequencePlayer->MotionVectorSimulation->PreserveSimulatedMotion(true);
+					MotionVectorSim->PreserveSimulatedMotion(true);
 				}
 
 				Actor->GetWorld()->GetTimerManager().SetTimer(DelayTimer, FTimerDelegate::CreateUObject(this, &UAutomatedLevelSequenceCapture::PauseFinished), DelayBeforeShotWarmUp + DelayEveryFrame, false);
@@ -859,9 +862,9 @@ void UAutomatedLevelSequenceCapture::SequenceUpdated(const UMovieSceneSequencePl
 			else if (CaptureState == ELevelSequenceCaptureState::FinishedWarmUp)
 			{
 				// If we were preserving simulated motion, now's the time to stop that since we've captured the frame that was being simulated
-				if (Actor->SequencePlayer->MotionVectorSimulation.IsValid())
+				if (UMovieSceneMotionVectorSimulationSystem* MotionVectorSim = FindMotionVectorSimulation(Actor))
 				{
-					Actor->SequencePlayer->MotionVectorSimulation->PreserveSimulatedMotion(false);
+					MotionVectorSim->PreserveSimulatedMotion(false);
 				}
 
 				// These are called each frame to allow the state machine inside the protocol to transition back to capturing
diff --git a/Engine/Source/Editor/Sequencer/Private/Sequencer.cpp b/Engine/Source/Editor/Sequencer/Private/Sequencer.cpp
--- a/Engine/Source/Editor/Sequencer/Private/Sequencer.cpp
+++ b/Engine/Source/Editor/Sequencer/Private/Sequencer.cpp
@@ -146,6 +146,7 @@
 
 #include "EntitySystem/MovieSceneEntitySystemLinker.h"
 #include "EntitySystem/MovieScenePreAnimatedStateSystem.h"
+#include "Systems/MovieSceneMotionVectorSimulationSystem.h"
 
 
 #define LOCTEXT_NAMESPACE "Sequencer"
@@ -4093,6 +4094,11 @@ void FSequencer::PossessPIEViewports(UObject* CameraObject, const EMovieSceneCam
 			{
 				CameraComponent->NotifyCameraCut();
 			}
+
+			if (UMovieSceneMotionVectorSimulationSystem* MotionVectorSim = RootTemplateInstance.GetEntitySystemLinker()->FindSystem<UMovieSceneMotionVectorSimulationSystem>())
+			{
+				MotionVectorSim->SimulateAllTransforms();
+			}
 		}
 		return;
 	}
@@ -4132,6 +4138,11 @@ void FSequencer::PossessPIEViewports(UObject* CameraObject, const EMovieSceneCam
 		PC->PlayerCameraManager->bClientSimulatingViewTarget = (CameraActor != nullptr);
 		PC->PlayerCameraManager->SetGameCameraCutThisFrame();
 	}
+
+	if (UMovieSceneMotionVectorSimulationSystem* MotionVectorSim = RootTemplateInstance.GetEntitySystemLinker()->FindSystem<UMovieSceneMotionVectorSimulationSystem>())
+	{
+		MotionVectorSim->SimulateAllTransforms();
+	}
 }
 
 TSharedPtr<class ITimeSlider> FSequencer::GetTopTimeSliderWidget() const
@@ -6044,6 +6055,11 @@ void FSequencer::UpdatePreviewLevelViewportClientFromCameraCut(FLevelEditorViewp
 	if (bCameraHasBeenCut)
 	{
 		InViewportClient.SetIsCameraCut();
+
+		if (UMovieSceneMotionVectorSimulationSystem* MotionVectorSim = RootTemplateInstance.GetEntitySystemLinker()->FindSystem<UMovieSceneMotionVectorSimulationSystem>())
+		{
+			MotionVectorSim->SimulateAllTransforms();
+		}
 	}
 
 	// Set the actor lock.
diff --git a/Engine/Source/Runtime/LevelSequence/Private/LevelSequencePlayer.cpp b/Engine/Source/Runtime/LevelSequence/Private/LevelSequencePlayer.cpp
--- a/Engine/Source/Runtime/LevelSequence/Private/LevelSequencePlayer.cpp
+++ b/Engine/Source/Runtime/LevelSequence/Private/LevelSequencePlayer.cpp
@@ -19,6 +19,8 @@
 #include "Engine/LocalPlayer.h"
 #include "Tracks/MovieSceneCinematicShotTrack.h"
 #include "Sections/MovieSceneCinematicShotSection.h"
+#include "Systems/MovieSceneMotionVectorSimulationSystem.h"
+#include "EntitySystem/MovieSceneEntitySystemLinker.h"
 #include "LevelSequenceActor.h"
 #include "Modules/ModuleManager.h"
 #include "LevelUtils.h"
@@ -270,6 +272,11 @@ void ULevelSequencePlayer::UpdateCameraCut(UObject* CameraObject, const EMovieSc
 			{
 				CameraComponent->NotifyCameraCut();
 			}
+
+			if (UMovieSceneMotionVectorSimulationSystem* MotionVectorSim = RootTemplateInstance.GetEntitySystemLinker()->FindSystem<UMovieSceneMotionVectorSimulationSystem>())
+			{
+				MotionVectorSim->SimulateAllTransforms();
+			}
 		}
 		return;
 	}
@@ -381,6 +388,11 @@ void ULevelSequencePlayer::UpdateCameraCut(UObject* CameraObject, const EMovieSc
 
 	if (bIsStraightCut)
 	{
+		if (UMovieSceneMotionVectorSimulationSystem* MotionVectorSim = RootTemplateInstance.GetEntitySystemLinker()->FindSystem<UMovieSceneMotionVectorSimulationSystem>())
+		{
+			MotionVectorSim->SimulateAllTransforms();
+		}
+
 		if (OnCameraCut.IsBound())
 		{
 			OnCameraCut.Broadcast(CameraComponent);
diff --git a/Engine/Source/Runtime/MovieScene/Private/Evaluation/MovieSceneMotionVectorSimulation.cpp b/Engine/Source/Runtime/MovieScene/Private/Evaluation/MovieSceneMotionVectorSimulation.cpp
deleted file mode 100644
diff --git a/Engine/Source/Runtime/MovieScene/Private/Evaluation/MovieSceneMotionVectorSimulation.h b/Engine/Source/Runtime/MovieScene/Private/Evaluation/MovieSceneMotionVectorSimulation.h
deleted file mode 100644
diff --git a/Engine/Source/Runtime/MovieScene/Public/Evaluation/IMovieSceneMotionVectorSimulation.h b/Engine/Source/Runtime/MovieScene/Public/Evaluation/IMovieSceneMotionVectorSimulation.h
deleted file mode 100644
diff --git a/Engine/Source/Runtime/MovieScene/Public/IMovieScenePlayer.h b/Engine/Source/Runtime/MovieScene/Public/IMovieScenePlayer.h
--- a/Engine/Source/Runtime/MovieScene/Public/IMovieScenePlayer.h
+++ b/Engine/Source/Runtime/MovieScene/Public/IMovieScenePlayer.h
@@ -12,7 +12,6 @@
 #include "Containers/ArrayView.h"
 #include "Evaluation/MovieSceneEvaluationState.h"
 #include "Misc/InlineValue.h"
-#include "Evaluation/IMovieSceneMotionVectorSimulation.h"
 #include "Evaluation/MovieSceneEvaluationOperand.h"
 #include "Generators/MovieSceneEasingCurves.h"
 
@@ -334,9 +333,6 @@ public:
 	/** Container that stores any per-animated state tokens  */
 	FMovieScenePreAnimatedState PreAnimatedState;
 
-	/** Motion Vector Simulation */
-	TInlineValue<IMovieSceneMotionVectorSimulation> MotionVectorSimulation;
-
 	/** List of binding overrides to use for the sequence */
 	TMap<FMovieSceneEvaluationOperand, FMovieSceneEvaluationOperand> BindingOverrides;
 
diff --git a/Engine/Source/Runtime/MovieSceneTracks/Private/Evaluation/MovieSceneSkeletalAnimationTemplate.cpp b/Engine/Source/Runtime/MovieSceneTracks/Private/Evaluation/MovieSceneSkeletalAnimationTemplate.cpp
--- a/Engine/Source/Runtime/MovieSceneTracks/Private/Evaluation/MovieSceneSkeletalAnimationTemplate.cpp
+++ b/Engine/Source/Runtime/MovieSceneTracks/Private/Evaluation/MovieSceneSkeletalAnimationTemplate.cpp
@@ -12,7 +12,8 @@
 #include "IMovieScenePlayer.h"
 #include "UObject/ObjectKey.h"
 #include "Rendering/MotionVectorSimulation.h"
-#include "Evaluation/IMovieSceneMotionVectorSimulation.h"
+#include "Systems/MovieSceneMotionVectorSimulationSystem.h"
+#include "EntitySystem/MovieSceneEntitySystemLinker.h"
 #include "UObject/StrongObjectPtr.h"
 #include "SkeletalMeshRestoreState.h"
 #include "AnimSequencerInstanceProxy.h"
@@ -331,16 +332,17 @@ namespace MovieScene
 				}
 			}
 
-			if (InFinalValue.SimulatedAnimations.Num() != 0 && Player.MotionVectorSimulation.IsValid())
+			if (InFinalValue.SimulatedAnimations.Num() != 0)
 			{
-				ApplyAnimations(PersistentData, Player, SkeletalMeshComponent, InFinalValue.SimulatedAnimations, DeltaTime, bPreviewPlayback, bFireNotifies, bResetDynamics);
-
-				SkeletalMeshComponent->TickAnimation(0.f, false);
-				SkeletalMeshComponent->RefreshBoneTransforms();
-				SkeletalMeshComponent->FinalizeBoneTransform();
-				SkeletalMeshComponent->ForceMotionVector();
+				UMovieSceneMotionVectorSimulationSystem* MotionVectorSim = Player.GetEvaluationTemplate().GetEntitySystemLinker()->FindSystem<UMovieSceneMotionVectorSimulationSystem>();
+				if (MotionVectorSim && MotionVectorSim->IsSimulationEnabled())
+				{
+					ApplyAnimations(PersistentData, Player, SkeletalMeshComponent, InFinalValue.SimulatedAnimations, DeltaTime, bPreviewPlayback, bFireNotifies, bResetDynamics);
+					SkeletalMeshComponent->TickAnimation(0.f, false);
+					SkeletalMeshComponent->ForceMotionVector();
 
-				SimulateMotionVectors(PersistentData, SkeletalMeshComponent, Player);
+					SimulateMotionVectors(PersistentData, SkeletalMeshComponent, MotionVectorSim);
+				}
 			}
 
 			ApplyAnimations(PersistentData, Player, SkeletalMeshComponent, InFinalValue.AllAnimations, DeltaTime, bPreviewPlayback, bFireNotifies, bResetDynamics);
@@ -389,7 +391,7 @@ namespace MovieScene
 			return nullptr;
 		}
 
-		void SimulateMotionVectors(FPersistentEvaluationData& PersistentData, USkeletalMeshComponent* SkeletalMeshComponent, IMovieScenePlayer& Player)
+		void SimulateMotionVectors(FPersistentEvaluationData& PersistentData, USkeletalMeshComponent* SkeletalMeshComponent, UMovieSceneMotionVectorSimulationSystem* MotionVectorSim)
 		{
 			for (USceneComponent* Child : SkeletalMeshComponent->GetAttachChildren())
 			{
@@ -397,7 +399,7 @@ namespace MovieScene
 				if (SocketName != NAME_None)
 				{
 					FTransform SocketTransform = SkeletalMeshComponent->GetSocketTransform(SocketName, RTS_Component);
-					Player.MotionVectorSimulation->Add(SkeletalMeshComponent, SocketTransform, SocketName);
+					MotionVectorSim->AddSimulatedTransform(SkeletalMeshComponent, SocketTransform, SocketName);
 				}
 			}
 		}
@@ -602,9 +604,9 @@ void FMovieSceneSkeletalAnimationSectionTemplate::Evaluate(const FMovieSceneEval
 		);
 		ExecutionTokens.BlendToken(ActuatorTypeID, TBlendableToken<UE::MovieScene::FBlendedAnimation>(AnimParams, BlendType.Get(), 1.f));
 
-		if (IMovieSceneMotionVectorSimulation::IsEnabled(PersistentData, Context))
+		if (FMotionVectorSimulation::IsEnabled())
 		{
-			FFrameTime SimulatedTime = IMovieSceneMotionVectorSimulation::GetSimulationTime(Context);
+			FFrameTime SimulatedTime = UE::MovieScene::GetSimulatedMotionVectorTime(Context);
 
 			// Calculate the time at which to evaluate the animation
 			float CurrentEvalTime = Params.MapTimeToAnimation(Context.GetTime(), Context.GetFrameRate());
diff --git a/Engine/Source/Runtime/MovieSceneTracks/Private/Systems/MovieSceneMotionVectorSimulationSystem.cpp b/Engine/Source/Runtime/MovieSceneTracks/Private/Systems/MovieSceneMotionVectorSimulationSystem.cpp
new file mode 100644
--- /dev/null
+++ b/Engine/Source/Runtime/MovieSceneTracks/Private/Systems/MovieSceneMotionVectorSimulationSystem.cpp
@@ -0,0 +1,341 @@
+// Copyright Epic Games, Inc. All Rights Reserved.
+
+#include "Systems/MovieSceneMotionVectorSimulationSystem.h"
+#include "Systems/MovieSceneTransformOriginSystem.h"
+#include "EntitySystem/MovieSceneEvalTimeSystem.h"
+#include "EntitySystem/MovieSceneEntitySystemLinker.h"
+#include "EntitySystem/MovieSceneEntitySystemRunner.h"
+#include "EntitySystem/MovieSceneInstanceRegistry.h"
+#include "EntitySystem/BuiltInComponentTypes.h"
+#include "EntitySystem/MovieSceneEntityMutations.h"
+#include "Tracks/IMovieSceneTransformOrigin.h"
+#include "MovieSceneTracksComponentTypes.h"
+
+#include "Systems/FloatChannelEvaluatorSystem.h"
+#include "Systems/MovieScenePiecewiseFloatBlenderSystem.h"
+#include "Systems/MovieSceneComponentTransformSystem.h"
+
+#include "Rendering/MotionVectorSimulation.h"
+
+#include "IMovieScenePlayer.h"
+#include "IMovieScenePlaybackClient.h"
+
+namespace UE
+{
+namespace MovieScene
+{
+
+
+FFrameTime GetSimulatedMotionVectorTime(const FMovieSceneContext& Context)
+{
+	FFrameTime DeltaTime     = FMath::Max(Context.GetDelta(), 0.0041666666666667 * Context.GetFrameRate());
+	FFrameTime SimulatedTime = Context.GetOffsetTime(DeltaTime);
+
+	return SimulatedTime;
+}
+
+
+} // namespace MovieScene
+} // namespace UE
+
+
+UMovieSceneMotionVectorSimulationSystem::UMovieSceneMotionVectorSimulationSystem(const FObjectInitializer& ObjInit)
+	: Super(ObjInit)
+{
+	using namespace UE::MovieScene;
+
+	Phase = ESystemPhase::Finalization;
+}
+
+bool UMovieSceneMotionVectorSimulationSystem::IsRelevantImpl(UMovieSceneEntitySystemLinker* InLinker) const
+{
+	using namespace UE::MovieScene;
+
+	return FMotionVectorSimulation::IsEnabled();
+}
+
+void UMovieSceneMotionVectorSimulationSystem::OnRun(FSystemTaskPrerequisites& InPrerequisites, FSystemSubsequentTasks& Subsequents)
+{
+	using namespace UE::MovieScene;
+
+	FMovieSceneEntitySystemRunner* Runner = Linker->GetActiveRunner();
+	if (Runner->GetCurrentPhase() == ESystemPhase::Finalization && bSimulationEnabled)
+	{
+		Runner->GetQueuedEventTriggers().AddLambda([this] { this->OnPostEvaluation(); });
+	}
+
+}
+
+void UMovieSceneMotionVectorSimulationSystem::OnPostEvaluation()
+{
+	using namespace UE::MovieScene;
+
+	if (bSimulateTransformsRequested)
+	{
+		ComputeSimulatedMotion();
+	}
+
+	PropagateMotionToComponents();
+}
+
+void UMovieSceneMotionVectorSimulationSystem::ComputeSimulatedMotion()
+{
+	using namespace UE::MovieScene;
+
+	// ---------------------------------------------------------------------------------------
+	// Simulated transforms are computed by disabling everything in the ecs that is not a
+	// component transform, fudging the eval times, then harvesting the results.
+
+	FBuiltInComponentTypes*          BuiltInComponents = FBuiltInComponentTypes::Get();
+	FMovieSceneTracksComponentTypes* TracksComponents  = FMovieSceneTracksComponentTypes::Get();
+
+	// --------------------------------------------------------------------------------------------------------------------------------------------
+	// Disable everything and re-execute the evaluation phase
+	Disable();
+	if (UMovieSceneComponentTransformSystem* ComponentTransformSystem = Linker->FindSystem<UMovieSceneComponentTransformSystem>())
+	{
+		ComponentTransformSystem->Disable();
+	}
+	if (UMovieSceneEvalTimeSystem* EvalTime = Linker->FindSystem<UMovieSceneEvalTimeSystem>())
+	{
+		EvalTime->Disable();
+	}
+
+	FEntityComponentFilter* GlobalFilter     = &Linker->EntityManager.ModifyGlobalIterationFilter();
+	FEntityComponentFilter  GlobalFilterCopy = *GlobalFilter;
+
+	// Only allow visitation of component transforms for this procedure
+	GlobalFilter->All({ TracksComponents->ComponentTransform.PropertyTag });
+
+	// --------------------------------------------------------------------------------------------------------------------------------------------
+	// Fudge frame times based on GetSimulatedMotionVectorTime
+	{
+		TArray<FFrameTime> FrameTimesByInstance;
+
+		FInstanceRegistry* InstanceRegistry = Linker->GetInstanceRegistry();
+
+		FrameTimesByInstance.SetNum(InstanceRegistry->GetSparseInstances().GetMaxIndex());
+		for (auto It = InstanceRegistry->GetSparseInstances().CreateConstIterator(); It; ++It)
+		{
+			FrameTimesByInstance[It.GetIndex()] = GetSimulatedMotionVectorTime(It->GetContext());
+		}
+
+		FEntityTaskBuilder()
+		.Read(BuiltInComponents->InstanceHandle)
+		.Write(BuiltInComponents->EvalTime)
+		.FilterNone({ FBuiltInComponentTypes::Get()->Tags.FixedTime })
+		.Iterate_PerEntity(&Linker->EntityManager,
+			[&FrameTimesByInstance](UE::MovieScene::FInstanceHandle InstanceHandle, FFrameTime& EvalTime)
+			{
+				EvalTime = FrameTimesByInstance[InstanceHandle.InstanceID];
+			}
+		);
+	}
+
+	// --------------------------------------------------------------------------------------------------------------------------------------------
+	// Re-execute the evaluation phase to flush the new transforms
+	{
+		Linker->EntityManager.LockDown();
+
+		FGraphEventArray AllTasks;
+		Linker->SystemGraph.ExecutePhase(ESystemPhase::Evaluation, Linker, AllTasks);
+
+		if (AllTasks.Num() != 0)
+		{
+			FTaskGraphInterface::Get().WaitUntilTasksComplete(AllTasks, ENamedThreads::GameThread_Local);
+		}
+
+		Linker->EntityManager.ReleaseLockDown();
+	}
+
+	// --------------------------------------------------------------------------------------------------------------------------------------------
+	// Harvest results
+	FEntityTaskBuilder()
+	.Read(BuiltInComponents->BoundObject)
+	.ReadAnyOf(
+		BuiltInComponents->FloatResult[0], BuiltInComponents->FloatResult[1], BuiltInComponents->FloatResult[2],
+		BuiltInComponents->FloatResult[3], BuiltInComponents->FloatResult[4], BuiltInComponents->FloatResult[5],
+		BuiltInComponents->FloatResult[6], BuiltInComponents->FloatResult[7], BuiltInComponents->FloatResult[8])
+	.FilterAll({ TracksComponents->ComponentTransform.PropertyTag })
+	.FilterAny({ BuiltInComponents->CustomPropertyIndex, BuiltInComponents->FastPropertyOffset, BuiltInComponents->SlowProperty })
+	.FilterAny({ 
+		// must have at least one float result component
+		BuiltInComponents->FloatResult[0], BuiltInComponents->FloatResult[1], BuiltInComponents->FloatResult[2],
+		BuiltInComponents->FloatResult[3], BuiltInComponents->FloatResult[4], BuiltInComponents->FloatResult[5],
+		BuiltInComponents->FloatResult[6], BuiltInComponents->FloatResult[7], BuiltInComponents->FloatResult[8] })
+	.Iterate_PerEntity(&Linker->EntityManager, [this](UObject* BoundObject, const float* T_X, const float* T_Y, const float* T_Z, const float* R_X, const float* R_Y, const float* R_Z, const float* S_X, const float* S_Y, const float* S_Z)
+	{
+		if (USceneComponent* SceneComponent = Cast<USceneComponent>(BoundObject))
+		{
+			FTransform SimulatedTransform = SceneComponent->GetRelativeTransform();
+
+			FVector    SimulatedLocation  = SimulatedTransform.GetLocation();
+			FRotator   SimulatedRotation  = SimulatedTransform.Rotator();
+			FVector    SimulatedScale     = SimulatedTransform.GetScale3D();
+
+			if (T_X) { SimulatedLocation.X     = *T_X; }
+			if (T_Y) { SimulatedLocation.Y     = *T_Y; }
+			if (T_Z) { SimulatedLocation.Z     = *T_Z; }
+			if (R_X) { SimulatedRotation.Roll  = *R_X; }
+			if (R_Y) { SimulatedRotation.Pitch = *R_Y; }
+			if (R_Z) { SimulatedRotation.Yaw   = *R_Z; }
+			if (S_X) { SimulatedScale.X        = *S_X; }
+			if (S_Y) { SimulatedScale.Y        = *S_Y; }
+			if (S_Z) { SimulatedScale.Z        = *S_Z; }
+
+			this->AddSimulatedTransform(SceneComponent, FTransform(SimulatedRotation, SimulatedLocation, SimulatedScale), NAME_None);
+		}
+	});
+
+	// Reset the global filter to what it was before
+	*GlobalFilter = MoveTemp(GlobalFilterCopy);
+
+	// --------------------------------------------------------------------------------------------------------------------------------------------
+	// Re-enable systems and reset the simulation components
+	if (UMovieSceneComponentTransformSystem* ComponentTransformSystem = Linker->FindSystem<UMovieSceneComponentTransformSystem>())
+	{
+		ComponentTransformSystem->Enable();
+	}
+	if (UMovieSceneEvalTimeSystem* EvalTime = Linker->FindSystem<UMovieSceneEvalTimeSystem>())
+	{
+		EvalTime->Enable();
+	}
+	Enable();
+
+	bSimulateTransformsRequested = false;
+}
+
+void UMovieSceneMotionVectorSimulationSystem::EnableThisFrame()
+{
+	bSimulationEnabled = true;
+}
+
+void UMovieSceneMotionVectorSimulationSystem::SimulateAllTransforms()
+{
+	bSimulationEnabled = true;
+	bSimulateTransformsRequested = true;
+}
+
+void UMovieSceneMotionVectorSimulationSystem::PreserveSimulatedMotion(bool bShouldPreserveTransforms)
+{
+	bPreserveTransforms = bShouldPreserveTransforms;
+}
+
+void UMovieSceneMotionVectorSimulationSystem::AddSimulatedTransform(USceneComponent* Component, const FTransform& SimulatedTransform, FName SocketName)
+{
+	TransformData.Add(Component, FSimulatedTransform(SimulatedTransform, SocketName));
+}
+
+void UMovieSceneMotionVectorSimulationSystem::PropagateMotionToComponents()
+{
+	TSet<USceneComponent*> RootComponents;
+
+	for (auto It = TransformData.CreateIterator(); It; ++It)
+	{
+		// If this is a socket transform, we want to add the component as a whole, to ensure that anything attached to this socket gets simulated correctly
+
+		USceneComponent* Component = Cast<USceneComponent>(It.Key().ResolveObjectPtr());
+		if (!Component || HavePreviousTransformForParent(Component))
+		{
+			continue;
+		}
+
+		RootComponents.Add(Component);
+	}
+
+	for (USceneComponent* Component : RootComponents)
+	{
+		FTransform ParentToWorld = FTransform::Identity;
+
+		USceneComponent* ParentComp = Component->GetAttachParent();
+		FName AttachSocket = Component->GetAttachSocketName();
+		if (ParentComp)
+		{
+			FTransform ParentTransform = ParentComp->GetSocketTransform(AttachSocket, RTS_World);
+			if (!Component->IsUsingAbsoluteLocation())
+			{
+				ParentToWorld.SetTranslation(ParentTransform.GetTranslation());
+			}
+			if (!Component->IsUsingAbsoluteRotation())
+			{
+				ParentToWorld.SetRotation(ParentTransform.GetRotation());
+			}
+			if (!Component->IsUsingAbsoluteScale())
+			{
+				ParentToWorld.SetScale3D(ParentTransform.GetScale3D());
+			}
+		}
+
+		ApplySimulatedTransforms(Component, GetTransform(Component) * ParentToWorld);
+	}
+
+	if (!bPreserveTransforms)
+	{
+		TransformData.Reset();
+		bSimulationEnabled = false;
+	}
+}
+
+FTransform UMovieSceneMotionVectorSimulationSystem::GetTransform(USceneComponent* Component)
+{
+	FObjectKey Key(Component);
+	for (auto It = TransformData.CreateConstKeyIterator(Key); It; ++It)
+	{
+		if (It.Value().SocketName == NAME_None)
+		{
+			return It.Value().Transform;
+		}
+	}
+
+	return Component->GetRelativeTransform();
+}
+
+FTransform UMovieSceneMotionVectorSimulationSystem::GetSocketTransform(USceneComponent* Component, FName SocketName)
+{
+	FObjectKey Key(Component);
+	for (auto It = TransformData.CreateConstKeyIterator(Key); It; ++It)
+	{
+		if (It.Value().SocketName == SocketName)
+		{
+			return It.Value().Transform;
+		}
+	}
+
+	return Component->GetSocketTransform(SocketName, RTS_Component);
+}
+
+bool UMovieSceneMotionVectorSimulationSystem::HavePreviousTransformForParent(USceneComponent* InComponent) const
+{
+	USceneComponent* Parent = InComponent->GetAttachParent();
+	return Parent && (TransformData.Contains(Parent) || HavePreviousTransformForParent(Parent));
+}
+
+void UMovieSceneMotionVectorSimulationSystem::ApplySimulatedTransforms(USceneComponent* InComponent, const FTransform& InPreviousTransform)
+{
+	check(InComponent);
+	FMotionVectorSimulation::Get().SetPreviousTransform(InComponent, InPreviousTransform);
+
+	for (USceneComponent* Child : InComponent->GetAttachChildren())
+	{
+		FName AttachSocketName = Child->GetAttachSocketName();
+
+		FTransform SocketTransform = AttachSocketName == NAME_None ? FTransform::Identity : GetSocketTransform(InComponent, AttachSocketName);
+		FTransform ParentToWorld = SocketTransform * InPreviousTransform;
+
+		if (Child->IsUsingAbsoluteLocation())
+		{
+			ParentToWorld.SetTranslation(FVector::ZeroVector);
+		}
+		if (Child->IsUsingAbsoluteRotation())
+		{
+			ParentToWorld.SetRotation(FQuat::Identity);
+		}
+		if (Child->IsUsingAbsoluteScale())
+		{
+			ParentToWorld.SetScale3D(FVector(1.f, 1.f, 1.f));
+		}
+
+		FTransform ChildTransform = GetTransform(Child);
+		ApplySimulatedTransforms(Child, ChildTransform * ParentToWorld);
+	}
+}
diff --git a/Engine/Source/Runtime/MovieSceneTracks/Private/TrackInstances/MovieSceneCameraCutTrackInstance.cpp b/Engine/Source/Runtime/MovieSceneTracks/Private/TrackInstances/MovieSceneCameraCutTrackInstance.cpp
--- a/Engine/Source/Runtime/MovieSceneTracks/Private/TrackInstances/MovieSceneCameraCutTrackInstance.cpp
+++ b/Engine/Source/Runtime/MovieSceneTracks/Private/TrackInstances/MovieSceneCameraCutTrackInstance.cpp
@@ -2,7 +2,6 @@
 
 #include "TrackInstances/MovieSceneCameraCutTrackInstance.h"
 #include "ContentStreaming.h"
-#include "Evaluation/IMovieSceneMotionVectorSimulation.h"
 #include "Evaluation/MovieSceneEvaluation.h"
 #include "GameFramework/Actor.h"
 #include "Generators/MovieSceneEasingCurves.h"
@@ -173,8 +172,6 @@ namespace MovieScene
 				CameraCutParams.UnlockIfCameraObject = CameraCutCache.LastLockedCamera.Get();
 				Player.UpdateCameraCut(CameraActor, CameraCutParams);
 				CameraCutCache.LastLockedCamera = CameraActor;
-				// TODO-ludovic
-				//IMovieSceneMotionVectorSimulation::EnableThisFrame(PersistentData);
 				return true;
 			}
 			else if (CameraActor || CameraCutParams.BlendTime > 0.f)
diff --git a/Engine/Source/Runtime/MovieSceneTracks/Public/Systems/MovieSceneMotionVectorSimulationSystem.h b/Engine/Source/Runtime/MovieSceneTracks/Public/Systems/MovieSceneMotionVectorSimulationSystem.h
new file mode 100644
--- /dev/null
+++ b/Engine/Source/Runtime/MovieSceneTracks/Public/Systems/MovieSceneMotionVectorSimulationSystem.h
@@ -0,0 +1,96 @@
+// Copyright Epic Games, Inc. All Rights Reserved.
+
+#pragma once
+
+#include "EntitySystem/MovieSceneEntitySystem.h"
+#include "Math/Transform.h"
+#include "Containers/Map.h"
+#include "UObject/ObjectKey.h"
+#include "UObject/NameTypes.h"
+
+#include "MovieSceneMotionVectorSimulationSystem.generated.h"
+
+struct FMovieSceneAnimTypeID;
+
+namespace UE
+{
+namespace MovieScene
+{
+
+
+FFrameTime GetSimulatedMotionVectorTime(const FMovieSceneContext& Context);
+
+
+} // namespace MovieScene
+} // namespace UE
+
+UCLASS()
+class MOVIESCENETRACKS_API UMovieSceneMotionVectorSimulationSystem : public UMovieSceneEntitySystem
+{
+public:
+
+	GENERATED_BODY()
+
+	UMovieSceneMotionVectorSimulationSystem(const FObjectInitializer& ObjInit);
+
+	void EnableThisFrame();
+	void SimulateAllTransforms();
+	bool IsSimulationEnabled() const
+	{
+		return bSimulationEnabled;
+	}
+
+	void PreserveSimulatedMotion(bool bShouldPreserveTransforms);
+
+	void AddSimulatedTransform(USceneComponent* Component, const FTransform& SimulatedTransform, FName SocketName);
+
+private:
+
+	virtual bool IsRelevantImpl(UMovieSceneEntitySystemLinker* InLinker) const override;
+	virtual void OnRun(FSystemTaskPrerequisites& InPrerequisites, FSystemSubsequentTasks& Subsequents) override final;
+
+	void OnPostEvaluation();
+
+	void ComputeSimulatedMotion();
+
+private:
+
+	void PropagateMotionToComponents();
+
+	FTransform GetTransform(USceneComponent* Component);
+
+	FTransform GetSocketTransform(USceneComponent* Component, FName SocketName);
+
+	bool HavePreviousTransformForParent(USceneComponent* InComponent) const;
+
+	void ApplySimulatedTransforms(USceneComponent* InComponent, const FTransform& InPreviousTransform);
+
+private:
+
+	/** Transform data relating to a specific object or socket */
+	struct FSimulatedTransform
+	{
+		FSimulatedTransform(const FTransform& InTransform, FName InSocketName = NAME_None)
+			: Transform(InTransform), SocketName(InSocketName)
+		{}
+
+		/** The transform for the object */
+		FTransform Transform;
+
+		/** The socket name to which this relates */
+		FName SocketName;
+	};
+
+	/** Map of object -> previous transform data */
+	TMultiMap<FObjectKey, FSimulatedTransform> TransformData;
+
+	/** Whether to reset TransformData at the end of the frame or not */
+	bool bPreserveTransforms = false;
+
+	bool bSimulationEnabled = false;
+
+	bool bSimulateTransformsRequested = false;
+};
+
+
+
