Fixed pedantic static analysis warning

diff --git a/Engine/Source/Runtime/MovieSceneTracks/Private/Systems/MovieSceneMotionVectorSimulationSystem.cpp b/Engine/Source/Runtime/MovieSceneTracks/Private/Systems/MovieSceneMotionVectorSimulationSystem.cpp
--- a/Engine/Source/Runtime/MovieSceneTracks/Private/Systems/MovieSceneMotionVectorSimulationSystem.cpp
+++ b/Engine/Source/Runtime/MovieSceneTracks/Private/Systems/MovieSceneMotionVectorSimulationSystem.cpp
@@ -319,7 +319,7 @@ void UMovieSceneMotionVectorSimulationSystem::ApplySimulatedTransforms(USceneCom
 	{
 		FName AttachSocketName = Child->GetAttachSocketName();
 
-		FTransform SocketTransform = AttachSocketName == NAME_None ? FTransform::Identity : GetSocketTransform(InComponent, AttachSocketName);
+		FTransform SocketTransform = (AttachSocketName == NAME_None) ? FTransform::Identity : GetSocketTransform(InComponent, AttachSocketName);
 		FTransform ParentToWorld = SocketTransform * InPreviousTransform;
 
 		if (Child->IsUsingAbsoluteLocation())
