Sequencer: Only stop camera shake if it was started.

We could try to restore pre-animated state on a stopped shake if the shake had some duration and the section was longer that this duration, without any looping.

diff --git a/Engine/Source/Runtime/MovieSceneTracks/Private/Evaluation/MovieSceneCameraAnimTemplate.cpp b/Engine/Source/Runtime/MovieSceneTracks/Private/Evaluation/MovieSceneCameraAnimTemplate.cpp
--- a/Engine/Source/Runtime/MovieSceneTracks/Private/Evaluation/MovieSceneCameraAnimTemplate.cpp
+++ b/Engine/Source/Runtime/MovieSceneTracks/Private/Evaluation/MovieSceneCameraAnimTemplate.cpp
@@ -488,7 +488,10 @@ struct FPreAnimatedCameraShakeTokenProducer : IMovieScenePreAnimatedTokenProduce
 			virtual void RestoreState(UObject& InObject, const UE::MovieScene::FRestoreStateParams& Params) override
 			{
 				UCameraShakeBase* CameraShake = CastChecked<UCameraShakeBase>(&InObject);
-				CameraShake->StopShake(true);
+				if (!CameraShake->IsFinished())	
+				{
+					CameraShake->StopShake(true);
+				}
 				CameraShake->TeardownShake();
 			}
 		};
