Sequencer: Fixed shot tracks not compiling overlapping segments correctly

 -The AddIfEmpty function was implemented incorrectly as an operator, rather than a predicate meaning it would nearly always be empty for new segments
- Removed pre/postroll handling from the cinematic shot track evaluation field since preroll is always handled separately (and should always be added to the field)

diff --git a/Engine/Source/Runtime/MovieScene/Public/Evaluation/MovieSceneEvaluationTree.h b/Engine/Source/Runtime/MovieScene/Public/Evaluation/MovieSceneEvaluationTree.h
--- a/Engine/Source/Runtime/MovieScene/Public/Evaluation/MovieSceneEvaluationTree.h
+++ b/Engine/Source/Runtime/MovieScene/Public/Evaluation/MovieSceneEvaluationTree.h
@@ -638,7 +638,9 @@ struct TMovieSceneEvaluationTree : FMovieSceneEvaluationTree
 	 */
 	void AddIfEmpty(TRange<FFrameNumber> InTimeRange, DataType InData)
 	{
-		AddTimeRange(InTimeRange, FAddIfEmptyOperator(*this, MoveTemp(InData)), FMovieSceneEvaluationTreeNodeHandle::Root(), nullptr);
+		AddSelective(InTimeRange, MoveTemp(InData), [this](FMovieSceneEvaluationTreeNodeHandle InNodeHandle){
+			return this->GetAllData(InNodeHandle).IsValid() == false;
+		});
 	}
 
 	/**
@@ -649,7 +651,9 @@ struct TMovieSceneEvaluationTree : FMovieSceneEvaluationTree
 	 */
 	void AddIfEmptySelective(TRange<FFrameNumber> InTimeRange, DataType InData, TFunctionRef<bool(FMovieSceneEvaluationTreeNodeHandle)> Predicate)
 	{
-		AddTimeRange(InTimeRange, FAddIfEmptyOperator(*this, MoveTemp(InData)), FMovieSceneEvaluationTreeNodeHandle::Root(), &Predicate);
+		AddSelective(InTimeRange, MoveTemp(InData), [this, &Predicate](FMovieSceneEvaluationTreeNodeHandle InNodeHandle){
+			return this->GetAllData(InNodeHandle).IsValid() && Predicate(InNodeHandle);
+		});
 	}
 
 	/**
@@ -764,23 +768,6 @@ private:
 		DataType DataToInsert;
 	};
 
-	/** Operator that adds data to nodes only the first time (i.e. if the node is empty) */
-	struct FAddIfEmptyOperator : IMovieSceneEvaluationTreeNodeOperator
-	{
-		FAddIfEmptyOperator(TMovieSceneEvaluationTree<DataType>& InTree, DataType&& InDataToInsert) : Tree(InTree), DataToInsert(InDataToInsert) {}
-
-		virtual void operator()(FMovieSceneEvaluationTreeNode& InNode) const override
-		{
-			if (!InNode.DataID.IsValid())
-			{
-				InNode.DataID = Tree.Data.AllocateEntry(1);
-				Tree.Data.Add(InNode.DataID, CopyTemp(DataToInsert));
-			}
-		}
-		TMovieSceneEvaluationTree<DataType>& Tree;
-		DataType DataToInsert;
-	};
-
 	/** Operator that adds data to nodes */
 	struct FAddOperator : IMovieSceneEvaluationTreeNodeOperator
 	{
diff --git a/Engine/Source/Runtime/MovieSceneTracks/Private/Tracks/MovieSceneCinematicShotTrack.cpp b/Engine/Source/Runtime/MovieSceneTracks/Private/Tracks/MovieSceneCinematicShotTrack.cpp
--- a/Engine/Source/Runtime/MovieSceneTracks/Private/Tracks/MovieSceneCinematicShotTrack.cpp
+++ b/Engine/Source/Runtime/MovieSceneTracks/Private/Tracks/MovieSceneCinematicShotTrack.cpp
@@ -81,39 +81,19 @@ namespace UE
 {
 namespace MovieScene
 {
-	enum class ECinematicShotSectionSortFlags
-	{
-		None = 0,
-		PreRoll = 1 << 0,
-		PostRoll = 1 << 1
-	};
-	ENUM_CLASS_FLAGS(ECinematicShotSectionSortFlags);
-
 	struct FCinematicShotSectionSortData
 	{
 		int32 Row;
 		int32 OverlapPriority;
 		int32 SectionIndex;
-		TRangeBound<FFrameNumber> LowerBound;
-		ECinematicShotSectionSortFlags Flags = ECinematicShotSectionSortFlags::None;
 
 		friend bool operator<(const FCinematicShotSectionSortData& A, const FCinematicShotSectionSortData& B)
 		{
-			const bool PrePostRollA = EnumHasAnyFlags(A.Flags, ECinematicShotSectionSortFlags::PreRoll | ECinematicShotSectionSortFlags::PostRoll);
-			const bool PrePostRollB = EnumHasAnyFlags(B.Flags, ECinematicShotSectionSortFlags::PreRoll | ECinematicShotSectionSortFlags::PostRoll);
-
-			if (PrePostRollA != PrePostRollB)
+			if (A.Row != B.Row)
 			{
-				return PrePostRollA;
-			}
-			else if (PrePostRollA)
-			{
-				return false;
-			}
-			else if (A.OverlapPriority == B.OverlapPriority)
-			{
-				return TRangeBound<FFrameNumber>::MaxLower(A.LowerBound, B.LowerBound) == A.LowerBound;
+				return A.Row < B.Row;
 			}
+
 			return A.OverlapPriority > B.OverlapPriority;
 		}
 	};
@@ -135,17 +115,8 @@ bool UMovieSceneCinematicShotTrack::PopulateEvaluationTree(TMovieSceneEvaluation
 			const TRange<FFrameNumber> SectionRange = Section->GetRange();
 			if (!SectionRange.IsEmpty())
 			{
-				FCinematicShotSectionSortData SectionData{ 
-					Section->GetRowIndex(), Section->GetOverlapPriority(), SectionIndex, Section->GetRange().GetLowerBound() };
-				if (!SectionRange.GetLowerBound().IsOpen() && Section->GetPreRollFrames() > 0)
-				{
-					SectionData.Flags |= ECinematicShotSectionSortFlags::PreRoll;
-				}
-				if (!SectionRange.GetUpperBound().IsOpen() && Section->GetPostRollFrames() > 0)
-				{
-					SectionData.Flags |= ECinematicShotSectionSortFlags::PostRoll;
-				}
-				SortedSections.Add(SectionData);
+				FCinematicShotSectionSortData SectionData{ Section->GetRowIndex(), Section->GetOverlapPriority(), SectionIndex };
+				SortedSections.Emplace(SectionData);
 			}
 		}
 	}
