Sequencer: Use the proper property traits for the component transform property handler.

diff --git a/Engine/Source/Runtime/MovieSceneTracks/Private/MovieSceneTracksComponentTypes.cpp b/Engine/Source/Runtime/MovieSceneTracks/Private/MovieSceneTracksComponentTypes.cpp
--- a/Engine/Source/Runtime/MovieSceneTracks/Private/MovieSceneTracksComponentTypes.cpp
+++ b/Engine/Source/Runtime/MovieSceneTracks/Private/MovieSceneTracksComponentTypes.cpp
@@ -461,7 +461,7 @@ struct FDoubleVectorHandler : TPropertyComponentHandler<FDoubleVectorPropertyTra
 };
 
 
-struct FTransformHandler : TPropertyComponentHandler<FTransformPropertyTraits, double, double, double, double, double, double, double, double, double>
+struct FComponentTransformHandler : TPropertyComponentHandler<FComponentTransformPropertyTraits, double, double, double, double, double, double, double, double, double>
 {
 	TSharedPtr<IPreAnimatedStorage> GetPreAnimatedStateStorage(const FPropertyDefinition& Definition, FPreAnimatedStateExtension* Container) override
 	{
@@ -642,7 +642,7 @@ FMovieSceneTracksComponentTypes::FMovieSceneTracksComponentTypes()
 		.AddComposite(BuiltInComponents->DoubleResult[8], &FIntermediate3DTransform::S_Z)
 		.SetBlenderSystem<UMovieScenePiecewiseDoubleBlenderSystem>()
 		.SetCustomAccessors(&Accessors.ComponentTransform)
-		.Commit(FTransformHandler());
+		.Commit(FComponentTransformHandler());
 	}
 
 	// --------------------------------------------------------------------------------------------
