Support for Metasounds in Sequencer via transmission interface.

diff --git a/Engine/Plugins/Editor/FacialAnimation/Source/FacialAnimation/FacialAnimation.Build.cs b/Engine/Plugins/Editor/FacialAnimation/Source/FacialAnimation/FacialAnimation.Build.cs
--- a/Engine/Plugins/Editor/FacialAnimation/Source/FacialAnimation/FacialAnimation.Build.cs
+++ b/Engine/Plugins/Editor/FacialAnimation/Source/FacialAnimation/FacialAnimation.Build.cs
@@ -12,8 +12,9 @@ namespace UnrealBuildTool.Rules
 					"CoreUObject",
 					"InputCore",
 					"Engine",
+					"AudioExtensions",
 					"AudioMixer"
-                }
+				}
 			);
 		}
 	}
diff --git a/Engine/Plugins/Runtime/Metasound/Source/MetasoundEngine/Private/MetasoundSource.cpp b/Engine/Plugins/Runtime/Metasound/Source/MetasoundEngine/Private/MetasoundSource.cpp
--- a/Engine/Plugins/Runtime/Metasound/Source/MetasoundEngine/Private/MetasoundSource.cpp
+++ b/Engine/Plugins/Runtime/Metasound/Source/MetasoundEngine/Private/MetasoundSource.cpp
@@ -462,6 +462,81 @@ ISoundGeneratorPtr UMetaSoundSource::CreateSoundGenerator(const FSoundGeneratorI
 	return ISoundGeneratorPtr(new FMetasoundGenerator(MoveTemp(InitParams)));
 }
 
+bool UMetaSoundSource::GetAllDefaultParameters(TArray<FAudioParameter>& OutParameters) const
+{
+	using namespace Metasound;
+	using namespace Metasound::Frontend;
+	using namespace Metasound::Engine;
+
+	const FMetasoundFrontendDocument& Document = GetDocumentChecked();
+	for(const FMetasoundFrontendClassInput& Input : Document.RootGraph.Interface.Inputs)
+	{		
+		FAudioParameter Params;
+		Params.ParamName = Input.Name;
+		switch (Input.DefaultLiteral.GetType())
+		{
+		// bool
+		case EMetasoundFrontendLiteralType::Boolean:
+			Params.ParamType = EAudioParameterType::Boolean;
+			ensure(Input.DefaultLiteral.TryGet(Params.BoolParam));
+			break;
+		case EMetasoundFrontendLiteralType::BooleanArray:
+			Params.ParamType = EAudioParameterType::BooleanArray;
+			ensure(Input.DefaultLiteral.TryGet(Params.ArrayBoolParam));
+			break;
+
+		// int32
+		case EMetasoundFrontendLiteralType::Integer:
+			Params.ParamType = EAudioParameterType::Integer;
+			ensure(Input.DefaultLiteral.TryGet(Params.IntParam));
+			break;
+		case EMetasoundFrontendLiteralType::IntegerArray:
+			Params.ParamType = EAudioParameterType::IntegerArray;
+			ensure(Input.DefaultLiteral.TryGet(Params.ArrayIntParam));
+			break;
+
+		// float
+		case EMetasoundFrontendLiteralType::Float:
+			Params.ParamType = EAudioParameterType::Float;
+			ensure(Input.DefaultLiteral.TryGet(Params.FloatParam));
+			break;
+		case EMetasoundFrontendLiteralType::FloatArray:
+			Params.ParamType = EAudioParameterType::FloatArray;
+			ensure(Input.DefaultLiteral.TryGet(Params.ArrayFloatParam));
+			break;
+
+		// FStrings.
+		case EMetasoundFrontendLiteralType::String:
+			Params.ParamType = EAudioParameterType::String;
+			ensure(Input.DefaultLiteral.TryGet(Params.StringParam));
+			break;
+		case EMetasoundFrontendLiteralType::StringArray:
+			Params.ParamType = EAudioParameterType::StringArray;
+			ensure(Input.DefaultLiteral.TryGet(Params.ArrayStringParam));
+			break;
+
+		// UObjects.
+		//case EMetasoundFrontendLiteralType::UObject:
+		//	Params.ParamType = EAudioParameterType::Object;
+		//	ensure(Input.DefaultLiteral.TryGet(Params.ObjectParam));
+		//	break;
+		case EMetasoundFrontendLiteralType::UObjectArray:
+			Params.ParamType = EAudioParameterType::ObjectArray;
+			ensure(Input.DefaultLiteral.TryGet(Params.ArrayObjectParam));
+			break;
+		default:
+			break;
+		}
+
+		if (Params.ParamType != EAudioParameterType::None)
+		{
+			OutParameters.Add(Params);
+		}
+	}
+	return true;
+}
+
+
 bool UMetaSoundSource::IsParameterValid(const FAudioParameter& InParameter) const
 {
 	TMap<FName, FName> InputNameTypeMap;
diff --git a/Engine/Plugins/Runtime/Metasound/Source/MetasoundEngine/Public/MetasoundSource.h b/Engine/Plugins/Runtime/Metasound/Source/MetasoundEngine/Public/MetasoundSource.h
--- a/Engine/Plugins/Runtime/Metasound/Source/MetasoundEngine/Public/MetasoundSource.h
+++ b/Engine/Plugins/Runtime/Metasound/Source/MetasoundEngine/Public/MetasoundSource.h
@@ -174,6 +174,9 @@ public:
 
 
 protected:
+	/** Gets all the default parameters for this Asset.  */
+	virtual bool GetAllDefaultParameters(TArray<FAudioParameter>& OutParameters) const override;
+
 	virtual void SetReferencedAssetClassKeys(TSet<Metasound::Frontend::FNodeRegistryKey>&& InKeys) override;
 
 	bool IsParameterValid(const FAudioParameter& InParameter, const TMap<FName, FName>& InInputNameTypePairs) const;
diff --git a/Engine/Plugins/Runtime/Metasound/Source/MetasoundFrontend/Private/MetasoundFrontendLiteral.cpp b/Engine/Plugins/Runtime/Metasound/Source/MetasoundFrontend/Private/MetasoundFrontendLiteral.cpp
--- a/Engine/Plugins/Runtime/Metasound/Source/MetasoundFrontend/Private/MetasoundFrontendLiteral.cpp
+++ b/Engine/Plugins/Runtime/Metasound/Source/MetasoundFrontend/Private/MetasoundFrontendLiteral.cpp
@@ -478,6 +478,85 @@ bool FMetasoundFrontendLiteral::TryGet(TArray<UObject*>& OutValue) const
 	return false;
 }
 
+bool FMetasoundFrontendLiteral::TryGet(bool& OutValue) const
+{
+	if (Type == EMetasoundFrontendLiteralType::Boolean)
+	{
+		OutValue = AsBoolean[0];
+		return true;
+	}
+	return false;
+}
+
+bool FMetasoundFrontendLiteral::TryGet(TArray<bool>& OutValue) const
+{
+	if (Type == EMetasoundFrontendLiteralType::BooleanArray)
+	{
+		OutValue = AsBoolean;
+		return true;
+	}
+	return false;
+}
+
+bool FMetasoundFrontendLiteral::TryGet(int32& OutValue) const
+{
+	if (Type == EMetasoundFrontendLiteralType::Integer)
+	{
+		OutValue = AsInteger[0];
+		return true;
+	}
+	return false;
+}
+
+bool FMetasoundFrontendLiteral::TryGet(TArray<int32>& OutValue) const
+{
+	if (Type == EMetasoundFrontendLiteralType::Integer)
+	{
+		OutValue = AsInteger;
+		return true;
+	}
+	return false;
+}
+
+bool FMetasoundFrontendLiteral::TryGet(float& OutValue) const
+{
+	if (Type == EMetasoundFrontendLiteralType::Float)
+	{
+		OutValue = AsFloat[0];
+		return true;
+	}
+	return false;
+}
+
+bool FMetasoundFrontendLiteral::TryGet(TArray<float>& OutValue) const
+{
+	if (Type == EMetasoundFrontendLiteralType::FloatArray)
+	{
+		OutValue = AsFloat;
+		return true;
+	}
+	return false;
+}
+
+bool FMetasoundFrontendLiteral::TryGet(FString& OutValue) const
+{
+	if (Type == EMetasoundFrontendLiteralType::String)
+	{
+		OutValue = AsString[0];
+		return true;
+	}
+	return false;
+}
+
+bool FMetasoundFrontendLiteral::TryGet(TArray<FString>& OutValue) const
+{
+	if (Type == EMetasoundFrontendLiteralType::StringArray)
+	{
+		OutValue = AsString;
+		return true;
+	}
+	return false;
+}
 
 Metasound::FLiteral FMetasoundFrontendLiteral::ToLiteral(const FName& InMetasoundDataTypeName) const
 {
diff --git a/Engine/Plugins/Runtime/Metasound/Source/MetasoundFrontend/Public/MetasoundFrontendLiteral.h b/Engine/Plugins/Runtime/Metasound/Source/MetasoundFrontend/Public/MetasoundFrontendLiteral.h
--- a/Engine/Plugins/Runtime/Metasound/Source/MetasoundFrontend/Public/MetasoundFrontendLiteral.h
+++ b/Engine/Plugins/Runtime/Metasound/Source/MetasoundFrontend/Public/MetasoundFrontendLiteral.h
@@ -100,6 +100,15 @@ public:
 	bool TryGet(UObject*& OutValue) const;
 	bool TryGet(TArray<UObject*>& OutValue) const;
 
+	bool TryGet(bool& OutValue) const;
+	bool TryGet(TArray<bool>& OutValue) const;
+	bool TryGet(int32& OutValue) const;
+	bool TryGet(TArray<int32>& OutValue) const;
+	bool TryGet(float& OutValue) const;
+	bool TryGet(TArray<float>& OutValue) const;
+	bool TryGet(FString& OutValue) const;
+	bool TryGet(TArray<FString>& OutValue) const;
+
 	// Sets the literal to the given type and value to default;
 	void SetType(EMetasoundFrontendLiteralType InType);
 
diff --git a/Engine/Source/Runtime/AudioExtensions/Public/AudioParameterInterface.h b/Engine/Source/Runtime/AudioExtensions/Public/AudioParameterInterface.h
--- a/Engine/Source/Runtime/AudioExtensions/Public/AudioParameterInterface.h
+++ b/Engine/Source/Runtime/AudioExtensions/Public/AudioParameterInterface.h
@@ -436,4 +436,12 @@ public:
 
 	// Sets an array of parameters as a batch
 	virtual void SetParameters(TArray<FAudioParameter>&& InValues) = 0;
+
+	// Template Specialization for non-script clients.
+	template<typename DataType> void SetParameter(FName InName, DataType&&) = delete;	
+	template<> void SetParameter(FName InName, bool&& InBool)			{ SetBoolParameter(InName, InBool); }
+	template<> void SetParameter(FName InName, float&& InFloat)			{ SetFloatParameter(InName, InFloat); }
+	template<> void SetParameter(FName InName, int32&& InInteger)		{ SetIntParameter(InName, InInteger); }
+	template<> void SetParameter(FName InName, FString&& InString)		{ SetStringParameter(InName, InString); }
+	// Do the rest jim
 };
\ No newline at end of file
diff --git a/Engine/Source/Runtime/Engine/Classes/Sound/SoundBase.h b/Engine/Source/Runtime/Engine/Classes/Sound/SoundBase.h
--- a/Engine/Source/Runtime/Engine/Classes/Sound/SoundBase.h
+++ b/Engine/Source/Runtime/Engine/Classes/Sound/SoundBase.h
@@ -303,4 +303,7 @@ public:
 
 	/** Returns whether parameter is valid input for the given sound */
 	virtual bool IsParameterValid(const FAudioParameter& InParameter) const;
+
+	/** Gets all the default parameters for this Asset.  */
+	virtual bool GetAllDefaultParameters(TArray<FAudioParameter>& OutParameters) const { return false; }
 };
diff --git a/Engine/Source/Runtime/MovieSceneTracks/MovieSceneTracks.Build.cs b/Engine/Source/Runtime/MovieSceneTracks/MovieSceneTracks.Build.cs
--- a/Engine/Source/Runtime/MovieSceneTracks/MovieSceneTracks.Build.cs
+++ b/Engine/Source/Runtime/MovieSceneTracks/MovieSceneTracks.Build.cs
@@ -21,6 +21,7 @@ public class MovieSceneTracks : ModuleRules
 				"MovieScene",
 				"TimeManagement",
 				"AnimationCore",
+				"AudioExtensions",
 			});
 
 		PrivateDependencyModuleNames.AddRange(new string[]
diff --git a/Engine/Source/Runtime/MovieSceneTracks/Private/Evaluation/MovieSceneAudioTemplate.cpp b/Engine/Source/Runtime/MovieSceneTracks/Private/Evaluation/MovieSceneAudioTemplate.cpp
--- a/Engine/Source/Runtime/MovieSceneTracks/Private/Evaluation/MovieSceneAudioTemplate.cpp
+++ b/Engine/Source/Runtime/MovieSceneTracks/Private/Evaluation/MovieSceneAudioTemplate.cpp
@@ -380,6 +380,21 @@ struct FAudioSectionExecutionToken : IMovieSceneExecutionToken
 		}
 	}
 
+	// Helper template to pair channel evaluation and parameter application.
+	template<typename ChannelType, typename ValueType>
+	void EvaluateAllAndSetParameters(IAudioParameterInterface& InParamaterInterface, const FFrameTime& InTime) const
+	{
+		AudioSection->ForEachInput([&InParamaterInterface, &InTime](FName InName, const ChannelType& InChannel)
+		{
+			using namespace UE::MovieScene;
+			ValueType OutValue{};
+			if (EvaluateChannel(&InChannel, InTime, OutValue))
+			{
+				InParamaterInterface.SetParameter(InName, MoveTempIfPossible(OutValue));
+			}
+		});
+	}
+
 	void EnsureAudioIsPlaying(UAudioComponent& AudioComponent, FPersistentEvaluationData& PersistentData, const FMovieSceneContext& Context, bool bAllowSpatialization, IMovieScenePlayer& Player) const
 	{
 		Player.SavePreAnimatedState(AudioComponent, FStopAudioPreAnimatedToken::GetAnimTypeID(), FStopAudioPreAnimatedToken::FProducer());
@@ -399,6 +414,12 @@ struct FAudioSectionExecutionToken : IMovieSceneExecutionToken
 			AudioComponent.SetPitchMultiplier(PitchMultiplier);
 		}
 
+		// Evaluate inputs and apply the params.
+		EvaluateAllAndSetParameters<FMovieSceneFloatChannel, float>(AudioComponent, Context.GetTime());
+		EvaluateAllAndSetParameters<FMovieSceneBoolChannel, bool> (AudioComponent, Context.GetTime());
+		EvaluateAllAndSetParameters<FMovieSceneIntegerChannel, int32>(AudioComponent, Context.GetTime());
+		EvaluateAllAndSetParameters<FMovieSceneStringChannel, FString>(AudioComponent, Context.GetTime());
+		
 		float SectionStartTimeSeconds = (AudioSection->HasStartFrame() ? AudioSection->GetInclusiveStartFrame() : 0) / AudioSection->GetTypedOuter<UMovieScene>()->GetTickResolution();
 
 		FCachedAudioTrackData& TrackData = PersistentData.GetOrAddTrackData<FCachedAudioTrackData>();
diff --git a/Engine/Source/Runtime/MovieSceneTracks/Private/Sections/MovieSceneAudioSection.cpp b/Engine/Source/Runtime/MovieSceneTracks/Private/Sections/MovieSceneAudioSection.cpp
--- a/Engine/Source/Runtime/MovieSceneTracks/Private/Sections/MovieSceneAudioSection.cpp
+++ b/Engine/Source/Runtime/MovieSceneTracks/Private/Sections/MovieSceneAudioSection.cpp
@@ -57,6 +57,26 @@ UMovieSceneAudioSection::UMovieSceneAudioSection( const FObjectInitializer& Obje
 	PitchMultiplier.SetDefault(1.f);
 }
 
+namespace MovieSceneAudioSectionPrivate
+{
+	template<typename ChannelType, typename ValueType>
+	void AddInputChannels(UMovieSceneAudioSection* InSection, FMovieSceneChannelProxyData& InChannelProxyData)
+	{
+		InSection->ForEachInput([&InChannelProxyData](FName InName, const ChannelType& InChannel)
+		{
+#if WITH_EDITOR
+			FMovieSceneChannelMetaData Data;
+			FText TextName = FText::FromName(InName);	
+			Data.SetIdentifiers(InName, TextName, TextName);
+			InChannelProxyData.Add(const_cast<ChannelType&>(InChannel), Data, TMovieSceneExternalValue<ValueType>());
+#else //WITH_EDITOR
+			InChannelProxyData.Add(const_cast<ChannelType&>(InChannel));
+#endif //WITH_EDITOR
+
+		});
+	}
+}
+
 EMovieSceneChannelProxyType  UMovieSceneAudioSection::CacheChannelProxy()
 {
 	// Set up the channel proxy
@@ -86,11 +106,50 @@ EMovieSceneChannelProxyType  UMovieSceneAudioSection::CacheChannelProxy()
 
 #endif
 
+	using namespace MovieSceneAudioSectionPrivate;
+	SetupSoundInputParameters(GetSound());
+	AddInputChannels<FMovieSceneFloatChannel, float>(this, Channels);
+	AddInputChannels<FMovieSceneBoolChannel, bool>(this, Channels);
+	AddInputChannels<FMovieSceneIntegerChannel, int32>(this, Channels);
+	AddInputChannels<FMovieSceneStringChannel, FString>(this, Channels);
+
 	ChannelProxy = MakeShared<FMovieSceneChannelProxy>(MoveTemp(Channels));
 
 	return EMovieSceneChannelProxyType::Dynamic;
 }
 
+void UMovieSceneAudioSection::SetupSoundInputParameters(const USoundBase* InSoundBase)
+{
+	// Populate with defaults.
+	if (InSoundBase)
+	{
+		TArray<FAudioParameter> DefaultParams;
+		InSoundBase->GetAllDefaultParameters(DefaultParams);
+
+		for (const FAudioParameter& Param : DefaultParams)
+		{
+			switch(Param.ParamType)
+			{
+			case EAudioParameterType::Float:
+				Inputs_Float.FindOrAdd(Param.ParamName, FMovieSceneFloatChannel{}).SetDefault(Param.FloatParam);
+				break;
+			case EAudioParameterType::Boolean:
+				Inputs_Bool.FindOrAdd(Param.ParamName, FMovieSceneBoolChannel{}).SetDefault(Param.BoolParam);
+				break;
+			case EAudioParameterType::Integer:
+				Inputs_Int.FindOrAdd(Param.ParamName, FMovieSceneIntegerChannel{}).SetDefault(Param.IntParam);
+				break;
+			case EAudioParameterType::String:
+				Inputs_String.FindOrAdd(Param.ParamName, FMovieSceneStringChannel{}).SetDefault(Param.StringParam);
+				break;
+			default:
+				// Not supported yet.
+				break;
+			}
+		}
+	}
+}
+
 TOptional<FFrameTime> UMovieSceneAudioSection::GetOffsetTime() const
 {
 	return TOptional<FFrameTime>(StartFrameOffset);
@@ -212,6 +271,11 @@ UMovieSceneSection* UMovieSceneAudioSection::SplitSection(FQualifiedFrameTime Sp
 	return NewSection;
 }
 
+void UMovieSceneAudioSection::SetSound(USoundBase* InSound)
+{
+	Sound = InSound;
+	CacheChannelProxy();
+}
 
 USceneComponent* UMovieSceneAudioSection::GetAttachComponent(const AActor* InParentActor, const FMovieSceneActorReferenceKey& Key) const
 {
diff --git a/Engine/Source/Runtime/MovieSceneTracks/Public/Sections/MovieSceneAudioSection.h b/Engine/Source/Runtime/MovieSceneTracks/Public/Sections/MovieSceneAudioSection.h
--- a/Engine/Source/Runtime/MovieSceneTracks/Public/Sections/MovieSceneAudioSection.h
+++ b/Engine/Source/Runtime/MovieSceneTracks/Public/Sections/MovieSceneAudioSection.h
@@ -8,6 +8,9 @@
 #include "Runtime/Engine/Classes/Components/AudioComponent.h"
 #include "Sound/SoundAttenuation.h"
 #include "Channels/MovieSceneFloatChannel.h"
+#include "Channels/MovieSceneStringChannel.h"
+#include "Channels/MovieSceneBoolChannel.h"
+#include "Channels/MovieSceneIntegerChannel.h"
 #include "Sections/MovieSceneActorReferenceSection.h"
 #include "MovieSceneAudioSection.generated.h"
 
@@ -26,7 +29,7 @@ public:
 
 	/** Sets this section's sound */
 	UFUNCTION(BlueprintCallable, Category="Sequencer|Section")
-	void SetSound(class USoundBase* InSound) {Sound = InSound;}
+	void SetSound(class USoundBase* InSound);
 
 	/** Gets the sound for this section */
 	UFUNCTION(BlueprintPure, Category = "Sequencer|Section")
@@ -158,6 +161,12 @@ public:
 	{
 		return OnAudioPlaybackPercent;
 	}
+	
+	/** Overloads for each input type, const */
+	void ForEachInput(TFunction<void(FName, const FMovieSceneBoolChannel&)> InFunction) const { ForEachInternal(InFunction, Inputs_Bool); }
+	void ForEachInput(TFunction<void(FName, const FMovieSceneStringChannel&)> InFunction) const { ForEachInternal(InFunction, Inputs_String); }
+	void ForEachInput(TFunction<void(FName, const FMovieSceneIntegerChannel&)> InFunction) const  { ForEachInternal(InFunction, Inputs_Int); }
+	void ForEachInput(TFunction<void(FName, const FMovieSceneFloatChannel&)> InFunction) const { ForEachInternal(InFunction, Inputs_Float); }
 
 public:
 
@@ -171,6 +180,16 @@ public:
 	virtual EMovieSceneChannelProxyType CacheChannelProxy() override;
 
 private:
+	template<typename ChannelType, typename ForEachFunction>
+	FORCEINLINE static void ForEachInternal(ForEachFunction InFuncton, const TMap<FName, ChannelType>& InMapToIterate) 
+	{
+		for (auto& Item : InMapToIterate)
+		{
+			InFuncton(Item.Key, Item.Value);
+		}	
+	}
+
+	void SetupSoundInputParameters(const USoundBase* InSoundBase);
 
 	/** The sound cue or wave that this section plays */
 	UPROPERTY(EditAnywhere, Category="Audio")
@@ -204,6 +223,16 @@ private:
 	UPROPERTY( )
 	FMovieSceneFloatChannel PitchMultiplier;
 
+	/** Generic inputs for the sound  */
+	UPROPERTY()
+	TMap<FName, FMovieSceneFloatChannel> Inputs_Float;
+	UPROPERTY()
+	TMap<FName, FMovieSceneStringChannel> Inputs_String;
+	UPROPERTY()
+	TMap<FName, FMovieSceneBoolChannel> Inputs_Bool;
+	UPROPERTY()
+	TMap<FName, FMovieSceneIntegerChannel> Inputs_Int;
+
 	UPROPERTY()
 	FMovieSceneActorReferenceData AttachActorData;
 
