PR #8682: Fix FIntermediateColor did not convert from FColor properly. (Contributed by [REDACTED_PERSON])

diff --git a/Engine/Source/Runtime/MovieSceneTracks/Public/MovieSceneTracksPropertyTypes.h b/Engine/Source/Runtime/MovieSceneTracks/Public/MovieSceneTracksPropertyTypes.h
--- a/Engine/Source/Runtime/MovieSceneTracks/Public/MovieSceneTracksPropertyTypes.h
+++ b/Engine/Source/Runtime/MovieSceneTracks/Public/MovieSceneTracksPropertyTypes.h
@@ -102,8 +102,13 @@ struct FIntermediateColor
 		: R(InColor.R), G(InColor.G), B(InColor.B), A(InColor.A)
 	{}
 	explicit FIntermediateColor(const FColor& InColor)
-		: R(InColor.R), G(InColor.G), B(InColor.B), A(InColor.A)
-	{}
+	{
+		FLinearColor NewColor = FLinearColor::FromSRGBColor(InColor);
+		R = NewColor.R;
+		G = NewColor.G;
+		B = NewColor.B;
+		A = NewColor.A;
+	}
 
 	explicit FIntermediateColor(const FSlateColor& InSlateColor)
 	{
