Added Insights tags to identify MovieSceneDataLayerSystem FlushStreamingFull/FlushStreamingVisibility.

diff --git a/Engine/Source/Runtime/MovieSceneTracks/Private/Systems/MovieSceneDataLayerSystem.cpp b/Engine/Source/Runtime/MovieSceneTracks/Private/Systems/MovieSceneDataLayerSystem.cpp
--- a/Engine/Source/Runtime/MovieSceneTracks/Private/Systems/MovieSceneDataLayerSystem.cpp
+++ b/Engine/Source/Runtime/MovieSceneTracks/Private/Systems/MovieSceneDataLayerSystem.cpp
@@ -521,11 +521,14 @@ void UMovieSceneDataLayerSystem::OnRun(FSystemTaskPrerequisites& InPrerequisites
 
 				if (EnumHasAnyFlags(UpdateFlags, EDataLayerUpdateFlags::FlushStreamingFull))
 				{
+					TRACE_CPUPROFILER_EVENT_SCOPE(UMovieSceneDataLayerSystem_FlushStreamingFull);
 					World->BlockTillLevelStreamingCompleted();
 				}
 				else if (EnumHasAnyFlags(UpdateFlags, EDataLayerUpdateFlags::FlushStreamingVisibility))
 				{
-					// Make sure any DataLayer state change is processed before flushing visibility
+					TRACE_CPUPROFILER_EVENT_SCOPE(UMovieSceneDataLayerSystem_FlushStreamingVisibility);
+
+					// Make sure any DataLayer state change is processed before flushing visibility					
 					WorldPartitionSubsystem->UpdateStreamingState();
 					World->FlushLevelStreaming(EFlushLevelStreamingType::Visibility);
 				}
