Sequencer: make data layers tracks desired and preroll state activated by default, which is what we want most of the times and avoids large hitches on the GT when activating data layers.

diff --git a/Engine/Source/Runtime/MovieSceneTracks/Private/Sections/MovieSceneDataLayerSection.cpp b/Engine/Source/Runtime/MovieSceneTracks/Private/Sections/MovieSceneDataLayerSection.cpp
--- a/Engine/Source/Runtime/MovieSceneTracks/Private/Sections/MovieSceneDataLayerSection.cpp
+++ b/Engine/Source/Runtime/MovieSceneTracks/Private/Sections/MovieSceneDataLayerSection.cpp
@@ -7,8 +7,8 @@
 UMovieSceneDataLayerSection::UMovieSceneDataLayerSection(const FObjectInitializer& ObjInit)
 	: Super(ObjInit)
 {
-	DesiredState = EDataLayerRuntimeState::Loaded;
-	PrerollState = EDataLayerRuntimeState::Loaded;
+	DesiredState = EDataLayerRuntimeState::Activated;
+	PrerollState = EDataLayerRuntimeState::Activated;
 	EvalOptions.EnableAndSetCompletionMode(EMovieSceneCompletionMode::RestoreState);
 }
 
