Sequencer: added debug information when flushing streaming or visibility.

diff --git a/Engine/Source/Runtime/MovieSceneTracks/Private/Systems/MovieSceneDataLayerSystem.cpp b/Engine/Source/Runtime/MovieSceneTracks/Private/Systems/MovieSceneDataLayerSystem.cpp
--- a/Engine/Source/Runtime/MovieSceneTracks/Private/Systems/MovieSceneDataLayerSystem.cpp
+++ b/Engine/Source/Runtime/MovieSceneTracks/Private/Systems/MovieSceneDataLayerSystem.cpp
@@ -8,6 +8,7 @@
 #include "Evaluation/PreAnimatedState/MovieScenePreAnimatedEntityCaptureSource.h"
 #include "Evaluation/PreAnimatedState/MovieScenePreAnimatedStateStorage.h"
 #include "Evaluation/PreAnimatedState/MovieScenePreAnimatedStorageID.inl"
+#include "Evaluation/MovieSceneEvaluationTemplateInstance.h"
 #include "MovieSceneTracksComponentTypes.h"
 #include "Sections/MovieSceneDataLayerSection.h"
 #include "WorldPartition/DataLayer/DataLayerSubsystem.h"
@@ -16,6 +17,8 @@
 #include "WorldPartition/WorldPartitionStreamingSource.h"
 #include "Engine/World.h"
 #include "Misc/EnumClassFlags.h"
+#include "MovieSceneSequence.h"
+#include "IMovieScenePlayer.h"
 
 #include "Kismet/KismetSystemLibrary.h"
 
@@ -333,12 +336,12 @@ EDataLayerUpdateFlags FDesiredLayerStates::Apply(FPreAnimatedDataLayerStorage* P
 					if (DesiredStateValue == EDataLayerRuntimeState::Activated && !IsDataLayerReady(DataLayer, EDataLayerRuntimeState::Loaded, false))
 					{
 						Flags |= EDataLayerUpdateFlags::FlushStreamingFull;
-						UE_LOG(LogMovieScene, Warning, TEXT("Data layer with name '%s' is causing a full streaming flush"), *DataLayer->GetDataLayerLabel().ToString());
+						UE_LOG(LogMovieScene, Warning, TEXT("[UMovieSceneDataLayerSystem] Data layer with name '%s' is causing a full streaming flush (%s)"), *DataLayer->GetDataLayerLabel().ToString(), GetDataLayerStateName(DesiredStateValue));
 					}
 					else
 					{
 						Flags |= EDataLayerUpdateFlags::FlushStreamingVisibility;
-						UE_LOG(LogMovieScene, Log, TEXT("Data layer with name '%s' is causing a visibility streaming flush"), *DataLayer->GetDataLayerLabel().ToString());
+						UE_LOG(LogMovieScene, Log, TEXT("[UMovieSceneDataLayerSystem] Data layer with name '%s' is causing a visibility streaming flush (%s)"), *DataLayer->GetDataLayerLabel().ToString(), GetDataLayerStateName(DesiredStateValue));
 					}
 				}
 			}
@@ -519,19 +522,46 @@ void UMovieSceneDataLayerSystem::OnRun(FSystemTaskPrerequisites& InPrerequisites
 			{
 				EDataLayerUpdateFlags UpdateFlags = DesiredLayerStates->Apply(WeakPreAnimatedStorage.Pin().Get(), DataLayerSubsystem, WorldPartitionSubsystem);
 
+				const double StartTime = FPlatformTime::Seconds();
+				const TCHAR* FlushTypeString = nullptr;
+
 				if (EnumHasAnyFlags(UpdateFlags, EDataLayerUpdateFlags::FlushStreamingFull))
 				{
 					TRACE_CPUPROFILER_EVENT_SCOPE(UMovieSceneDataLayerSystem_FlushStreamingFull);
+					FlushTypeString = TEXT("FlushStreamingFull");
+
 					World->BlockTillLevelStreamingCompleted();
 				}
 				else if (EnumHasAnyFlags(UpdateFlags, EDataLayerUpdateFlags::FlushStreamingVisibility))
 				{
 					TRACE_CPUPROFILER_EVENT_SCOPE(UMovieSceneDataLayerSystem_FlushStreamingVisibility);
+					FlushTypeString = TEXT("FlushStreamingVisibility");
 
 					// Make sure any DataLayer state change is processed before flushing visibility					
 					WorldPartitionSubsystem->UpdateStreamingState();
 					World->FlushLevelStreaming(EFlushLevelStreamingType::Visibility);
 				}
+
+				UE_SUPPRESS(LogMovieScene, Warning,
+				{
+					if (FlushTypeString)
+					{
+						FString SequenceList;
+						for (const FSequenceInstance& Instance : Linker->GetInstanceRegistry()->GetSparseInstances())
+						{
+							UMovieSceneSequence* Sequence = Instance.GetPlayer()->GetEvaluationTemplate().GetSequence(Instance.GetSequenceID());
+
+							if (SequenceList.Len())
+							{
+								SequenceList += TEXT(",");
+							}
+							SequenceList += Sequence->GetName();
+							
+						}
+
+						UE_LOG(LogMovieScene, Warning, TEXT("[UMovieSceneDataLayerSystem] %s took %.4f seconds (%s)"), FlushTypeString, FPlatformTime::Seconds() - StartTime, *SequenceList);
+					}
+				});
 			}
 		}
 	}
