Add Insight bookmarks to show data layers prerolls and flushes

diff --git a/Engine/Source/Runtime/MovieSceneTracks/Private/Systems/MovieSceneDataLayerSystem.cpp b/Engine/Source/Runtime/MovieSceneTracks/Private/Systems/MovieSceneDataLayerSystem.cpp
--- a/Engine/Source/Runtime/MovieSceneTracks/Private/Systems/MovieSceneDataLayerSystem.cpp
+++ b/Engine/Source/Runtime/MovieSceneTracks/Private/Systems/MovieSceneDataLayerSystem.cpp
@@ -337,11 +337,13 @@ EDataLayerUpdateFlags FDesiredLayerStates::Apply(FPreAnimatedDataLayerStorage* P
 					{
 						Flags |= EDataLayerUpdateFlags::FlushStreamingFull;
 						UE_LOG(LogMovieScene, Warning, TEXT("[UMovieSceneDataLayerSystem] Data layer with name '%s' is causing a full streaming flush (%s)"), *DataLayer->GetDataLayerLabel().ToString(), GetDataLayerStateName(DesiredStateValue));
+						CSV_EVENT_GLOBAL(TEXT("SeqDataLayerFlushFull-%s-%s"), *DataLayer->GetDataLayerLabel().ToString(), GetDataLayerStateName(DesiredStateValue));
 					}
 					else
 					{
 						Flags |= EDataLayerUpdateFlags::FlushStreamingVisibility;
 						UE_LOG(LogMovieScene, Log, TEXT("[UMovieSceneDataLayerSystem] Data layer with name '%s' is causing a visibility streaming flush (%s)"), *DataLayer->GetDataLayerLabel().ToString(), GetDataLayerStateName(DesiredStateValue));
+						CSV_EVENT_GLOBAL(TEXT("SeqDataLayerFlushVis-%s-%s"), *DataLayer->GetDataLayerLabel().ToString(), GetDataLayerStateName(DesiredStateValue));
 					}
 				}
 			}
