Sequencer: Fix DataLayer preroll being cancelled out by previous shot unload track (Priority is given to Activation, then Loaded, then Unload)

diff --git a/Engine/Source/Runtime/MovieSceneTracks/Private/Systems/MovieSceneDataLayerSystem.cpp b/Engine/Source/Runtime/MovieSceneTracks/Private/Systems/MovieSceneDataLayerSystem.cpp
--- a/Engine/Source/Runtime/MovieSceneTracks/Private/Systems/MovieSceneDataLayerSystem.cpp
+++ b/Engine/Source/Runtime/MovieSceneTracks/Private/Systems/MovieSceneDataLayerSystem.cpp
@@ -231,25 +231,20 @@ bool FDataLayerState::IsEmpty() const
 
 TOptional<EDataLayerRuntimeState> FDataLayerState::ComputeDesiredState() const
 {
-	// If we have any requests to keep a layer loaded, always keep it loaded (even if things ask for it to be hidden)
-	EDataLayerRuntimeState FallbackState = LoadedCount != 0 ? EDataLayerRuntimeState::Loaded : EDataLayerRuntimeState::Unloaded;
-
-	if (ActivatedCount == UnloadedCount)
+	if (ActivatedCount > 0)
 	{
-		// Equal number of requests for active and unloaded - just leave the data layer alone
-		if (LoadedCount != 0)
-		{
-			return EDataLayerRuntimeState::Loaded;
-		}
-		return TOptional<EDataLayerRuntimeState>();
+		return EDataLayerRuntimeState::Activated;
 	}
-	
-	if (ActivatedCount > UnloadedCount)
+	else if (LoadedCount > 0)
 	{
-		return EDataLayerRuntimeState::Activated;
+		return EDataLayerRuntimeState::Loaded;
+	}
+	else if (UnloadedCount > 0)
+	{
+		return EDataLayerRuntimeState::Unloaded;
 	}
 
-	return FallbackState;
+	return TOptional<EDataLayerRuntimeState>();
 }
 
 bool FDataLayerState::ShouldFlushStreaming(EDataLayerRuntimeState ComputedState) const
