Sequencer: added bFlushOnUnload to data layer tracks to control if we should flush level streaming visibility when unloading a data layer to avoid GT hitches.

diff --git a/Engine/Source/Runtime/MovieSceneTracks/Private/Sections/MovieSceneDataLayerSection.cpp b/Engine/Source/Runtime/MovieSceneTracks/Private/Sections/MovieSceneDataLayerSection.cpp
--- a/Engine/Source/Runtime/MovieSceneTracks/Private/Sections/MovieSceneDataLayerSection.cpp
+++ b/Engine/Source/Runtime/MovieSceneTracks/Private/Sections/MovieSceneDataLayerSection.cpp
@@ -9,6 +9,7 @@ UMovieSceneDataLayerSection::UMovieSceneDataLayerSection(const FObjectInitialize
 {
 	DesiredState = EDataLayerRuntimeState::Activated;
 	PrerollState = EDataLayerRuntimeState::Activated;
+	bFlushOnUnload = true;
 	EvalOptions.EnableAndSetCompletionMode(EMovieSceneCompletionMode::RestoreState);
 }
 
@@ -32,6 +33,16 @@ void UMovieSceneDataLayerSection::SetPrerollState(EDataLayerRuntimeState InPrero
 	PrerollState = InPrerollState;
 }
 
+bool UMovieSceneDataLayerSection::GetFlushOnUnload() const
+{
+	return bFlushOnUnload;
+}
+
+void UMovieSceneDataLayerSection::SetFlushOnUnload(bool bInFlushOnUnload)
+{
+	bFlushOnUnload = bInFlushOnUnload;
+}
+
 void UMovieSceneDataLayerSection::ImportEntityImpl(UMovieSceneEntitySystemLinker* EntityLinker, const FEntityImportParams& Params, FImportedEntity* OutImportedEntity)
 {
 	using namespace UE::MovieScene;
diff --git a/Engine/Source/Runtime/MovieSceneTracks/Private/Systems/MovieSceneDataLayerSystem.cpp b/Engine/Source/Runtime/MovieSceneTracks/Private/Systems/MovieSceneDataLayerSystem.cpp
--- a/Engine/Source/Runtime/MovieSceneTracks/Private/Systems/MovieSceneDataLayerSystem.cpp
+++ b/Engine/Source/Runtime/MovieSceneTracks/Private/Systems/MovieSceneDataLayerSystem.cpp
@@ -587,8 +587,8 @@ void UMovieSceneDataLayerSystem::UpdateDesiredStates()
 				continue;
 			}
 
-			const bool bRequiresStreamingFlush = bPreroll == false;
 			EDataLayerRuntimeState DesiredState = bPreroll ? Section->GetPrerollState() : Section->GetDesiredState();
+			const bool bRequiresStreamingFlush = (DesiredState == EDataLayerState::Unloaded) ? Section->GetFlushOnUnload() : !bPreroll;
 
 			for (const FActorDataLayer& ActorDataLayer : Section->GetDataLayers())
 			{
diff --git a/Engine/Source/Runtime/MovieSceneTracks/Public/Sections/MovieSceneDataLayerSection.h b/Engine/Source/Runtime/MovieSceneTracks/Public/Sections/MovieSceneDataLayerSection.h
--- a/Engine/Source/Runtime/MovieSceneTracks/Public/Sections/MovieSceneDataLayerSection.h
+++ b/Engine/Source/Runtime/MovieSceneTracks/Public/Sections/MovieSceneDataLayerSection.h
@@ -32,6 +32,12 @@ public:
 	UFUNCTION(BlueprintCallable, Category = "Sequencer|Section")
 	MOVIESCENETRACKS_API void SetPrerollState(EDataLayerRuntimeState InPrerollState);
 
+	UFUNCTION(BlueprintPure, Category = "Sequencer|Section")
+	MOVIESCENETRACKS_API bool GetFlushOnUnload() const;
+
+	UFUNCTION(BlueprintCallable, Category = "Sequencer|Section")
+	MOVIESCENETRACKS_API void SetFlushOnUnload(bool bFlushOnUnload);
+
 	UFUNCTION(BlueprintPure, Category = "Sequencer|Section")
 	const TArray<FActorDataLayer>& GetDataLayers() const { return DataLayers; }
 
@@ -55,4 +61,8 @@ private:
 	/** The desired state for the data layers on this section when the section is pre or post-rolling. */
 	UPROPERTY(EditAnywhere, Category=DataLayer)
 	EDataLayerRuntimeState PrerollState;
+
+	/** Determine if we need to flush level streaming when the data layers unloads. */
+	UPROPERTY(EditAnywhere, Category=DataLayer)
+	bool bFlushOnUnload;
 };
