Sequencer: make data layers tracks  bFlushOnUnload false by default to make sure we don't flush level streaming visibility when unloading a data layer.

diff --git a/Engine/Source/Runtime/MovieSceneTracks/Private/Sections/MovieSceneDataLayerSection.cpp b/Engine/Source/Runtime/MovieSceneTracks/Private/Sections/MovieSceneDataLayerSection.cpp
--- a/Engine/Source/Runtime/MovieSceneTracks/Private/Sections/MovieSceneDataLayerSection.cpp
+++ b/Engine/Source/Runtime/MovieSceneTracks/Private/Sections/MovieSceneDataLayerSection.cpp
@@ -9,7 +9,7 @@ UMovieSceneDataLayerSection::UMovieSceneDataLayerSection(const FObjectInitialize
 {
 	DesiredState = EDataLayerRuntimeState::Activated;
 	PrerollState = EDataLayerRuntimeState::Activated;
-	bFlushOnUnload = true;
+	bFlushOnUnload = false;
 	EvalOptions.EnableAndSetCompletionMode(EMovieSceneCompletionMode::RestoreState);
 }
 
