Build Fix: DataLayer API changes related errors

diff --git a/Engine/Source/Runtime/Engine/Private/WorldPartition/DataLayer/WorldDataLayers.cpp b/Engine/Source/Runtime/Engine/Private/WorldPartition/DataLayer/WorldDataLayers.cpp
--- a/Engine/Source/Runtime/Engine/Private/WorldPartition/DataLayer/WorldDataLayers.cpp
+++ b/Engine/Source/Runtime/Engine/Private/WorldPartition/DataLayer/WorldDataLayers.cpp
@@ -153,7 +153,7 @@ PRAGMA_DISABLE_DEPRECATION_WARNINGS
 				{
 					if (!DataLayersFilterDelegate.Execute(DataLayerLabelName, CurrentState, InState))
 					{
-						UE_LOG(LogWorldPartition, Log, TEXT("Data Layer '%s' was filtered out: %s -> %s"),  *DataLayerLabelName.ToString(), GetDataLayerStateName(CurrentState), GetDataLayerStateName(InState));
+						UE_LOG(LogWorldPartition, Log, TEXT("Data Layer '%s' was filtered out: %s -> %s"),  *DataLayerLabelName.ToString(), GetDataLayerRuntimeStateName(CurrentState), GetDataLayerRuntimeStateName(InState));
 						return;
 					}
 				}
@@ -171,7 +171,7 @@ PRAGMA_ENABLE_DEPRECATION_WARNINGS
 			{
 				ActiveDataLayerNames.Add(InDataLayer.Name);
 			}
-			else if (InState == EDataLayerState::Unloaded)
+			else if (InState == EDataLayerRuntimeState::Unloaded)
 			{
 				GLevelStreamingContinuouslyIncrementalGCWhileLevelsPendingPurgeOverride = 1;
 			}
@@ -184,8 +184,8 @@ PRAGMA_ENABLE_DEPRECATION_WARNINGS
 
 #if !NO_LOGGING || CSV_PROFILER
 			const FString DataLayerLabel = DataLayer->GetDataLayerLabel().ToString();
-			UE_LOG(LogWorldPartition, Log, TEXT("Data Layer '%s' state changed: %s -> %s"),  *DataLayerLabel,  GetDataLayerStateName(CurrentState), GetDataLayerStateName(InState));
-			CSV_EVENT_GLOBAL(TEXT("DataLayer-%s-%s"), *DataLayerLabel, GetDataLayerStateName(InState));
+			UE_LOG(LogWorldPartition, Log, TEXT("Data Layer '%s' state changed: %s -> %s"),  *DataLayerLabel, GetDataLayerRuntimeStateName(CurrentState), GetDataLayerRuntimeStateName(InState));
+			CSV_EVENT_GLOBAL(TEXT("DataLayer-%s-%s"), *DataLayerLabel, GetDataLayerRuntimeStateName(InState));
 #endif
 
 			ResolveEffectiveRuntimeState(DataLayer);
diff --git a/Engine/Source/Runtime/Engine/Public/WorldPartition/DataLayer/WorldDataLayers.h b/Engine/Source/Runtime/Engine/Public/WorldPartition/DataLayer/WorldDataLayers.h
--- a/Engine/Source/Runtime/Engine/Public/WorldPartition/DataLayer/WorldDataLayers.h
+++ b/Engine/Source/Runtime/Engine/Public/WorldPartition/DataLayer/WorldDataLayers.h
@@ -144,7 +144,7 @@ private:
 	static int32 DataLayersStateEpoch;
 
 public:
-	DECLARE_DELEGATE_RetVal_ThreeParams(bool, FDataLayersFilterDelegate, FName /*DataLayerName*/, EDataLayerState /*CurrentState*/, EDataLayerState /*TargetState*/);
+	DECLARE_DELEGATE_RetVal_ThreeParams(bool, FDataLayersFilterDelegate, FName /*DataLayerName*/, EDataLayerRuntimeState /*CurrentState*/, EDataLayerRuntimeState /*TargetState*/);
 
 	UE_DEPRECATED(5.00, "do not use, will be replaced by another mechanism for initial release.")
 	FDataLayersFilterDelegate DataLayersFilterDelegate;
diff --git a/Engine/Source/Runtime/MovieSceneTracks/Private/Systems/MovieSceneDataLayerSystem.cpp b/Engine/Source/Runtime/MovieSceneTracks/Private/Systems/MovieSceneDataLayerSystem.cpp
--- a/Engine/Source/Runtime/MovieSceneTracks/Private/Systems/MovieSceneDataLayerSystem.cpp
+++ b/Engine/Source/Runtime/MovieSceneTracks/Private/Systems/MovieSceneDataLayerSystem.cpp
@@ -331,14 +331,14 @@ EDataLayerUpdateFlags FDesiredLayerStates::Apply(FPreAnimatedDataLayerStorage* P
 					if (DesiredStateValue == EDataLayerRuntimeState::Activated && !IsDataLayerReady(DataLayer, EDataLayerRuntimeState::Loaded, false))
 					{
 						Flags |= EDataLayerUpdateFlags::FlushStreamingFull;
-						UE_LOG(LogMovieScene, Warning, TEXT("[UMovieSceneDataLayerSystem] Data layer with name '%s' is causing a full streaming flush (%s)"), *DataLayer->GetDataLayerLabel().ToString(), GetDataLayerStateName(DesiredStateValue));
-						CSV_EVENT_GLOBAL(TEXT("SeqDataLayerFlushFull-%s-%s"), *DataLayer->GetDataLayerLabel().ToString(), GetDataLayerStateName(DesiredStateValue));
+						UE_LOG(LogMovieScene, Warning, TEXT("[UMovieSceneDataLayerSystem] Data layer with name '%s' is causing a full streaming flush (%s)"), *DataLayer->GetDataLayerLabel().ToString(), GetDataLayerRuntimeStateName(DesiredStateValue));
+						CSV_EVENT_GLOBAL(TEXT("SeqDataLayerFlushFull-%s-%s"), *DataLayer->GetDataLayerLabel().ToString(), GetDataLayerRuntimeStateName(DesiredStateValue));
 					}
 					else
 					{
 						Flags |= EDataLayerUpdateFlags::FlushStreamingVisibility;
-						UE_LOG(LogMovieScene, Log, TEXT("[UMovieSceneDataLayerSystem] Data layer with name '%s' is causing a visibility streaming flush (%s)"), *DataLayer->GetDataLayerLabel().ToString(), GetDataLayerStateName(DesiredStateValue));
-						CSV_EVENT_GLOBAL(TEXT("SeqDataLayerFlushVis-%s-%s"), *DataLayer->GetDataLayerLabel().ToString(), GetDataLayerStateName(DesiredStateValue));
+						UE_LOG(LogMovieScene, Log, TEXT("[UMovieSceneDataLayerSystem] Data layer with name '%s' is causing a visibility streaming flush (%s)"), *DataLayer->GetDataLayerLabel().ToString(), GetDataLayerRuntimeStateName(DesiredStateValue));
+						CSV_EVENT_GLOBAL(TEXT("SeqDataLayerFlushVis-%s-%s"), *DataLayer->GetDataLayerLabel().ToString(), GetDataLayerRuntimeStateName(DesiredStateValue));
 					}
 				}
 			}
@@ -588,7 +588,7 @@ void UMovieSceneDataLayerSystem::UpdateDesiredStates()
 			}
 
 			EDataLayerRuntimeState DesiredState = bPreroll ? Section->GetPrerollState() : Section->GetDesiredState();
-			const bool bRequiresStreamingFlush = (DesiredState == EDataLayerState::Unloaded) ? Section->GetFlushOnUnload() : !bPreroll;
+			const bool bRequiresStreamingFlush = (DesiredState == EDataLayerRuntimeState::Unloaded) ? Section->GetFlushOnUnload() : !bPreroll;
 
 			for (const FActorDataLayer& ActorDataLayer : Section->GetDataLayers())
 			{
