No unity fixes

diff --git a/Engine/Plugins/VirtualProduction/StageMonitoring/Source/StageDataProvider/Private/FramePerformanceProvider.cpp b/Engine/Plugins/VirtualProduction/StageMonitoring/Source/StageDataProvider/Private/FramePerformanceProvider.cpp
--- a/Engine/Plugins/VirtualProduction/StageMonitoring/Source/StageDataProvider/Private/FramePerformanceProvider.cpp
+++ b/Engine/Plugins/VirtualProduction/StageMonitoring/Source/StageDataProvider/Private/FramePerformanceProvider.cpp
@@ -7,6 +7,7 @@
 #include "IStageDataProvider.h"
 #include "Misc/App.h"
 #include "Misc/CoreDelegates.h"
+#include "RHIDefinitions.h"
 #include "RenderCore.h"
 #include "StageDataProviderModule.h"
 #include "StageMessages.h"
diff --git a/Engine/Source/Editor/UnrealEd/Private/Commandlets/CookGlobalShaderCommandlet.cpp b/Engine/Source/Editor/UnrealEd/Private/Commandlets/CookGlobalShaderCommandlet.cpp
--- a/Engine/Source/Editor/UnrealEd/Private/Commandlets/CookGlobalShaderCommandlet.cpp
+++ b/Engine/Source/Editor/UnrealEd/Private/Commandlets/CookGlobalShaderCommandlet.cpp
@@ -1,6 +1,7 @@
 // Copyright Epic Games, Inc. All Rights Reserved.
 
 #include "Commandlets/CookGlobalShadersCommandlet.h"
+#include "Engine/Engine.h"
 #include "Misc/Paths.h"
 #include "Misc/FileHelper.h"
 #include "Misc/ConfigCacheIni.h"
diff --git a/Engine/Source/Runtime/MovieSceneTracks/Private/Systems/MovieSceneDataLayerSystem.cpp b/Engine/Source/Runtime/MovieSceneTracks/Private/Systems/MovieSceneDataLayerSystem.cpp
--- a/Engine/Source/Runtime/MovieSceneTracks/Private/Systems/MovieSceneDataLayerSystem.cpp
+++ b/Engine/Source/Runtime/MovieSceneTracks/Private/Systems/MovieSceneDataLayerSystem.cpp
@@ -10,6 +10,7 @@
 #include "Evaluation/PreAnimatedState/MovieScenePreAnimatedStorageID.inl"
 #include "Evaluation/MovieSceneEvaluationTemplateInstance.h"
 #include "MovieSceneTracksComponentTypes.h"
+#include "ProfilingDebugging/CsvProfiler.h"
 #include "Sections/MovieSceneDataLayerSection.h"
 #include "WorldPartition/DataLayer/DataLayerSubsystem.h"
 #include "WorldPartition/WorldPartitionSubsystem.h"
