Runtime side of Overridable Channels in Sequencer. Check MovieSceneTracks/Private/Sections/MovieScene3DTransformSection to see how to apply override logic in a section

Jira: UE-105300

diff --git a/Engine/Source/Runtime/MovieScene/Private/Channels/MovieSceneSectionChannelOverrideRegistry.cpp b/Engine/Source/Runtime/MovieScene/Private/Channels/MovieSceneSectionChannelOverrideRegistry.cpp
new file mode 100644
--- /dev/null
+++ b/Engine/Source/Runtime/MovieScene/Private/Channels/MovieSceneSectionChannelOverrideRegistry.cpp
@@ -0,0 +1,54 @@
+// Copyright Epic Games, Inc. All Rights Reserved.
+
+#include "Channels/MovieSceneSectionChannelOverrideRegistry.h"
+#include "Channels/MovieSceneChannelProxy.h"
+#include "Channels/MovieSceneChannelOverrideContainer.h"
+#include "EntitySystem/MovieSceneEntityBuilder.h"
+
+UMovieSceneSectionChannelOverrideRegistry::UMovieSceneSectionChannelOverrideRegistry()
+{
+}
+
+EMovieSceneChannelProxyType UMovieSceneSectionChannelOverrideRegistry::CacheChannelProxy()
+{
+	// TODO in Phase 2
+	return EMovieSceneChannelProxyType();
+}
+
+void UMovieSceneSectionChannelOverrideRegistry::AddChannel(int32 ChannelToOverride, UMovieSceneChannelOverrideContainer* ChannelContainer)
+{
+	Overrides.Emplace(ChannelToOverride, ChannelContainer);
+}
+
+void UMovieSceneSectionChannelOverrideRegistry::RemoveChannel(int32 ChannelToRemove)
+{
+	Overrides.Remove(ChannelToRemove);
+}
+
+bool UMovieSceneSectionChannelOverrideRegistry::IsOverriden(int32 ChannelIndex) const
+{
+	return Overrides.Contains(ChannelIndex);
+}
+
+void UMovieSceneSectionChannelOverrideRegistry::ImportEntityImpl(const UE::MovieScene::FEntityImportParams& Params, UE::MovieScene::FImportedEntity* OutImportedEntity)
+{
+	using namespace UE::MovieScene;
+;
+	int32 ChannelIndex = UMovieSceneSectionChannelOverrideRegistry::ToChannelIndex(Params.EntityID);
+	check(IsOverriden(ChannelIndex));	
+	Overrides[ChannelIndex]->ImportEntityImpl(ChannelIndex, Params, OutImportedEntity);
+}
+
+bool UMovieSceneSectionChannelOverrideRegistry::PopulateEvaluationFieldImpl(const TRange<FFrameNumber>& EffectiveRange, const FMovieSceneEvaluationFieldEntityMetaData& InMetaData, FMovieSceneEntityComponentFieldBuilder* OutFieldBuilder, UMovieSceneSection& OwnerSection)
+{
+	using namespace UE::MovieScene;
+
+	for (const TPair<int32, TObjectPtr<UMovieSceneChannelOverrideContainer>>& OverrideChannelPair : Overrides)
+	{
+		const int32 EntityIndex = OutFieldBuilder->FindOrAddEntity(&OwnerSection, ToEntityID(OverrideChannelPair.Key));
+		const int32 MetaDataIndex = OutFieldBuilder->AddMetaData(InMetaData);
+		OutFieldBuilder->AddPersistentEntity(EffectiveRange, EntityIndex, MetaDataIndex);
+	}
+
+	return false;
+}
diff --git a/Engine/Source/Runtime/MovieScene/Public/Channels/MovieSceneChannelOverrideContainer.h b/Engine/Source/Runtime/MovieScene/Public/Channels/MovieSceneChannelOverrideContainer.h
new file mode 100644
--- /dev/null
+++ b/Engine/Source/Runtime/MovieScene/Public/Channels/MovieSceneChannelOverrideContainer.h
@@ -0,0 +1,31 @@
+// Copyright Epic Games, Inc. All Rights Reserved.
+
+#pragma once
+#include "CoreMinimal.h"
+#include "UObject/ObjectMacros.h"
+#include "UObject/Object.h"
+#include "Evaluation/MovieSceneEvaluationField.h"
+#include "EntitySystem/IMovieSceneEntityProvider.h"
+#include "MovieSceneSection.h"
+#include "MovieSceneChannelOverrideContainer.generated.h"
+
+class UMovieSceneEntitySystemLinker;
+struct FMovieSceneEntityComponentFieldBuilder;
+struct FMovieSceneChannel;
+
+/**
+* A wrapper to implement polymorphism for FMovieSceneChannel. Meant to be override by children classes
+*/
+UCLASS(Abstract)
+class MOVIESCENE_API UMovieSceneChannelOverrideContainer : public UObject
+{
+	GENERATED_BODY()
+
+public:
+	/** Virtual methods to be overriden by children containers */
+	virtual void ImportEntityImpl(int32 ChannelIndex, const UE::MovieScene::FEntityImportParams& Params, UE::MovieScene::FImportedEntity* OutImportedEntity) {};
+	virtual EMovieSceneChannelProxyType CacheChannelProxy() { return EMovieSceneChannelProxyType(); }
+
+	virtual const FMovieSceneChannel* GetChannel() const { return nullptr; }
+	virtual FMovieSceneChannel* GetChannel() { return nullptr; }
+};
\ No newline at end of file
diff --git a/Engine/Source/Runtime/MovieScene/Public/Channels/MovieSceneSectionChannelOverrideRegistry.h b/Engine/Source/Runtime/MovieScene/Public/Channels/MovieSceneSectionChannelOverrideRegistry.h
new file mode 100644
--- /dev/null
+++ b/Engine/Source/Runtime/MovieScene/Public/Channels/MovieSceneSectionChannelOverrideRegistry.h
@@ -0,0 +1,131 @@
+// Copyright Epic Games, Inc. All Rights Reserved.
+
+#pragma once
+#include "CoreMinimal.h"
+#include "UObject/ObjectMacros.h"
+#include "UObject/Object.h"
+#include "MovieSceneSection.h"
+#include "Evaluation/MovieSceneEvaluationField.h"
+#include "EntitySystem/IMovieSceneEntityProvider.h"
+#include "MovieSceneSectionChannelOverrideRegistry.generated.h"
+
+class UMovieSceneChannelOverrideContainer;
+class UMovieSceneEntitySystemLinker;
+class UMovieScenePropertyTrack;
+struct FMovieSceneChannel;
+
+namespace UE
+{
+namespace MovieScene
+{
+	struct FImportedEntity;
+	struct FEntityImportParams;
+}
+}
+
+/**
+* This object contains a map of actual channel overrides, where each override is a channel identifier and a channel container.
+*/
+UCLASS()
+class MOVIESCENE_API UMovieSceneSectionChannelOverrideRegistry : public UObject
+{
+	GENERATED_BODY()
+
+	static constexpr int32 ImportEntityIDOffset = 10;
+
+private:
+	/** Map of channel overrides. */
+	UPROPERTY()
+	TMap<int32, TObjectPtr<UMovieSceneChannelOverrideContainer>> Overrides;
+
+public:
+	UMovieSceneSectionChannelOverrideRegistry();
+
+	/** 
+	* A way for all kinds of sections inform the sequencer editor what channels they have
+	*/
+	EMovieSceneChannelProxyType CacheChannelProxy();
+
+	/** 
+	* Add channel to the registry 
+	* @param ChannelToOverride		The key indicates which channel to override
+	* @param Data					The container that owns a overriden channel instanse
+	*/
+	void AddChannel(int32 ChannelToOverride, UMovieSceneChannelOverrideContainer* ChannelContainer);
+
+	/**
+	* Removes a channel from the registry
+	* @param ChannelToOverride		The key indicates which channel to override
+	*/
+	void RemoveChannel(int32 ChannelToRemove);
+
+	/**
+	* Forward ImportEntityImpl calls to overriden channels
+	*/
+	void ImportEntityImpl(const UE::MovieScene::FEntityImportParams& Params, UE::MovieScene::FImportedEntity* OutImportedEntity);
+
+	/**
+	* Called when overridden channels should populate evaluation field
+	*/
+	bool PopulateEvaluationFieldImpl(const TRange<FFrameNumber>& EffectiveRange, const FMovieSceneEvaluationFieldEntityMetaData& InMetaData, FMovieSceneEntityComponentFieldBuilder* OutFieldBuilder, UMovieSceneSection& OwnerSection);
+
+	/**
+	* From Channel Index to EntityID
+	* @param ChannelIndex	Represents the channel index in Overrides. Should be from 0 to 9
+	*/
+	static int32 ToEntityID(int32 ChannelIndex) { return ChannelIndex + ImportEntityIDOffset; }
+
+	/**
+	* From EntityID to channel index
+	* @param EntityID	Represents the EntityID.
+	*/
+	static int32 ToChannelIndex(int32 EntityID) { return EntityID - ImportEntityIDOffset; }
+
+	/**
+	* Returns if the channel is overriden
+	* @param ChannelIndex	Index for the channel.
+	* @return	If this channel is overriden
+	*/
+	bool IsOverriden(int32 ChannelIndex) const;
+
+	/**
+	* Returns the overriden map
+	*/
+	const TMap<int32, TObjectPtr<UMovieSceneChannelOverrideContainer>>& GetOverrides() const { return Overrides; }
+};
+
+namespace UE
+{
+namespace MovieScene 
+{
+inline bool IsChannelOverriden(const UMovieSceneSectionChannelOverrideRegistry* OverrideRegistry, int32 ChannelIndex)
+{
+	return OverrideRegistry && OverrideRegistry->IsOverriden(ChannelIndex);
+}
+
+template<typename ChannelType>
+inline bool HasAnyDataImpl(const UMovieSceneSectionChannelOverrideRegistry* OverrideRegistry, int32 Index, const ChannelType& InChannel)
+{
+	return InChannel.HasAnyData() || IsChannelOverriden(OverrideRegistry, Index);
+}
+
+template<typename HeadChannelType, typename ...TailChannelTypes>
+inline bool HasAnyDataImpl(const UMovieSceneSectionChannelOverrideRegistry* OverrideRegistry, int32 HeadIndex, const HeadChannelType& InHeadChannel, TailChannelTypes ...InTailChannels)
+{
+	if (InHeadChannel.HasAnyData() || IsChannelOverriden(OverrideRegistry, HeadIndex))
+	{
+		return true;
+	}
+	return HasAnyDataImpl(OverrideRegistry, HeadIndex + 1, InTailChannels...);
+}
+
+/**
+* Returns if at least one channel in the InChannels has any data. An example is HasAnyData(OverrideRegistry, 0, RedCurve, GreenCurve, BlueCurve, AlphaCurve)
+*/
+template<typename ...ChannelTypes>
+inline bool HasAnyData(const UMovieSceneSectionChannelOverrideRegistry* OverrideRegistry, int32 HeadIndex, ChannelTypes ...InChannels)
+{
+	return HasAnyDataImpl(OverrideRegistry, 0, InChannels...);
+}
+}
+}
\ No newline at end of file
diff --git a/Engine/Source/Runtime/MovieSceneTracks/Private/Channels/MovieSceneDoublePerlinNoiseChannel.cpp b/Engine/Source/Runtime/MovieSceneTracks/Private/Channels/MovieSceneDoublePerlinNoiseChannel.cpp
new file mode 100644
--- /dev/null
+++ b/Engine/Source/Runtime/MovieSceneTracks/Private/Channels/MovieSceneDoublePerlinNoiseChannel.cpp
@@ -0,0 +1,20 @@
+// Copyright Epic Games, Inc. All Rights Reserved.
+
+#include "Channels/MovieSceneDoublePerlinNoiseChannel.h"
+
+FMovieSceneDoublePerlinNoiseChannel::FMovieSceneDoublePerlinNoiseChannel()
+	: DoublePerlinNoiseParams{}
+{
+}
+
+FMovieSceneDoublePerlinNoiseChannel::FMovieSceneDoublePerlinNoiseChannel(const FDoublePerlinNoiseParams& InDoublePerlinNoiseParams)
+	: DoublePerlinNoiseParams{ InDoublePerlinNoiseParams }
+{
+}
+
+double FMovieSceneDoublePerlinNoiseChannel::Evaluate(double InSeconds)
+{
+	double Result = static_cast<double>(FMath::PerlinNoise1D(InSeconds * DoublePerlinNoiseParams.Frequency));
+	Result *= DoublePerlinNoiseParams.Amplitude;
+	return Result;
+}
diff --git a/Engine/Source/Runtime/MovieSceneTracks/Private/Channels/MovieSceneDoublePerlinNoiseChannelContainer.cpp b/Engine/Source/Runtime/MovieSceneTracks/Private/Channels/MovieSceneDoublePerlinNoiseChannelContainer.cpp
new file mode 100644
--- /dev/null
+++ b/Engine/Source/Runtime/MovieSceneTracks/Private/Channels/MovieSceneDoublePerlinNoiseChannelContainer.cpp
@@ -0,0 +1,26 @@
+// Copyright Epic Games, Inc. All Rights Reserved.
+
+#include "Channels/MovieSceneDoublePerlinNoiseChannelContainer.h"
+#include "Tracks/MovieScenePropertyTrack.h"
+#include "MovieSceneTracksComponentTypes.h"
+#include "Misc/Guid.h"
+
+void UMovieSceneDoublePerlinNoiseChannelContainer::ImportEntityImpl(int32 ChannelIndex, const UE::MovieScene::FEntityImportParams& Params, UE::MovieScene::FImportedEntity* OutImportedEntity)
+{
+	using namespace UE::MovieScene;
+
+	const FMovieSceneTracksComponentTypes* TracksComponents = FMovieSceneTracksComponentTypes::Get();
+	const FBuiltInComponentTypes* BuiltInComponents = FBuiltInComponentTypes::Get();
+
+	OutImportedEntity->AddBuilder(
+		FEntityBuilder()
+		.Add(TracksComponents->DoublePerlinNoiseChannel, DoublePerlinNoiseChannel.GetParam())
+		.Add(BuiltInComponents->DoubleResult[ChannelIndex], TNumericLimits<double>::Max())
+	);
+}
+
+EMovieSceneChannelProxyType UMovieSceneDoublePerlinNoiseChannelContainer::CacheChannelProxy()
+{
+	//TODO in Phase 2
+	return EMovieSceneChannelProxyType();
+}
diff --git a/Engine/Source/Runtime/MovieSceneTracks/Private/Channels/MovieSceneFloatPerlinNoiseChannel.cpp b/Engine/Source/Runtime/MovieSceneTracks/Private/Channels/MovieSceneFloatPerlinNoiseChannel.cpp
new file mode 100644
--- /dev/null
+++ b/Engine/Source/Runtime/MovieSceneTracks/Private/Channels/MovieSceneFloatPerlinNoiseChannel.cpp
@@ -0,0 +1,20 @@
+// Copyright Epic Games, Inc. All Rights Reserved.
+
+#include "Channels/MovieSceneFloatPerlinNoiseChannel.h"
+
+FMovieSceneFloatPerlinNoiseChannel::FMovieSceneFloatPerlinNoiseChannel()
+	: FloatPerlinNoiseParams{}
+{
+}
+
+FMovieSceneFloatPerlinNoiseChannel::FMovieSceneFloatPerlinNoiseChannel(const FFloatPerlinNoiseParams& InFloatPerlinNoiseParams)
+	: FloatPerlinNoiseParams{ InFloatPerlinNoiseParams }
+{
+}
+
+float FMovieSceneFloatPerlinNoiseChannel::Evaluate(double InSeconds)
+{
+	float Result = FMath::PerlinNoise1D(InSeconds * FloatPerlinNoiseParams.Frequency);
+	Result *= FloatPerlinNoiseParams.Amplitude;
+	return Result;
+}
\ No newline at end of file
diff --git a/Engine/Source/Runtime/MovieSceneTracks/Private/Channels/MovieSceneFloatPerlinNoiseChannelContainer.cpp b/Engine/Source/Runtime/MovieSceneTracks/Private/Channels/MovieSceneFloatPerlinNoiseChannelContainer.cpp
new file mode 100644
--- /dev/null
+++ b/Engine/Source/Runtime/MovieSceneTracks/Private/Channels/MovieSceneFloatPerlinNoiseChannelContainer.cpp
@@ -0,0 +1,26 @@
+// Copyright Epic Games, Inc. All Rights Reserved.
+
+#include "Channels/MovieSceneFloatPerlinNoiseChannelContainer.h"
+#include "Tracks/MovieScenePropertyTrack.h"
+#include "MovieSceneTracksComponentTypes.h"
+#include "Misc/Guid.h"
+
+void UMovieSceneFloatPerlinNoiseChannelContainer::ImportEntityImpl(int32 ChannelIndex, const UE::MovieScene::FEntityImportParams& Params, UE::MovieScene::FImportedEntity* OutImportedEntity)
+{
+	using namespace UE::MovieScene;
+
+	const FMovieSceneTracksComponentTypes* TracksComponents = FMovieSceneTracksComponentTypes::Get();
+	const FBuiltInComponentTypes* BuiltInComponents = FBuiltInComponentTypes::Get();
+
+	OutImportedEntity->AddBuilder(
+		FEntityBuilder()
+		.Add(TracksComponents->FloatPerlinNoiseChannel, FloatPerlinNoiseChannel.GetParam())
+		.Add(BuiltInComponents->FloatResult[ChannelIndex], TNumericLimits<float>::Max())
+	);
+}
+
+EMovieSceneChannelProxyType UMovieSceneFloatPerlinNoiseChannelContainer::CacheChannelProxy()
+{
+	//TODO in Phase 2
+	return EMovieSceneChannelProxyType();
+}
diff --git a/Engine/Source/Runtime/MovieSceneTracks/Private/MovieSceneTracksComponentTypes.cpp b/Engine/Source/Runtime/MovieSceneTracks/Private/MovieSceneTracksComponentTypes.cpp
--- a/Engine/Source/Runtime/MovieSceneTracks/Private/MovieSceneTracksComponentTypes.cpp
+++ b/Engine/Source/Runtime/MovieSceneTracks/Private/MovieSceneTracksComponentTypes.cpp
@@ -498,6 +498,8 @@ FMovieSceneTracksComponentTypes::FMovieSceneTracksComponentTypes()
 	ComponentRegistry->NewComponentType(&AttachParent, TEXT("Attach Parent"));
 	ComponentRegistry->NewComponentType(&AttachComponent, TEXT("Attachment Component"));
 	ComponentRegistry->NewComponentType(&AttachParentBinding, TEXT("Attach Parent Binding"));
+	ComponentRegistry->NewComponentType(&FloatPerlinNoiseChannel, TEXT("Float Perlin Noise Channel"));
+	ComponentRegistry->NewComponentType(&DoublePerlinNoiseChannel, TEXT("Double Perlin Noise Channel"));
 
 	ComponentRegistry->NewComponentType(&LevelVisibility, TEXT("Level Visibility"));
 	ComponentRegistry->NewComponentType(&DataLayer, TEXT("Data Layer"));
@@ -668,6 +670,11 @@ FMovieSceneTracksComponentTypes::FMovieSceneTracksComponentTypes()
 
 	ComponentRegistry->Factories.DuplicateChildComponent(AttachParentBinding);
 	ComponentRegistry->Factories.DuplicateChildComponent(AttachComponent);
+
+	// --------------------------------------------------------------------------------------------
+	// Set up PerlinNoise components
+	ComponentRegistry->Factories.DuplicateChildComponent(FloatPerlinNoiseChannel);
+	ComponentRegistry->Factories.DuplicateChildComponent(DoublePerlinNoiseChannel);
 }
 
 FMovieSceneTracksComponentTypes::~FMovieSceneTracksComponentTypes()
@@ -690,6 +697,29 @@ FMovieSceneTracksComponentTypes* FMovieSceneTracksComponentTypes::Get()
 	return GMovieSceneTracksComponentTypes.Get();
 }
 
-
 } // namespace MovieScene
 } // namespace UE
+
+FFloatPerlinNoiseParams::FFloatPerlinNoiseParams()
+	: Frequency(4.0f)
+	, Amplitude(1.0f)
+{
+}
+
+FFloatPerlinNoiseParams::FFloatPerlinNoiseParams(float InFrequency, float InAmplitude)
+	: Frequency(InFrequency)
+	, Amplitude(InAmplitude)
+{
+}
+
+FDoublePerlinNoiseParams::FDoublePerlinNoiseParams()
+	: Frequency(4.0)
+	, Amplitude(1.0)
+{
+}
+
+FDoublePerlinNoiseParams::FDoublePerlinNoiseParams(double InFrequency, double InAmplitude)
+	: Frequency(InFrequency)
+	, Amplitude(InAmplitude)
+{
+}
diff --git a/Engine/Source/Runtime/MovieSceneTracks/Private/Sections/MovieScene3DTransformSection.cpp b/Engine/Source/Runtime/MovieSceneTracks/Private/Sections/MovieScene3DTransformSection.cpp
--- a/Engine/Source/Runtime/MovieSceneTracks/Private/Sections/MovieScene3DTransformSection.cpp
+++ b/Engine/Source/Runtime/MovieSceneTracks/Private/Sections/MovieScene3DTransformSection.cpp
@@ -333,7 +333,7 @@ UMovieScene3DTransformSection::UMovieScene3DTransformSection(const FObjectInitia
 }
 
 template<typename BaseBuilderType>
-void UMovieScene3DTransformSection::BuildEntity(BaseBuilderType& InBaseBuilder, UMovieSceneEntitySystemLinker* EntityLinker, const FEntityImportParams& Params, FImportedEntity* OutImportedEntity)
+void UMovieScene3DTransformSection::BuildEntity(BaseBuilderType& InBaseBuilder, UMovieSceneEntitySystemLinker* Linker, const FEntityImportParams& Params, FImportedEntity* OutImportedEntity)
 {
 
 	using namespace UE::MovieScene;
@@ -390,22 +390,45 @@ void UMovieScene3DTransformSection::BuildEntity(BaseBuilderType& InBaseBuilder,
 		RotationChannel[2] = TrackComponents->QuaternionRotationChannel[2];
 	}
 
+	// Let override registry handles overriden channel here
+	int32 ChannelIndex = UMovieSceneSectionChannelOverrideRegistry::ToChannelIndex(Params.EntityID);
+
+	// If we are building the entity of an overriden channel, we pass the partial builder we have at this point (the one with the bindings) and add the tag to it.
+	// Then call into the override repository so it can add a second builder with the override stuff.
+	if (IsChannelOverriden(OverrideRegistry, ChannelIndex))
+	{
+		OutImportedEntity->AddBuilder(InBaseBuilder.AddTag(PropertyTag));
+		OverrideRegistry->ImportEntityImpl(Params, OutImportedEntity);
+		return;
+	}	
+	
+	// Here proceed with the all the normal channels
 	OutImportedEntity->AddBuilder(
 		InBaseBuilder
-		.AddConditional(BuiltInComponentTypes->DoubleChannel[0], &Translation[0], ActiveChannelsMask[0])
-		.AddConditional(BuiltInComponentTypes->DoubleChannel[1], &Translation[1], ActiveChannelsMask[1])
-		.AddConditional(BuiltInComponentTypes->DoubleChannel[2], &Translation[2], ActiveChannelsMask[2])
-		.AddConditional(RotationChannel[0],                      &Rotation[0],    ActiveChannelsMask[3])
-		.AddConditional(RotationChannel[1],                      &Rotation[1],    ActiveChannelsMask[4])
-		.AddConditional(RotationChannel[2],                      &Rotation[2],    ActiveChannelsMask[5])
-		.AddConditional(BuiltInComponentTypes->DoubleChannel[6], &Scale[0],       ActiveChannelsMask[6])
-		.AddConditional(BuiltInComponentTypes->DoubleChannel[7], &Scale[1],       ActiveChannelsMask[7])
-		.AddConditional(BuiltInComponentTypes->DoubleChannel[8], &Scale[2],       ActiveChannelsMask[8])
-		.AddConditional(BuiltInComponentTypes->WeightChannel,    &ManualWeight,   EnumHasAnyFlags(Channels, EMovieSceneTransformChannel::Weight) && ManualWeight.HasAnyData())
+		.AddConditional(BuiltInComponentTypes->DoubleChannel[0], &Translation[0], ActiveChannelsMask[0] && !IsChannelOverriden(OverrideRegistry, 0))
+		.AddConditional(BuiltInComponentTypes->DoubleChannel[1], &Translation[1], ActiveChannelsMask[1] && !IsChannelOverriden(OverrideRegistry, 1))
+		.AddConditional(BuiltInComponentTypes->DoubleChannel[2], &Translation[2], ActiveChannelsMask[2] && !IsChannelOverriden(OverrideRegistry, 2))
+		.AddConditional(RotationChannel[0],                      &Rotation[0],    ActiveChannelsMask[3] && !IsChannelOverriden(OverrideRegistry, 3))
+		.AddConditional(RotationChannel[1],                      &Rotation[1],    ActiveChannelsMask[4] && !IsChannelOverriden(OverrideRegistry, 4))
+		.AddConditional(RotationChannel[2],                       &Rotation[2],   ActiveChannelsMask[5] && !IsChannelOverriden(OverrideRegistry, 5))
+		.AddConditional(BuiltInComponentTypes->DoubleChannel[6], &Scale[0],       ActiveChannelsMask[6] && !IsChannelOverriden(OverrideRegistry, 6))
+		.AddConditional(BuiltInComponentTypes->DoubleChannel[7], &Scale[1],       ActiveChannelsMask[7] && !IsChannelOverriden(OverrideRegistry, 7))
+		.AddConditional(BuiltInComponentTypes->DoubleChannel[8], &Scale[2],       ActiveChannelsMask[8] && !IsChannelOverriden(OverrideRegistry, 8))
+		.AddConditional(BuiltInComponentTypes->WeightChannel,    &ManualWeight,   EnumHasAnyFlags(Channels, EMovieSceneTransformChannel::Weight) && ManualWeight.HasAnyData() && !IsChannelOverriden(OverrideRegistry, 9))
 		.AddTag(PropertyTag)
 	);
 }
 
+bool UMovieScene3DTransformSection::PopulateEvaluationFieldImpl(const TRange<FFrameNumber>& EffectiveRange, const FMovieSceneEvaluationFieldEntityMetaData& InMetaData, FMovieSceneEntityComponentFieldBuilder* OutFieldBuilder)
+{
+	if (OverrideRegistry)
+	{
+		return OverrideRegistry->PopulateEvaluationFieldImpl(EffectiveRange, InMetaData, OutFieldBuilder, *this);
+	}
+
+	return false;
+}
+
 void UMovieScene3DTransformSection::ImportEntityImpl(UMovieSceneEntitySystemLinker* EntityLinker, const FEntityImportParams& Params, FImportedEntity* OutImportedEntity)
 {
 	using namespace UE::MovieScene;
diff --git a/Engine/Source/Runtime/MovieSceneTracks/Private/Systems/DoublePerlinNoiseChannelEvaluatorSystem.cpp b/Engine/Source/Runtime/MovieSceneTracks/Private/Systems/DoublePerlinNoiseChannelEvaluatorSystem.cpp
new file mode 100644
--- /dev/null
+++ b/Engine/Source/Runtime/MovieSceneTracks/Private/Systems/DoublePerlinNoiseChannelEvaluatorSystem.cpp
@@ -0,0 +1,83 @@
+// Copyright Epic Games, Inc. All Rights Reserved.
+
+#include "Systems/DoublePerlinNoiseChannelEvaluatorSystem.h"
+#include "Systems/MovieScenePiecewiseDoubleBlenderSystem.h"
+#include "EntitySystem/BuiltInComponentTypes.h"
+#include "EntitySystem/MovieSceneEntitySystemTask.h"
+#include "EntitySystem/MovieSceneEntitySystemLinker.h"
+#include "EntitySystem/MovieSceneEvalTimeSystem.h"
+#include "Channels/MovieSceneDoublePerlinNoiseChannel.h"
+#include "MovieSceneTracksComponentTypes.h"
+#include "Math/NumericLimits.h"
+
+DECLARE_CYCLE_STAT(TEXT("MovieScene: Evaluate Double Perlin Noise channels"), MovieSceneEval_EvaluateDoublePerlinNoiseChannelTask, STATGROUP_MovieSceneECS);
+
+namespace UE
+{
+	namespace MovieScene
+	{
+		struct FEvaluateDoublePerlinNoiseChannels
+		{
+			FInstanceRegistry* InstanceRegistry;
+
+			FEvaluateDoublePerlinNoiseChannels(FInstanceRegistry* InInstanceRegistry)
+				: InstanceRegistry{ InInstanceRegistry }
+			{
+				check(InstanceRegistry);
+			}
+
+			void ForEachEntity(FMovieSceneDoublePerlinNoiseChannel DoublePerlinNoiseChannel, FInstanceHandle InstanceHandle, double& OutResult)
+			{
+				const FMovieSceneContext& Context = InstanceRegistry->GetContext(InstanceHandle);
+				FFrameTime Time = Context.GetTime();
+				FFrameRate Rate = Context.GetFrameRate();
+
+				OutResult = DoublePerlinNoiseChannel.Evaluate(Rate.AsSeconds(Time));
+			}
+		};
+
+
+	} // namespace MovieScene
+} // namespace UE
+
+UDoublePerlinNoiseChannelEvaluatorSystem::UDoublePerlinNoiseChannelEvaluatorSystem(const FObjectInitializer& ObjInit)
+	: Super(ObjInit)
+{
+	using namespace UE::MovieScene;
+
+	const FBuiltInComponentTypes* BuiltInComponents = FBuiltInComponentTypes::Get();
+	const FMovieSceneTracksComponentTypes* TrackComponents = FMovieSceneTracksComponentTypes::Get();
+
+	RelevantComponent = TrackComponents->DoublePerlinNoiseChannel;
+
+	if (HasAnyFlags(RF_ClassDefaultObject))
+	{
+		// Allow writing to all the possible channels
+		for (int32 i = 0; i < sizeof(BuiltInComponents->DoubleResult); ++i)
+		{
+			DefineComponentProducer(UDoublePerlinNoiseChannelEvaluatorSystem::StaticClass(), BuiltInComponents->DoubleResult[i]);
+		}
+
+		DefineImplicitPrerequisite(UMovieSceneEvalTimeSystem::StaticClass(), GetClass());
+		DefineImplicitPrerequisite(GetClass(), UMovieScenePiecewiseDoubleBlenderSystem::StaticClass());
+	}
+}
+
+void UDoublePerlinNoiseChannelEvaluatorSystem::OnRun(FSystemTaskPrerequisites& InPrerequisites, FSystemSubsequentTasks& Subsequents)
+{
+	using namespace UE::MovieScene;
+
+	const FBuiltInComponentTypes* BuiltInComponents = FBuiltInComponentTypes::Get();
+	const FMovieSceneTracksComponentTypes* TrackComponents = FMovieSceneTracksComponentTypes::Get();
+
+	for (int32 i = 0; i < sizeof(BuiltInComponents->DoubleResult); ++i)
+	{
+		FEntityTaskBuilder()
+			.Read(TrackComponents->DoublePerlinNoiseChannel)
+			.Read(BuiltInComponents->InstanceHandle)
+			.Write(BuiltInComponents->DoubleResult[i])
+			.FilterNone({ BuiltInComponents->Tags.Ignored })
+			.SetStat(GET_STATID(MovieSceneEval_EvaluateDoublePerlinNoiseChannelTask))
+			.Dispatch_PerEntity<FEvaluateDoublePerlinNoiseChannels>(&Linker->EntityManager, InPrerequisites, &Subsequents, GetLinker()->GetInstanceRegistry());
+	}
+}
diff --git a/Engine/Source/Runtime/MovieSceneTracks/Private/Systems/FloatPerlinNoiseChannelEvaluatorSystem.cpp b/Engine/Source/Runtime/MovieSceneTracks/Private/Systems/FloatPerlinNoiseChannelEvaluatorSystem.cpp
new file mode 100644
--- /dev/null
+++ b/Engine/Source/Runtime/MovieSceneTracks/Private/Systems/FloatPerlinNoiseChannelEvaluatorSystem.cpp
@@ -0,0 +1,83 @@
+// Copyright Epic Games, Inc. All Rights Reserved.
+
+#include "Systems/FloatPerlinNoiseChannelEvaluatorSystem.h"
+#include "Systems/MovieScenePiecewiseFloatBlenderSystem.h"
+#include "EntitySystem/BuiltInComponentTypes.h"
+#include "EntitySystem/MovieSceneEntitySystemTask.h"
+#include "EntitySystem/MovieSceneEntitySystemLinker.h"
+#include "EntitySystem/MovieSceneEvalTimeSystem.h"
+#include "Channels/MovieSceneFloatPerlinNoiseChannel.h"
+#include "MovieSceneTracksComponentTypes.h"
+#include "Math/NumericLimits.h"
+
+DECLARE_CYCLE_STAT(TEXT("MovieScene: Evaluate Float Perlin Noise channels"), MovieSceneEval_EvaluateFloatPerlinNoiseChannelTask, STATGROUP_MovieSceneECS);
+
+namespace UE
+{
+	namespace MovieScene
+	{
+		struct FEvaluateFloatPerlinNoiseChannels
+		{
+			FInstanceRegistry* InstanceRegistry;
+
+			FEvaluateFloatPerlinNoiseChannels(FInstanceRegistry* InInstanceRegistry)
+				: InstanceRegistry{ InInstanceRegistry }
+			{
+				check(InstanceRegistry);
+			}
+
+			void ForEachEntity(FMovieSceneFloatPerlinNoiseChannel FloatPerlinNoiseChannel, FInstanceHandle InstanceHandle, float& OutResult)
+			{
+				const FMovieSceneContext& Context = InstanceRegistry->GetContext(InstanceHandle);
+				FFrameTime Time = Context.GetTime();
+				FFrameRate Rate = Context.GetFrameRate();
+
+				OutResult = FloatPerlinNoiseChannel.Evaluate(Rate.AsSeconds(Time));
+			}
+		};
+
+
+	} // namespace MovieScene
+} // namespace UE
+
+UFloatPerlinNoiseChannelEvaluatorSystem::UFloatPerlinNoiseChannelEvaluatorSystem(const FObjectInitializer& ObjInit)
+	: Super(ObjInit)
+{
+	using namespace UE::MovieScene;
+
+	const FBuiltInComponentTypes* BuiltInComponents = FBuiltInComponentTypes::Get();
+	const FMovieSceneTracksComponentTypes* TrackComponents = FMovieSceneTracksComponentTypes::Get();
+
+	RelevantComponent = TrackComponents->FloatPerlinNoiseChannel;
+
+	if (HasAnyFlags(RF_ClassDefaultObject))
+	{
+		// Allow writing to all the possible channels
+		for (int32 i = 0; i < sizeof(BuiltInComponents->FloatResult); ++i)
+		{
+			DefineComponentProducer(UFloatPerlinNoiseChannelEvaluatorSystem::StaticClass(), BuiltInComponents->FloatResult[i]);
+		}
+
+		DefineImplicitPrerequisite(UMovieSceneEvalTimeSystem::StaticClass(), GetClass());
+		DefineImplicitPrerequisite(GetClass(), UMovieScenePiecewiseFloatBlenderSystem::StaticClass());
+	}
+}
+
+void UFloatPerlinNoiseChannelEvaluatorSystem::OnRun(FSystemTaskPrerequisites& InPrerequisites, FSystemSubsequentTasks& Subsequents)
+{
+	using namespace UE::MovieScene;
+
+	const FBuiltInComponentTypes* BuiltInComponents = FBuiltInComponentTypes::Get();
+	const FMovieSceneTracksComponentTypes* TrackComponents = FMovieSceneTracksComponentTypes::Get();
+
+	for (int32 i = 0; i < sizeof(BuiltInComponents->FloatResult); ++i)
+	{
+		FEntityTaskBuilder()
+			.Read(TrackComponents->FloatPerlinNoiseChannel)
+			.Read(BuiltInComponents->InstanceHandle)
+			.Write(BuiltInComponents->FloatResult[i])
+			.FilterNone({ BuiltInComponents->Tags.Ignored })
+			.SetStat(GET_STATID(MovieSceneEval_EvaluateFloatPerlinNoiseChannelTask))
+			.Dispatch_PerEntity<FEvaluateFloatPerlinNoiseChannels>(&Linker->EntityManager, InPrerequisites, &Subsequents, GetLinker()->GetInstanceRegistry());
+	}
+}
diff --git a/Engine/Source/Runtime/MovieSceneTracks/Private/Tracks/MovieScenePropertyTrack.cpp b/Engine/Source/Runtime/MovieSceneTracks/Private/Tracks/MovieScenePropertyTrack.cpp
--- a/Engine/Source/Runtime/MovieSceneTracks/Private/Tracks/MovieScenePropertyTrack.cpp
+++ b/Engine/Source/Runtime/MovieSceneTracks/Private/Tracks/MovieScenePropertyTrack.cpp
@@ -1,6 +1,7 @@
 // Copyright Epic Games, Inc. All Rights Reserved.
 
 #include "Tracks/MovieScenePropertyTrack.h"
+#include "Channels/MovieSceneSectionChannelOverrideRegistry.h"
 #include "Algo/Sort.h"
 #include "MovieSceneCommonHelpers.h"
 #include "MovieSceneTracksComponentTypes.h"
@@ -330,10 +331,15 @@ UMovieSceneSection* UMovieScenePropertyTrack::GetSectionToKey() const
 const int32 FMovieScenePropertyTrackEntityImportHelper::SectionPropertyValueImportingID = 0;
 const int32 FMovieScenePropertyTrackEntityImportHelper::SectionEditConditionToggleImportingID = 1;
 
-void FMovieScenePropertyTrackEntityImportHelper::PopulateEvaluationField(UMovieSceneSection& Section, const TRange<FFrameNumber>& EffectiveRange, const FMovieSceneEvaluationFieldEntityMetaData& InMetaData, FMovieSceneEntityComponentFieldBuilder* OutFieldBuilder)
+void FMovieScenePropertyTrackEntityImportHelper::PopulateEvaluationField(UMovieSceneSection& Section, const TRange<FFrameNumber>& EffectiveRange, const FMovieSceneEvaluationFieldEntityMetaData& InMetaData, FMovieSceneEntityComponentFieldBuilder* OutFieldBuilder, UMovieSceneSectionChannelOverrideRegistry* OverridesRegistry/*= nullptr*/)
 {
 	using namespace UE::MovieScene;
 
+	if (OverridesRegistry)
+	{
+		OverridesRegistry->PopulateEvaluationFieldImpl(EffectiveRange, InMetaData, OutFieldBuilder, Section);
+	}
+
 	// Add the default entity for this section.
 	const int32 EntityIndex   = OutFieldBuilder->FindOrAddEntity(&Section, SectionPropertyValueImportingID);
 	const int32 MetaDataIndex = OutFieldBuilder->AddMetaData(InMetaData);
diff --git a/Engine/Source/Runtime/MovieSceneTracks/Public/Channels/MovieSceneDoublePerlinNoiseChannel.h b/Engine/Source/Runtime/MovieSceneTracks/Public/Channels/MovieSceneDoublePerlinNoiseChannel.h
new file mode 100644
--- /dev/null
+++ b/Engine/Source/Runtime/MovieSceneTracks/Public/Channels/MovieSceneDoublePerlinNoiseChannel.h
@@ -0,0 +1,44 @@
+// Copyright Epic Games, Inc. All Rights Reserved.
+
+#pragma once
+#include "Channels/MovieSceneChannel.h"
+#include "Channels/MovieSceneChannelData.h"
+#include "Channels/MovieSceneChannelTraits.h"
+#include "MovieSceneTracksComponentTypes.h"
+#include "MovieSceneDoublePerlinNoiseChannel.generated.h"
+
+USTRUCT()
+struct MOVIESCENETRACKS_API FMovieSceneDoublePerlinNoiseChannel : public FMovieSceneChannel
+{
+	GENERATED_BODY()
+
+	using CurveValueType = double;
+
+	FMovieSceneDoublePerlinNoiseChannel();
+	FMovieSceneDoublePerlinNoiseChannel(const FDoublePerlinNoiseParams& InDoublePerlinNoiseParams);
+
+	/**
+	* Evaluate this channel with the frame resolution
+	*
+	* @param InSeconds  The Frame second to evaluate at
+	* @return A value to receive the PerlinNoise result
+	*/
+	double Evaluate(double InSeconds);
+
+	const FDoublePerlinNoiseParams& GetParam() const { return DoublePerlinNoiseParams; }
+	FDoublePerlinNoiseParams& GetParam() { return DoublePerlinNoiseParams; }
+
+private:
+	FDoublePerlinNoiseParams DoublePerlinNoiseParams;
+};
+
+template<>
+struct TMovieSceneChannelTraits<FMovieSceneDoublePerlinNoiseChannel> : TMovieSceneChannelTraitsBase<FMovieSceneDoublePerlinNoiseChannel>
+{
+#if WITH_EDITOR
+
+	/** FMovieSceneFloatPerlinNoiseChannel channels can have external values (ie, they can get their values from external objects for UI purposes) */
+	typedef TMovieSceneExternalValue<double> ExtendedEditorDataType;
+
+#endif
+};
diff --git a/Engine/Source/Runtime/MovieSceneTracks/Public/Channels/MovieSceneDoublePerlinNoiseChannelContainer.h b/Engine/Source/Runtime/MovieSceneTracks/Public/Channels/MovieSceneDoublePerlinNoiseChannelContainer.h
new file mode 100644
--- /dev/null
+++ b/Engine/Source/Runtime/MovieSceneTracks/Public/Channels/MovieSceneDoublePerlinNoiseChannelContainer.h
@@ -0,0 +1,26 @@
+// Copyright Epic Games, Inc. All Rights Reserved.
+
+#pragma once
+#include "Channels/MovieSceneChannelOverrideContainer.h"
+#include "Channels/MovieSceneDoublePerlinNoiseChannel.h"
+#include "MovieSceneDoublePerlinNoiseChannelContainer.generated.h"
+
+/**
+* Double perlin noise channel overriden container
+*/
+UCLASS(meta = (DisplayName = "Double Perlin Noise", ToolTip = "Override a channel to use double perlin noise"))
+class  MOVIESCENETRACKS_API UMovieSceneDoublePerlinNoiseChannelContainer : public UMovieSceneChannelOverrideContainer
+{
+	GENERATED_BODY()
+
+private:
+	UPROPERTY()
+	FMovieSceneDoublePerlinNoiseChannel DoublePerlinNoiseChannel;
+
+public:
+	virtual void ImportEntityImpl(int32 ChannelIndex, const UE::MovieScene::FEntityImportParams& Params, UE::MovieScene::FImportedEntity* OutImportedEntity) override final;
+	virtual EMovieSceneChannelProxyType CacheChannelProxy() override final;
+
+	virtual const FMovieSceneChannel* GetChannel() const override final { return &DoublePerlinNoiseChannel; }
+	virtual FMovieSceneChannel* GetChannel() override final { return &DoublePerlinNoiseChannel; }
+};
diff --git a/Engine/Source/Runtime/MovieSceneTracks/Public/Channels/MovieSceneFloatPerlinNoiseChannel.h b/Engine/Source/Runtime/MovieSceneTracks/Public/Channels/MovieSceneFloatPerlinNoiseChannel.h
new file mode 100644
--- /dev/null
+++ b/Engine/Source/Runtime/MovieSceneTracks/Public/Channels/MovieSceneFloatPerlinNoiseChannel.h
@@ -0,0 +1,44 @@
+// Copyright Epic Games, Inc. All Rights Reserved.
+
+#pragma once
+#include "Channels/MovieSceneChannel.h"
+#include "Channels/MovieSceneChannelData.h"
+#include "Channels/MovieSceneChannelTraits.h"
+#include "MovieSceneTracksComponentTypes.h"
+#include "MovieSceneFloatPerlinNoiseChannel.generated.h"
+
+USTRUCT()
+struct MOVIESCENETRACKS_API FMovieSceneFloatPerlinNoiseChannel : public FMovieSceneChannel
+{
+	GENERATED_BODY()
+
+	using CurveValueType = float;
+
+	FMovieSceneFloatPerlinNoiseChannel();
+	FMovieSceneFloatPerlinNoiseChannel(const FFloatPerlinNoiseParams& InFloatPerlinNoiseParams);
+
+	/**
+	* Evaluate this channel with the frame resolution
+	*
+	* @param InSeconds  The Frame second to evaluate at
+	* @return A value to receive the PerlinNoise result
+	*/
+	float Evaluate(double InSeconds);
+
+	FFloatPerlinNoiseParams& GetParam() { return FloatPerlinNoiseParams; }
+	const FFloatPerlinNoiseParams& GetParam() const { return FloatPerlinNoiseParams; }
+
+private:
+	FFloatPerlinNoiseParams FloatPerlinNoiseParams;
+};
+
+template<>
+struct TMovieSceneChannelTraits<FMovieSceneFloatPerlinNoiseChannel> : TMovieSceneChannelTraitsBase<FMovieSceneFloatPerlinNoiseChannel>
+{
+#if WITH_EDITOR
+
+	/** FMovieSceneFloatPerlinNoiseChannel channels can have external values (ie, they can get their values from external objects for UI purposes) */
+	typedef TMovieSceneExternalValue<float> ExtendedEditorDataType;
+
+#endif
+};
diff --git a/Engine/Source/Runtime/MovieSceneTracks/Public/Channels/MovieSceneFloatPerlinNoiseChannelContainer.h b/Engine/Source/Runtime/MovieSceneTracks/Public/Channels/MovieSceneFloatPerlinNoiseChannelContainer.h
new file mode 100644
--- /dev/null
+++ b/Engine/Source/Runtime/MovieSceneTracks/Public/Channels/MovieSceneFloatPerlinNoiseChannelContainer.h
@@ -0,0 +1,26 @@
+// Copyright Epic Games, Inc. All Rights Reserved.
+
+#pragma once
+#include "Channels/MovieSceneChannelOverrideContainer.h"
+#include "Channels/MovieSceneFloatPerlinNoiseChannel.h"
+#include "MovieSceneFloatPerlinNoiseChannelContainer.generated.h"
+
+/**
+* Float perlin noise channel overriden container
+*/
+UCLASS(meta = (DisplayName = "Float Perlin Noise", ToolTip = "Override a channel to use float perlin noise"))
+class MOVIESCENETRACKS_API UMovieSceneFloatPerlinNoiseChannelContainer : public UMovieSceneChannelOverrideContainer
+{
+	GENERATED_BODY()
+
+private:
+	UPROPERTY()
+	FMovieSceneFloatPerlinNoiseChannel FloatPerlinNoiseChannel;
+
+public:
+	virtual void ImportEntityImpl(int32 ChannelIndex, const UE::MovieScene::FEntityImportParams& Params, UE::MovieScene::FImportedEntity* OutImportedEntity) override final;
+	virtual EMovieSceneChannelProxyType CacheChannelProxy() override final;
+
+	virtual const FMovieSceneChannel* GetChannel() const override final { return &FloatPerlinNoiseChannel; }
+	virtual FMovieSceneChannel* GetChannel() override final { return &FloatPerlinNoiseChannel; }
+};
diff --git a/Engine/Source/Runtime/MovieSceneTracks/Public/MovieSceneTracksComponentTypes.h b/Engine/Source/Runtime/MovieSceneTracks/Public/MovieSceneTracksComponentTypes.h
--- a/Engine/Source/Runtime/MovieSceneTracks/Public/MovieSceneTracksComponentTypes.h
+++ b/Engine/Source/Runtime/MovieSceneTracks/Public/MovieSceneTracksComponentTypes.h
@@ -6,11 +6,11 @@
 #include "EntitySystem/BuiltInComponentTypes.h"
 #include "EntitySystem/MovieSceneEntityIDs.h"
 #include "EntitySystem/MovieScenePropertySystemTypes.h"
+#include "EntitySystem/MovieScenePropertyMetaDataTraits.h"
+#include "EntitySystem/MovieScenePropertyTraits.h"
 #include "Engine/EngineTypes.h"
 #include "EulerTransform.h"
 #include "TransformData.h"
-#include "EntitySystem/MovieScenePropertyTraits.h"
-#include "EntitySystem/MovieScenePropertyMetaDataTraits.h"
 #include "MovieSceneTracksPropertyTypes.h"
 #include "Styling/SlateColor.h"
 #include "MovieSceneTracksComponentTypes.generated.h"
@@ -19,6 +19,39 @@ class UMovieSceneDataLayerSection;
 class UMovieSceneLevelVisibilitySection;
 struct FMovieSceneObjectBindingID;
 
+
+/** Component data for the Float Perlin Noise Channel */
+USTRUCT()
+struct MOVIESCENETRACKS_API FFloatPerlinNoiseParams
+{
+	GENERATED_BODY()
+
+	UPROPERTY()
+	float Frequency;
+
+	UPROPERTY()
+	float Amplitude;
+
+	FFloatPerlinNoiseParams();
+	FFloatPerlinNoiseParams(float InFrequency, float InAmplitude);
+};
+
+/** Component data for the Float Perlin Noise Channel */
+USTRUCT()
+struct MOVIESCENETRACKS_API FDoublePerlinNoiseParams
+{
+	GENERATED_BODY()
+
+	UPROPERTY()
+	double Frequency;
+
+	UPROPERTY()
+	double Amplitude;
+
+	FDoublePerlinNoiseParams();
+	FDoublePerlinNoiseParams(double InFrequency, double InAmplitude);
+};
+
 /** Component data for the level visibility system */
 USTRUCT()
 struct FLevelVisibilityComponentData
@@ -312,6 +345,8 @@ struct MOVIESCENETRACKS_API FMovieSceneTracksComponentTypes
 	TComponentTypeID<USceneComponent*> AttachParent;
 	TComponentTypeID<FAttachmentComponent> AttachComponent;
 	TComponentTypeID<FMovieSceneObjectBindingID> AttachParentBinding;
+	TComponentTypeID<FFloatPerlinNoiseParams> FloatPerlinNoiseChannel;
+	TComponentTypeID<FDoublePerlinNoiseParams> DoublePerlinNoiseChannel;
 
 	struct
 	{
diff --git a/Engine/Source/Runtime/MovieSceneTracks/Public/Sections/MovieScene3DTransformSection.h b/Engine/Source/Runtime/MovieSceneTracks/Public/Sections/MovieScene3DTransformSection.h
--- a/Engine/Source/Runtime/MovieSceneTracks/Public/Sections/MovieScene3DTransformSection.h
+++ b/Engine/Source/Runtime/MovieSceneTracks/Public/Sections/MovieScene3DTransformSection.h
@@ -10,12 +10,12 @@
 #include "MovieSceneKeyStruct.h"
 #include "Channels/MovieSceneDoubleChannel.h"
 #include "Channels/MovieSceneFloatChannel.h"
+#include "Channels/MovieSceneSectionChannelOverrideRegistry.h"
 #include "EntitySystem/IMovieSceneEntityProvider.h"
 #include "TransformData.h"
 #include "Misc/LargeWorldCoordinates.h"
 #include "MovieScene3DTransformSection.generated.h"
 
-
 #if WITH_EDITORONLY_DATA
 /** Visibility options for 3d trajectory. */
 UENUM()
@@ -258,12 +258,13 @@ protected:
 	virtual EMovieSceneChannelProxyType CacheChannelProxy() override;
 
 private:
-
+	
+	virtual bool PopulateEvaluationFieldImpl(const TRange<FFrameNumber>& EffectiveRange, const FMovieSceneEvaluationFieldEntityMetaData& InMetaData, FMovieSceneEntityComponentFieldBuilder* OutFieldBuilder) override;
 	virtual void ImportEntityImpl(UMovieSceneEntitySystemLinker* EntityLinker, const FEntityImportParams& Params, FImportedEntity* OutImportedEntity) override;
 	virtual void InterrogateEntityImpl(UMovieSceneEntitySystemLinker* EntityLinker, const FEntityImportParams& Params, FImportedEntity* OutImportedEntity) override;
 
 	template<typename BaseBuilderType>
-	void BuildEntity(BaseBuilderType& InBaseBuilder, UMovieSceneEntitySystemLinker* EntityLinker, const FEntityImportParams& Params, FImportedEntity* OutImportedEntity);
+	void BuildEntity(BaseBuilderType& InBaseBuilder, UMovieSceneEntitySystemLinker* Linker, const FEntityImportParams& Params, FImportedEntity* OutImportedEntity);
 
 private:
 
@@ -286,6 +287,10 @@ private:
 	UPROPERTY()
 	FMovieSceneFloatChannel ManualWeight;
 
+	/** Optional pointer to a “channels override” container object. This object would only be allocated if any channels are overridden with a non-standard channel 	*/
+	UPROPERTY()
+	TObjectPtr<UMovieSceneSectionChannelOverrideRegistry> OverrideRegistry;
+
 	/** Whether to use a quaternion linear interpolation between keys. This finds the 'shortest' distance between keys */
 	UPROPERTY(EditAnywhere, DisplayName = "Use Quaternion Interpolation", Category = "Rotation")
 	bool bUseQuaternionInterpolation;
diff --git a/Engine/Source/Runtime/MovieSceneTracks/Public/Systems/DoublePerlinNoiseChannelEvaluatorSystem.h b/Engine/Source/Runtime/MovieSceneTracks/Public/Systems/DoublePerlinNoiseChannelEvaluatorSystem.h
new file mode 100644
--- /dev/null
+++ b/Engine/Source/Runtime/MovieSceneTracks/Public/Systems/DoublePerlinNoiseChannelEvaluatorSystem.h
@@ -0,0 +1,19 @@
+// Copyright Epic Games, Inc. All Rights Reserved.
+
+#pragma once
+
+#include "EntitySystem/MovieSceneEntitySystem.h"
+#include "DoublePerlinNoiseChannelEvaluatorSystem.generated.h"
+
+/**
+ * System that is responsible for evaluating double perlin noise channels.
+ */
+UCLASS()
+class MOVIESCENETRACKS_API UDoublePerlinNoiseChannelEvaluatorSystem : public UMovieSceneEntitySystem
+{
+public:
+	GENERATED_BODY()
+
+	UDoublePerlinNoiseChannelEvaluatorSystem(const FObjectInitializer& ObjInit);
+	virtual void OnRun(FSystemTaskPrerequisites& InPrerequisites, FSystemSubsequentTasks& Subsequents) override;
+};
diff --git a/Engine/Source/Runtime/MovieSceneTracks/Public/Systems/FloatPerlinNoiseChannelEvaluatorSystem.h b/Engine/Source/Runtime/MovieSceneTracks/Public/Systems/FloatPerlinNoiseChannelEvaluatorSystem.h
new file mode 100644
--- /dev/null
+++ b/Engine/Source/Runtime/MovieSceneTracks/Public/Systems/FloatPerlinNoiseChannelEvaluatorSystem.h
@@ -0,0 +1,19 @@
+// Copyright Epic Games, Inc. All Rights Reserved.
+
+#pragma once
+
+#include "EntitySystem/MovieSceneEntitySystem.h"
+#include "FloatPerlinNoiseChannelEvaluatorSystem.generated.h"
+
+/**
+ * System that is responsible for evaluating double perlin noise channels.
+ */
+UCLASS()
+class MOVIESCENETRACKS_API UFloatPerlinNoiseChannelEvaluatorSystem : public UMovieSceneEntitySystem
+{
+public:
+	GENERATED_BODY()
+
+	UFloatPerlinNoiseChannelEvaluatorSystem(const FObjectInitializer& ObjInit);
+	virtual void OnRun(FSystemTaskPrerequisites& InPrerequisites, FSystemSubsequentTasks& Subsequents) override;
+};
diff --git a/Engine/Source/Runtime/MovieSceneTracks/Public/Tracks/MovieScenePropertyTrack.h b/Engine/Source/Runtime/MovieSceneTracks/Public/Tracks/MovieScenePropertyTrack.h
--- a/Engine/Source/Runtime/MovieSceneTracks/Public/Tracks/MovieScenePropertyTrack.h
+++ b/Engine/Source/Runtime/MovieSceneTracks/Public/Tracks/MovieScenePropertyTrack.h
@@ -12,6 +12,7 @@
 #include "Systems/MovieScenePiecewiseBoolBlenderSystem.h"
 #include "EntitySystem/MovieScenePropertyBinding.h"
 #include "UObject/ObjectMacros.h"
+#include "Channels/MovieSceneSectionChannelOverrideRegistry.h"
 #include "MovieScenePropertyTrack.generated.h"
 
 /**
@@ -154,7 +155,7 @@ struct MOVIESCENETRACKS_API FMovieScenePropertyTrackEntityImportHelper
 	static const int32 SectionPropertyValueImportingID;
 	static const int32 SectionEditConditionToggleImportingID;
 
-	static void PopulateEvaluationField(UMovieSceneSection& Section, const TRange<FFrameNumber>& EffectiveRange, const FMovieSceneEvaluationFieldEntityMetaData& InMetaData, FMovieSceneEntityComponentFieldBuilder* OutFieldBuilder);
+	static void PopulateEvaluationField(UMovieSceneSection& Section, const TRange<FFrameNumber>& EffectiveRange, const FMovieSceneEvaluationFieldEntityMetaData& InMetaData, FMovieSceneEntityComponentFieldBuilder* OutFieldBuilder, UMovieSceneSectionChannelOverrideRegistry* OverrideRegistry = nullptr);
 
 	static bool IsPropertyValueID(const UE::MovieScene::FEntityImportParams& Params);
 	static bool IsEditConditionToggleID(const UE::MovieScene::FEntityImportParams& Params);
@@ -176,27 +177,34 @@ namespace MovieScene
 template<typename... T>
 struct TPropertyTrackEntityImportHelperImpl
 {
-	TPropertyTrackEntityImportHelperImpl(TEntityBuilder<T...>&& InBuilder)
+	TPropertyTrackEntityImportHelperImpl(TEntityBuilder<T...>&& InBuilder, UMovieSceneSectionChannelOverrideRegistry* InOverrideRegistry = nullptr, FComponentTypeID InPropertyTag = {})
 		: Builder(MoveTemp(InBuilder))
+		, OverrideRegistry(InOverrideRegistry)
+		, PropertyTag(InPropertyTag)
 	{
 	}
 
 	template<typename U, typename PayloadType>
 	TPropertyTrackEntityImportHelperImpl<T..., TAdd<U>> Add(TComponentTypeID<U> ComponentType, PayloadType&& InPayload)
 	{
-		return TPropertyTrackEntityImportHelperImpl<T..., TAdd<U>>(Builder.Add(ComponentType, InPayload));
+		return TPropertyTrackEntityImportHelperImpl<T..., TAdd<U>>(Builder.Add(ComponentType, InPayload), OverrideRegistry, PropertyTag);
 	}
 
 	template<typename U, typename PayloadType>
 	TPropertyTrackEntityImportHelperImpl<T..., TAddConditional<U>> AddConditional(TComponentTypeID<U> ComponentType, PayloadType&& InPayload, bool bCondition)
 	{
-		return TPropertyTrackEntityImportHelperImpl<T..., TAddConditional<U>>(Builder.AddConditional(ComponentType, InPayload, bCondition));
+		return TPropertyTrackEntityImportHelperImpl<T..., TAddConditional<U>>(Builder.AddConditional(ComponentType, InPayload, bCondition), OverrideRegistry, PropertyTag);
+	}
+
+	template<typename U, typename PayloadType>
+	TPropertyTrackEntityImportHelperImpl<T..., TAddConditional<U>> AddConditional(TComponentTypeID<U> ComponentType, PayloadType&& InPayload, bool bCondition, int32 ChannelIndex)
+	{
+		return TPropertyTrackEntityImportHelperImpl<T..., TAddConditional<U>>(Builder.AddConditional(ComponentType, InPayload, bCondition && !IsChannelOverriden(OverrideRegistry, ChannelIndex)), OverrideRegistry, PropertyTag);
 	}
 
 	void Commit(const UMovieSceneSection* InSection, const UE::MovieScene::FEntityImportParams& Params, UE::MovieScene::FImportedEntity* OutImportedEntity)
 	{
 		const FGuid ObjectBindingID = Params.GetObjectBindingID();
-
 		const FBuiltInComponentTypes* BuiltInComponents = FBuiltInComponentTypes::Get();
 		const FMovieSceneTracksComponentTypes* TracksComponents = FMovieSceneTracksComponentTypes::Get();
 
@@ -206,8 +214,8 @@ struct TPropertyTrackEntityImportHelperImpl
 			{
 				OutImportedEntity->AddBuilder(
 					Builder
-						.Add(BuiltInComponents->PropertyBinding, PropertyTrack->GetPropertyBinding())
-						.AddConditional(BuiltInComponents->GenericObjectBinding, ObjectBindingID, ObjectBindingID.IsValid()));
+					.Add(BuiltInComponents->PropertyBinding, PropertyTrack->GetPropertyBinding())
+					.AddConditional(BuiltInComponents->GenericObjectBinding, ObjectBindingID, ObjectBindingID.IsValid()));
 			}
 			else if (FMovieScenePropertyTrackEntityImportHelper::IsEditConditionToggleID(Params))
 			{
@@ -215,18 +223,31 @@ struct TPropertyTrackEntityImportHelperImpl
 				// need to import the edit condition toggle entity.
 				FMovieScenePropertyTrackEntityImportHelper::ImportEditConditionToggleEntity(Params, OutImportedEntity);
 			}
+			// Overriden
+			else if (OverrideRegistry)
+			{
+				auto BaseBuilder = FEntityBuilder()
+					.Add(BuiltInComponents->PropertyBinding, PropertyTrack->GetPropertyBinding())
+					.AddConditional(BuiltInComponents->GenericObjectBinding, ObjectBindingID, ObjectBindingID.IsValid())
+					.AddTag(PropertyTag);
+			
+				OutImportedEntity->AddBuilder(BaseBuilder);
+				OverrideRegistry->ImportEntityImpl(Params, OutImportedEntity);
+			}
 		}
 		else
 		{
 			OutImportedEntity->AddBuilder(
 				Builder
-					.AddConditional(BuiltInComponents->GenericObjectBinding, ObjectBindingID, ObjectBindingID.IsValid()));
+				.AddConditional(BuiltInComponents->GenericObjectBinding, ObjectBindingID, ObjectBindingID.IsValid()));
 		}
 	}
 
 protected:
 
 	TEntityBuilder<T...> Builder;
+	TObjectPtr<UMovieSceneSectionChannelOverrideRegistry> OverrideRegistry;
+	FComponentTypeID PropertyTag;
 };
 
 
@@ -236,33 +257,43 @@ protected:
 template<>
 struct TPropertyTrackEntityImportHelperImpl<>
 {
-	TPropertyTrackEntityImportHelperImpl(FComponentTypeID PropertyTag)
-		: Builder(FEntityBuilder().AddTag(PropertyTag))
+	TPropertyTrackEntityImportHelperImpl(FComponentTypeID InPropertyTag, UMovieSceneSectionChannelOverrideRegistry* InOverrideRegistry = nullptr)
+		: Builder(FEntityBuilder().AddTag(InPropertyTag))
+		, OverrideRegistry(InOverrideRegistry)
+		, PropertyTag(InPropertyTag)
 	{
 	}
 
 	template<typename U, typename PayloadType>
 	TPropertyTrackEntityImportHelperImpl<FAdd, TAdd<U>> Add(TComponentTypeID<U> ComponentType, PayloadType&& InPayload)
 	{
-		return TPropertyTrackEntityImportHelperImpl<FAdd, TAdd<U>>(Builder.Add(ComponentType, InPayload));
+		return TPropertyTrackEntityImportHelperImpl<FAdd, TAdd<U>>(Builder.Add(ComponentType, InPayload), OverrideRegistry, PropertyTag);
 	}
 
 	template<typename U, typename PayloadType>
 	TPropertyTrackEntityImportHelperImpl<FAdd, TAddConditional<U>> AddConditional(TComponentTypeID<U> ComponentType, PayloadType&& InPayload, bool bCondition)
 	{
-		return TPropertyTrackEntityImportHelperImpl<FAdd, TAddConditional<U>>(Builder.AddConditional(ComponentType, InPayload, bCondition));
+		return TPropertyTrackEntityImportHelperImpl<FAdd, TAddConditional<U>>(Builder.AddConditional(ComponentType, InPayload, bCondition), OverrideRegistry, PropertyTag);
+	}
+
+	template<typename U, typename PayloadType>
+	TPropertyTrackEntityImportHelperImpl<FAdd, TAddConditional<U>> AddConditional(TComponentTypeID<U> ComponentType, PayloadType&& InPayload, bool bCondition, int32 ChannelIndex)
+	{
+		return TPropertyTrackEntityImportHelperImpl<FAdd, TAddConditional<U>>(Builder.AddConditional(ComponentType, InPayload, bCondition && !IsChannelOverriden(OverrideRegistry, ChannelIndex)), OverrideRegistry, PropertyTag);
 	}
 
 protected:
 
 	TEntityBuilder<FAdd> Builder;
+	TObjectPtr<UMovieSceneSectionChannelOverrideRegistry> OverrideRegistry;
+	FComponentTypeID PropertyTag;
 };
 
 struct MOVIESCENETRACKS_API FPropertyTrackEntityImportHelper : TPropertyTrackEntityImportHelperImpl<>
 {
 	template<typename PropertyTraits>
-	FPropertyTrackEntityImportHelper(const TPropertyComponents<PropertyTraits>& PropertyComponents)
-		: TPropertyTrackEntityImportHelperImpl<>(PropertyComponents.PropertyTag)
+	FPropertyTrackEntityImportHelper(const TPropertyComponents<PropertyTraits>& PropertyComponents, UMovieSceneSectionChannelOverrideRegistry* InOverrideRegistry = nullptr)
+		: TPropertyTrackEntityImportHelperImpl<>(PropertyComponents.PropertyTag, InOverrideRegistry)
 	{}
 };
 
