Sequencer: Static analysis fix - Use UE_ARRAY_COUNT instead of size_of

diff --git a/Engine/Source/Runtime/MovieSceneTracks/Private/Systems/DoublePerlinNoiseChannelEvaluatorSystem.cpp b/Engine/Source/Runtime/MovieSceneTracks/Private/Systems/DoublePerlinNoiseChannelEvaluatorSystem.cpp
--- a/Engine/Source/Runtime/MovieSceneTracks/Private/Systems/DoublePerlinNoiseChannelEvaluatorSystem.cpp
+++ b/Engine/Source/Runtime/MovieSceneTracks/Private/Systems/DoublePerlinNoiseChannelEvaluatorSystem.cpp
@@ -53,7 +53,7 @@ UDoublePerlinNoiseChannelEvaluatorSystem::UDoublePerlinNoiseChannelEvaluatorSyst
 	if (HasAnyFlags(RF_ClassDefaultObject))
 	{
 		// Allow writing to all the possible channels
-		for (int32 i = 0; i < sizeof(BuiltInComponents->DoubleResult); ++i)
+		for (int32 i = 0; i < UE_ARRAY_COUNT(BuiltInComponents->DoubleResult); ++i)
 		{
 			DefineComponentProducer(UDoublePerlinNoiseChannelEvaluatorSystem::StaticClass(), BuiltInComponents->DoubleResult[i]);
 		}
@@ -70,7 +70,7 @@ void UDoublePerlinNoiseChannelEvaluatorSystem::OnRun(FSystemTaskPrerequisites& I
 	const FBuiltInComponentTypes* BuiltInComponents = FBuiltInComponentTypes::Get();
 	const FMovieSceneTracksComponentTypes* TrackComponents = FMovieSceneTracksComponentTypes::Get();
 
-	for (int32 i = 0; i < sizeof(BuiltInComponents->DoubleResult); ++i)
+	for (int32 i = 0; i < UE_ARRAY_COUNT(BuiltInComponents->DoubleResult); ++i)
 	{
 		FEntityTaskBuilder()
 			.Read(TrackComponents->DoublePerlinNoiseChannel)
diff --git a/Engine/Source/Runtime/MovieSceneTracks/Private/Systems/FloatPerlinNoiseChannelEvaluatorSystem.cpp b/Engine/Source/Runtime/MovieSceneTracks/Private/Systems/FloatPerlinNoiseChannelEvaluatorSystem.cpp
--- a/Engine/Source/Runtime/MovieSceneTracks/Private/Systems/FloatPerlinNoiseChannelEvaluatorSystem.cpp
+++ b/Engine/Source/Runtime/MovieSceneTracks/Private/Systems/FloatPerlinNoiseChannelEvaluatorSystem.cpp
@@ -53,7 +53,7 @@ UFloatPerlinNoiseChannelEvaluatorSystem::UFloatPerlinNoiseChannelEvaluatorSystem
 	if (HasAnyFlags(RF_ClassDefaultObject))
 	{
 		// Allow writing to all the possible channels
-		for (int32 i = 0; i < sizeof(BuiltInComponents->FloatResult); ++i)
+		for (int32 i = 0; i < UE_ARRAY_COUNT(BuiltInComponents->FloatResult); ++i)
 		{
 			DefineComponentProducer(UFloatPerlinNoiseChannelEvaluatorSystem::StaticClass(), BuiltInComponents->FloatResult[i]);
 		}
@@ -70,7 +70,7 @@ void UFloatPerlinNoiseChannelEvaluatorSystem::OnRun(FSystemTaskPrerequisites& In
 	const FBuiltInComponentTypes* BuiltInComponents = FBuiltInComponentTypes::Get();
 	const FMovieSceneTracksComponentTypes* TrackComponents = FMovieSceneTracksComponentTypes::Get();
 
-	for (int32 i = 0; i < sizeof(BuiltInComponents->FloatResult); ++i)
+	for (int32 i = 0; i < UE_ARRAY_COUNT(BuiltInComponents->FloatResult); ++i)
 	{
 		FEntityTaskBuilder()
 			.Read(TrackComponents->FloatPerlinNoiseChannel)
