MovieSceneTracks: set source timecode when creating new skeletal animation sections with custom timecode bone attributes

diff --git a/Engine/Source/Runtime/MovieSceneTracks/MovieSceneTracks.Build.cs b/Engine/Source/Runtime/MovieSceneTracks/MovieSceneTracks.Build.cs
--- a/Engine/Source/Runtime/MovieSceneTracks/MovieSceneTracks.Build.cs
+++ b/Engine/Source/Runtime/MovieSceneTracks/MovieSceneTracks.Build.cs
@@ -39,6 +39,7 @@ public class MovieSceneTracks : ModuleRules
 				});
 			PrivateDependencyModuleNames.AddRange(new string[]
 				{
+					"AnimationBlueprintLibrary",
 					"DataLayerEditor",
 					"EditorFramework",
 					"UnrealEd"
diff --git a/Engine/Source/Runtime/MovieSceneTracks/Private/Tracks/MovieSceneSkeletalAnimationTrack.cpp b/Engine/Source/Runtime/MovieSceneTracks/Private/Tracks/MovieSceneSkeletalAnimationTrack.cpp
--- a/Engine/Source/Runtime/MovieSceneTracks/Private/Tracks/MovieSceneSkeletalAnimationTrack.cpp
+++ b/Engine/Source/Runtime/MovieSceneTracks/Private/Tracks/MovieSceneSkeletalAnimationTrack.cpp
@@ -17,8 +17,12 @@
 #include "SkeletalDebugRendering.h"
 #include "Rendering/SkeletalMeshRenderData.h"
 #include "BoneContainer.h"
-#include "Animation/AnimationPoseData.h"
-#include "Animation/AttributesRuntime.h"
+
+#if WITH_EDITORONLY_DATA
+#include "AnimationBlueprintLibrary.h"
+#include "Misc/QualifiedFrameTime.h"
+#include "Misc/Timecode.h"
+#endif
 
 #define LOCTEXT_NAMESPACE "MovieSceneSkeletalAnimationTrack"
 
@@ -60,6 +64,14 @@ UMovieSceneSection* UMovieSceneSkeletalAnimationTrack::AddNewAnimationOnRow(FFra
 		int32 IFrameNumber = AnimationLength.FrameNumber.Value + (int)(AnimationLength.GetSubFrame() + 0.5f) + 1;
 		NewSection->InitialPlacementOnRow(AnimationSections, KeyTime, IFrameNumber, RowIndex);
 		NewSection->Params.Animation = AnimSequence;
+
+#if WITH_EDITORONLY_DATA
+		FQualifiedFrameTime SourceStartFrameTime;
+		if (UAnimationBlueprintLibrary::EvaluateRootBoneTimecodeAttributesAtTime(NewSection->Params.Animation, 0.0f, SourceStartFrameTime))
+		{
+			NewSection->TimecodeSource.Timecode = SourceStartFrameTime.ToTimecode();
+		}
+#endif
 	}
 
 	AddSection(*NewSection);
