Sequencer: Update audio section channels when the inputs have changed. This follows the model for other sections and fixes a crash when undoing a change to a metasound parameter

diff --git a/Engine/Source/Runtime/MovieSceneTracks/Private/Sections/MovieSceneAudioSection.cpp b/Engine/Source/Runtime/MovieSceneTracks/Private/Sections/MovieSceneAudioSection.cpp
--- a/Engine/Source/Runtime/MovieSceneTracks/Private/Sections/MovieSceneAudioSection.cpp
+++ b/Engine/Source/Runtime/MovieSceneTracks/Private/Sections/MovieSceneAudioSection.cpp
@@ -77,6 +77,23 @@ namespace MovieSceneAudioSectionPrivate
 	}
 }
 
+void UMovieSceneAudioSection::Serialize(FArchive& Ar)
+{
+	Super::Serialize(Ar);
+
+	if (Ar.IsLoading())
+	{
+		CacheChannelProxy();
+	}
+}
+
+void UMovieSceneAudioSection::PostEditImport()
+{
+	Super::PostEditImport();
+
+	CacheChannelProxy();
+}
+
 EMovieSceneChannelProxyType  UMovieSceneAudioSection::CacheChannelProxy()
 {
 	// Set up the channel proxy
diff --git a/Engine/Source/Runtime/MovieSceneTracks/Public/Sections/MovieSceneAudioSection.h b/Engine/Source/Runtime/MovieSceneTracks/Public/Sections/MovieSceneAudioSection.h
--- a/Engine/Source/Runtime/MovieSceneTracks/Public/Sections/MovieSceneAudioSection.h
+++ b/Engine/Source/Runtime/MovieSceneTracks/Public/Sections/MovieSceneAudioSection.h
@@ -176,7 +176,8 @@ public:
 	virtual UMovieSceneSection* SplitSection(FQualifiedFrameTime SplitTime, bool bDeleteKeys) override;
 	virtual TOptional<FFrameTime> GetOffsetTime() const override;
 	virtual void MigrateFrameTimes(FFrameRate SourceRate, FFrameRate DestinationRate) override;
-
+	virtual void Serialize(FArchive& Ar) override;
+	virtual void PostEditImport() override;
 	virtual EMovieSceneChannelProxyType CacheChannelProxy() override;
 
 private:
