Sequencer: Fix regression bug where section-preroll wasn't handled for camera cut tracks.

diff --git a/Engine/Source/Runtime/MovieSceneTracks/Private/TrackInstances/MovieSceneCameraCutTrackInstance.cpp b/Engine/Source/Runtime/MovieSceneTracks/Private/TrackInstances/MovieSceneCameraCutTrackInstance.cpp
--- a/Engine/Source/Runtime/MovieSceneTracks/Private/TrackInstances/MovieSceneCameraCutTrackInstance.cpp
+++ b/Engine/Source/Runtime/MovieSceneTracks/Private/TrackInstances/MovieSceneCameraCutTrackInstance.cpp
@@ -212,7 +212,16 @@ void UMovieSceneCameraCutTrackInstance::OnAnimate()
 		FTransform CutTransform = Section->InitialCameraCutTransform;
 		const bool bHasCutTransform = Section->bHasInitialCameraCutTransform;
 
-		if (Context.IsPreRoll())
+		const int32 PreRollFrames = Section->GetPreRollFrames();
+		bool bIsSectionPreRoll = false;
+		if (PreRollFrames > 0 && Section->HasStartFrame())
+		{
+			const FFrameNumber SectionStartTime = Section->GetTrueRange().GetLowerBoundValue();
+			const TRange<FFrameNumber> PreRollRange(SectionStartTime - PreRollFrames, SectionStartTime);
+			bIsSectionPreRoll = PreRollRange.Contains(Context.GetTime().FloorToFrame());
+		}
+
+		if (Context.IsPreRoll() || bIsSectionPreRoll)
 		{
 			FPreRollCameraCut PreRollCameraCut { Input.InstanceHandle, CameraBindingID, CutTransform, bHasCutTransform };
 			CameraCutPreRolls.Add(PreRollCameraCut);
