Sequencer: Previous FloatVectorSections still remain as FloatVectorSections after the LWC upgrade. That means they will enter the FFloatVectorHandler and their struct name will be NAME_Vector. Now that LWC is enabled, FFloatVectorHandler also needs to test whether the struct name is NAME_Vector so that it can correctly assign the meta data channels and evaluate as double in FFloatVectorHandler.

diff --git a/Engine/Source/Runtime/MovieSceneTracks/Private/MovieSceneTracksComponentTypes.cpp b/Engine/Source/Runtime/MovieSceneTracks/Private/MovieSceneTracksComponentTypes.cpp
--- a/Engine/Source/Runtime/MovieSceneTracks/Private/MovieSceneTracksComponentTypes.cpp
+++ b/Engine/Source/Runtime/MovieSceneTracks/Private/MovieSceneTracksComponentTypes.cpp
@@ -405,13 +405,16 @@ struct FFloatVectorHandler : TPropertyComponentHandler<FFloatVectorPropertyTrait
 				{
 					OutMetaData.NumChannels = 2;
 				}
-				else if (BoundProperty->Struct->GetFName() == NAME_Vector3f
-#if UE_LARGE_WORLD_COORDINATES_DISABLED
-						|| BoundProperty->Struct->GetFName() == NAME_Vector
-#endif
-						)
+				else if (BoundProperty->Struct->GetFName() == NAME_Vector3f || BoundProperty->Struct->GetFName() == NAME_Vector)
 				{
 					OutMetaData.NumChannels = 3;
+
+#if !UE_LARGE_WORLD_COORDINATES_DISABLED
+					if (BoundProperty->Struct->GetFName() == NAME_Vector)
+					{
+						OutMetaData.bIsDouble = true;
+					}
+#endif
 				}
 				else
 				{
diff --git a/Engine/Source/Runtime/MovieSceneTracks/Public/MovieSceneTracksComponentTypes.h b/Engine/Source/Runtime/MovieSceneTracks/Public/MovieSceneTracksComponentTypes.h
--- a/Engine/Source/Runtime/MovieSceneTracks/Public/MovieSceneTracksComponentTypes.h
+++ b/Engine/Source/Runtime/MovieSceneTracks/Public/MovieSceneTracksComponentTypes.h
@@ -191,6 +191,14 @@ struct FFloatVectorPropertyTraits
 	}
 	static void GetObjectPropertyValue(const UObject* InObject, FVectorPropertyMetaData MetaData, uint16 PropertyOffset, FFloatIntermediateVector& OutValue)
 	{
+		if (MetaData.bIsDouble)
+		{
+			FDoubleIntermediateVector TempDoubleValue;
+			TIndirectPropertyTraits<FVector3d, FDoubleIntermediateVector>::GetObjectPropertyValue(InObject, PropertyOffset, TempDoubleValue); 
+			OutValue = FFloatIntermediateVector(TempDoubleValue.X, TempDoubleValue.Y, TempDoubleValue.Z, TempDoubleValue.W);
+			return;
+		}
+
 		switch (MetaData.NumChannels)
 		{
 			case 2: TIndirectPropertyTraits<FVector2D, FFloatIntermediateVector>::GetObjectPropertyValue(InObject, PropertyOffset, OutValue); return;
@@ -200,6 +208,14 @@ struct FFloatVectorPropertyTraits
 	}
 	static void GetObjectPropertyValue(const UObject* InObject, FVectorPropertyMetaData MetaData, FTrackInstancePropertyBindings* PropertyBindings, FFloatIntermediateVector& OutValue)
 	{
+		if (MetaData.bIsDouble)
+		{
+			FDoubleIntermediateVector TempDoubleValue;
+			TIndirectPropertyTraits<FVector3d, FDoubleIntermediateVector>::GetObjectPropertyValue(InObject, PropertyBindings, TempDoubleValue); 
+			OutValue = FFloatIntermediateVector(TempDoubleValue.X, TempDoubleValue.Y, TempDoubleValue.Z, TempDoubleValue.W);
+			return;
+		}
+
 		switch (MetaData.NumChannels)
 		{
 			case 2: TIndirectPropertyTraits<FVector2D, FFloatIntermediateVector>::GetObjectPropertyValue(InObject, PropertyBindings, OutValue); return;
@@ -209,6 +225,14 @@ struct FFloatVectorPropertyTraits
 	}
 	static void GetObjectPropertyValue(const UObject* InObject, FVectorPropertyMetaData MetaData, const FName& PropertyPath, StorageType& OutValue)
 	{
+		if (MetaData.bIsDouble)
+		{
+			FDoubleIntermediateVector TempDoubleValue;
+			TIndirectPropertyTraits<FVector3d, FDoubleIntermediateVector>::GetObjectPropertyValue(InObject, PropertyPath, TempDoubleValue); 
+			OutValue = FFloatIntermediateVector(TempDoubleValue.X, TempDoubleValue.Y, TempDoubleValue.Z, TempDoubleValue.W);
+			return;
+		}
+
 		switch (MetaData.NumChannels)
 		{
 			case 2: TIndirectPropertyTraits<FVector2D, FFloatIntermediateVector>::GetObjectPropertyValue(InObject, PropertyPath, OutValue); return;
@@ -224,6 +248,13 @@ struct FFloatVectorPropertyTraits
 	}
 	static void SetObjectPropertyValue(UObject* InObject, FVectorPropertyMetaData MetaData, uint16 PropertyOffset, const FFloatIntermediateVector& InValue)
 	{
+		if (MetaData.bIsDouble)
+		{
+			FDoubleIntermediateVector TempDoubleValue(InValue.X, InValue.Y, InValue.Z, InValue.W);
+			TIndirectPropertyTraits<FVector3d, FDoubleIntermediateVector>::SetObjectPropertyValue(InObject, PropertyOffset, TempDoubleValue);
+			return;
+		}
+
 		switch (MetaData.NumChannels)
 		{
 			case 2: TIndirectPropertyTraits<FVector2D, FFloatIntermediateVector>::SetObjectPropertyValue(InObject, PropertyOffset, InValue); return;
@@ -235,6 +266,13 @@ struct FFloatVectorPropertyTraits
 	}
 	static void SetObjectPropertyValue(UObject* InObject, FVectorPropertyMetaData MetaData, FTrackInstancePropertyBindings* PropertyBindings, const FFloatIntermediateVector& InValue)
 	{
+		if (MetaData.bIsDouble)
+		{
+			FDoubleIntermediateVector TempDoubleValue(InValue.X, InValue.Y, InValue.Z, InValue.W);
+			TIndirectPropertyTraits<FVector3d, FDoubleIntermediateVector>::SetObjectPropertyValue(InObject, PropertyBindings, TempDoubleValue);
+			return;
+		}
+
 		switch (MetaData.NumChannels)
 		{
 			case 2: TIndirectPropertyTraits<FVector2D, FFloatIntermediateVector>::SetObjectPropertyValue(InObject, PropertyBindings, InValue); return;
diff --git a/Engine/Source/Runtime/MovieSceneTracks/Public/MovieSceneTracksPropertyTypes.h b/Engine/Source/Runtime/MovieSceneTracks/Public/MovieSceneTracksPropertyTypes.h
--- a/Engine/Source/Runtime/MovieSceneTracks/Public/MovieSceneTracksPropertyTypes.h
+++ b/Engine/Source/Runtime/MovieSceneTracks/Public/MovieSceneTracksPropertyTypes.h
@@ -23,6 +23,7 @@ namespace MovieScene
 struct FVectorPropertyMetaData
 {
 	uint8 NumChannels = 0;
+	bool bIsDouble = false;
 };
 
 /** Intermediate type for the vector property system that lets us store how many dimensions the vector should have */
