Sequencer: Console variable track

diff --git a/Engine/Source/Editor/MovieSceneTools/Private/MovieSceneCVarOverridesPropertyTypeCustomization.cpp b/Engine/Source/Editor/MovieSceneTools/Private/MovieSceneCVarOverridesPropertyTypeCustomization.cpp
new file mode 100644
--- /dev/null
+++ b/Engine/Source/Editor/MovieSceneTools/Private/MovieSceneCVarOverridesPropertyTypeCustomization.cpp
@@ -0,0 +1,90 @@
+// Copyright Epic Games, Inc. All Rights Reserved.
+
+#include "MovieSceneCVarOverridesPropertyTypeCustomization.h"
+#include "Sections/MovieSceneCVarSection.h"
+#include "DetailWidgetRow.h"
+#include "Widgets/Text/SMultiLineEditableText.h"
+#include "Widgets/Layout/SBox.h"
+
+#define LOCTEXT_NAMESPACE "FCVarOverridesPropertyTypeCustomization"
+
+namespace UE
+{
+namespace MovieScene
+{
+
+TSharedRef<IPropertyTypeCustomization> FCVarOverridesPropertyTypeCustomization::MakeInstance()
+{
+	return MakeShared<FCVarOverridesPropertyTypeCustomization>();
+}
+
+void FCVarOverridesPropertyTypeCustomization::CustomizeHeader(TSharedRef<IPropertyHandle> StructPropertyHandle, FDetailWidgetRow& HeaderRow, IPropertyTypeCustomizationUtils& StructCustomizationUtils)
+{
+	PropertyHandle = StructPropertyHandle;
+
+	HeaderRow
+	.NameContent()
+	[
+		StructPropertyHandle->CreatePropertyNameWidget()
+	]
+	.ValueContent()
+	.MaxDesiredWidth(TOptional<float>())
+	[
+		SNew(SBox)
+		.MinDesiredHeight(150.f)
+		[
+			SNew(SMultiLineEditableText)
+			.TextStyle(FAppStyle::Get(), "DetailsView.CategoryTextStyle")
+			.OnTextCommitted(this, &FCVarOverridesPropertyTypeCustomization::OnCVarsCommitted)
+			.OnTextChanged(this, &FCVarOverridesPropertyTypeCustomization::OnCVarsChanged)
+			.Text(this, &FCVarOverridesPropertyTypeCustomization::GetCVarText)
+			.HintText(LOCTEXT("CVarSectionHint", "Add CVar values"))
+		]
+	];
+}
+
+void FCVarOverridesPropertyTypeCustomization::OnCVarsCommitted(const FText& NewText, ETextCommit::Type)
+{
+	OnCVarsChanged(NewText);
+}
+
+void FCVarOverridesPropertyTypeCustomization::OnCVarsChanged(const FText& NewText)
+{
+	FString NewValue = NewText.ToString();
+
+	PropertyHandle->NotifyPreChange();
+
+	TArray<void*> RawPtrs;
+	PropertyHandle->AccessRawData(RawPtrs);
+
+	for (void* Ptr : RawPtrs)
+	{
+		static_cast<FMovieSceneCVarOverrides*>(Ptr)->SetFromString(NewValue);
+	}
+
+	PropertyHandle->NotifyPostChange(EPropertyChangeType::ValueSet);
+	PropertyHandle->NotifyFinishedChangingProperties();
+}
+
+FText FCVarOverridesPropertyTypeCustomization::GetCVarText() const
+{
+	TArray<const void*> RawPtrs;
+	PropertyHandle->AccessRawData(RawPtrs);
+
+	if (RawPtrs.Num() == 0)
+	{
+		return FText();
+	}
+
+	if (RawPtrs.Num() == 1)
+	{
+		return FText::FromString(static_cast<const FMovieSceneCVarOverrides*>(RawPtrs[0])->GetString());
+	}
+
+	return LOCTEXT("Multiple Values", "<<Multiple Values>>");
+}
+
+} // namespace MovieScene
+} // namespace UE
+
+#undef LOCTEXT_NAMESPACE
\ No newline at end of file
diff --git a/Engine/Source/Editor/MovieSceneTools/Private/MovieSceneCVarOverridesPropertyTypeCustomization.h b/Engine/Source/Editor/MovieSceneTools/Private/MovieSceneCVarOverridesPropertyTypeCustomization.h
new file mode 100644
--- /dev/null
+++ b/Engine/Source/Editor/MovieSceneTools/Private/MovieSceneCVarOverridesPropertyTypeCustomization.h
@@ -0,0 +1,41 @@
+// Copyright Epic Games, Inc. All Rights Reserved.
+
+#pragma once
+
+#include "CoreTypes.h"
+#include "PropertyEditorModule.h"
+#include "IPropertyTypeCustomization.h"
+
+struct EVisibility;
+class FReply;
+class SWidget;
+class FDragDropOperation;
+class FDataLayerDragDropOp;
+class UDataLayer;
+struct FSlateColor;
+struct FSlateBrush;
+
+namespace UE
+{
+namespace MovieScene
+{
+
+struct FCVarOverridesPropertyTypeCustomization : public IPropertyTypeCustomization
+{
+	static TSharedRef<IPropertyTypeCustomization> MakeInstance();
+
+	virtual void CustomizeHeader(TSharedRef<IPropertyHandle> StructPropertyHandle, FDetailWidgetRow& HeaderRow, IPropertyTypeCustomizationUtils& StructCustomizationUtils) override;
+	virtual void CustomizeChildren(TSharedRef<IPropertyHandle> StructPropertyHandle, IDetailChildrenBuilder& StructBuilder, IPropertyTypeCustomizationUtils& StructCustomizationUtils) override {}
+
+private:
+
+	void OnCVarsCommitted(const FText& NewText, ETextCommit::Type);
+	void OnCVarsChanged(const FText& NewText);
+	FText GetCVarText() const;
+
+	TSharedPtr<IPropertyHandle> PropertyHandle;
+};
+
+
+} // namespace MovieScene
+} // namespace UE
\ No newline at end of file
diff --git a/Engine/Source/Editor/MovieSceneTools/Private/MovieSceneToolsModule.cpp b/Engine/Source/Editor/MovieSceneTools/Private/MovieSceneToolsModule.cpp
--- a/Engine/Source/Editor/MovieSceneTools/Private/MovieSceneToolsModule.cpp
+++ b/Engine/Source/Editor/MovieSceneTools/Private/MovieSceneToolsModule.cpp
@@ -8,6 +8,7 @@
 #include "ISequencerModule.h"
 #include "ICurveEditorModule.h"
 #include "MovieSceneToolsProjectSettingsCustomization.h"
+#include "MovieSceneCVarOverridesPropertyTypeCustomization.h"
 #include "Engine/Blueprint.h"
 #include "EdGraph/EdGraph.h"
 #include "K2Node_CustomEvent.h"
@@ -51,6 +52,7 @@
 #include "TrackEditors/ObjectPropertyTrackEditor.h"
 #include "TrackEditors/PrimitiveMaterialTrackEditor.h"
 #include "TrackEditors/CameraShakeSourceShakeTrackEditor.h"
+#include "TrackEditors/CVarTrackEditor.h"
 
 #include "MovieSceneBuiltInEasingFunctionCustomization.h"
 #include "MovieSceneObjectBindingIDCustomization.h"
@@ -140,6 +142,7 @@ void FMovieSceneToolsModule::StartupModule()
 		MPCTrackCreateEditorHandle = SequencerModule.RegisterTrackEditor(FOnCreateTrackEditor::CreateStatic(&FMaterialParameterCollectionTrackEditor::CreateTrackEditor));
 		PrimitiveMaterialCreateEditorHandle = SequencerModule.RegisterTrackEditor(FOnCreateTrackEditor::CreateStatic(&FPrimitiveMaterialTrackEditor::CreateTrackEditor));
 		CameraShakeSourceShakeCreateEditorHandle = SequencerModule.RegisterTrackEditor(FOnCreateTrackEditor::CreateStatic(&FCameraShakeSourceShakeTrackEditor::CreateTrackEditor));
+		CVarTrackCreateEditorHandle = SequencerModule.RegisterTrackEditor(FOnCreateTrackEditor::CreateStatic(&FCVarTrackEditor::CreateTrackEditor));
 
 		RegisterClipboardConversions();
 
@@ -149,6 +152,7 @@ void FMovieSceneToolsModule::StartupModule()
 		PropertyModule.RegisterCustomClassLayout("MovieSceneBuiltInEasingFunction", FOnGetDetailCustomizationInstance::CreateLambda(&MakeShared<FMovieSceneBuiltInEasingFunctionCustomization>));
 		PropertyModule.RegisterCustomPropertyTypeLayout("MovieSceneObjectBindingID", FOnGetPropertyTypeCustomizationInstance::CreateLambda(&MakeShared<FMovieSceneObjectBindingIDCustomization>));
 		PropertyModule.RegisterCustomPropertyTypeLayout("MovieSceneEvent", FOnGetPropertyTypeCustomizationInstance::CreateStatic(&FMovieSceneEventCustomization::MakeInstance));
+		PropertyModule.RegisterCustomPropertyTypeLayout("MovieSceneCVarOverrides", FOnGetPropertyTypeCustomizationInstance::CreateStatic(&UE::MovieScene::FCVarOverridesPropertyTypeCustomization::MakeInstance));
 
 		SequencerModule.RegisterChannelInterface<FMovieSceneBoolChannel>();
 		SequencerModule.RegisterChannelInterface<FMovieSceneByteChannel>();
@@ -271,6 +275,7 @@ void FMovieSceneToolsModule::ShutdownModule()
 	SequencerModule.UnRegisterTrackEditor( MPCTrackCreateEditorHandle );
 	SequencerModule.UnRegisterTrackEditor( ObjectTrackCreateEditorHandle );
 	SequencerModule.UnRegisterTrackEditor( PrimitiveMaterialCreateEditorHandle );
+	SequencerModule.UnRegisterTrackEditor( CVarTrackCreateEditorHandle );
 
 	if (FModuleManager::Get().IsModuleLoaded("PropertyEditor"))
 	{	
diff --git a/Engine/Source/Editor/MovieSceneTools/Private/TrackEditors/CVarTrackEditor.cpp b/Engine/Source/Editor/MovieSceneTools/Private/TrackEditors/CVarTrackEditor.cpp
new file mode 100644
--- /dev/null
+++ b/Engine/Source/Editor/MovieSceneTools/Private/TrackEditors/CVarTrackEditor.cpp
@@ -0,0 +1,182 @@
+// Copyright Epic Games, Inc. All Rights Reserved.
+
+#include "TrackEditors/CVarTrackEditor.h"
+#include "Tracks/MovieSceneCVarTrack.h"
+#include "Widgets/SBoxPanel.h"
+#include "Framework/MultiBox/MultiBoxBuilder.h"
+#include "EditorStyleSet.h"
+#include "Sections/MovieSceneCVarSection.h"
+#include "SequencerUtilities.h"
+#include "MovieSceneToolHelpers.h"
+
+#include "SequencerSectionPainter.h"
+
+#define LOCTEXT_NAMESPACE "FCVarTrackEditor"
+
+FCVarTrackEditor::FCVarTrackEditor(TSharedRef<ISequencer> InSequencer)
+	: FMovieSceneTrackEditor(InSequencer) 
+{
+}
+
+void FCVarTrackEditor::OnRelease()
+{	
+}
+
+TSharedRef<ISequencerTrackEditor> FCVarTrackEditor::CreateTrackEditor(TSharedRef<ISequencer> InSequencer)
+{
+	return MakeShareable(new FCVarTrackEditor(InSequencer));
+}
+
+
+// ISequencerTrackEditor interface
+void FCVarTrackEditor::BuildAddTrackMenu(FMenuBuilder& MenuBuilder)
+{
+	MenuBuilder.AddMenuEntry(
+		LOCTEXT("AddCVarTrack", "Console Variable Track"),
+		LOCTEXT("AddCVarTooltip", "Adds a new Console Variable track which will allow you to change console variables while a sequence plays."),
+		FSlateIcon(FEditorStyle::GetStyleSetName(), "Sequencer.Tracks.CVar"),
+		FUIAction(
+			FExecuteAction::CreateRaw(this, &FCVarTrackEditor::HandleAddCVarTrackMenuEntryExecute)
+		)
+	);
+}
+
+TSharedPtr<SWidget> FCVarTrackEditor::BuildOutlinerEditWidget(const FGuid& ObjectBinding, UMovieSceneTrack* Track, const FBuildEditWidgetParams& Params)
+{
+	TWeakPtr<ISequencer> WeakSequencer = GetSequencer();
+	const int32 RowIndex = Params.TrackInsertRowIndex;
+	auto OnClickedCallback = [=]() -> FReply
+	{
+		FSequencerUtilities::CreateNewSection(Track, WeakSequencer, RowIndex, EMovieSceneBlendType::Absolute);
+		return FReply::Handled();
+	};
+	
+	// Create a container edit box
+	return SNew(SHorizontalBox)
+
+	// Add the camera combo box
+	+ SHorizontalBox::Slot()
+	.AutoWidth()
+	.VAlign(VAlign_Center)
+	[
+		FSequencerUtilities::MakeAddButton(LOCTEXT("VarText", "Section"), FOnClicked::CreateLambda(OnClickedCallback), Params.NodeIsHovered, GetSequencer())
+	];
+}
+
+
+TSharedRef<ISequencerSection> FCVarTrackEditor::MakeSectionInterface(UMovieSceneSection& SectionObject, UMovieSceneTrack& Track, FGuid ObjectBinding)
+{
+	check(SupportsType(SectionObject.GetOuter()->GetClass()));
+	return MakeShareable(new FCVarSection(SectionObject, GetSequencer()));
+}
+
+
+bool FCVarTrackEditor::SupportsSequence(UMovieSceneSequence* InSequence) const
+{
+	ETrackSupport TrackSupported = InSequence ? InSequence->IsTrackSupported(UMovieSceneCVarTrack::StaticClass()) : ETrackSupport::NotSupported;
+	return TrackSupported == ETrackSupport::Supported;
+}
+
+bool FCVarTrackEditor::SupportsType(TSubclassOf<UMovieSceneTrack> Type) const
+{
+	return (Type == UMovieSceneCVarTrack::StaticClass());
+}
+
+const FSlateBrush* FCVarTrackEditor::GetIconBrush() const
+{
+	return FEditorStyle::GetBrush("Sequencer.Tracks.CVar");
+}
+// ~SequencerTrackEditor interface
+
+
+// FCVarTrackEditor implementation
+void FCVarTrackEditor::HandleAddCVarTrackMenuEntryExecute()
+{
+	UMovieScene* FocusedMovieScene = GetFocusedMovieScene();
+
+	if (FocusedMovieScene == nullptr)
+	{
+		return;
+	}
+
+	if (FocusedMovieScene->IsReadOnly())
+	{
+		return;
+	}
+
+	const FScopedTransaction Transaction(NSLOCTEXT("Sequencer", "AddCVarTrack_Transaction", "Add Console Variable Track"));
+	FocusedMovieScene->Modify();
+	
+	UMovieSceneCVarTrack* NewTrack = FocusedMovieScene->AddMasterTrack<UMovieSceneCVarTrack>();
+	ensure(NewTrack);
+
+	NewTrack->SetDisplayName(LOCTEXT("CVarTrackName", "Console Variable"));
+
+	if (GetSequencer().IsValid())
+	{
+		GetSequencer()->OnAddTrack(NewTrack, FGuid());
+	}
+}
+
+
+FCVarSection::FCVarSection( UMovieSceneSection& InSection, TWeakPtr<ISequencer> InSequencer )
+	: Section( InSection )
+	, Sequencer(InSequencer)
+{
+}
+
+FCVarSection::~FCVarSection()
+{
+}
+
+UMovieSceneSection* FCVarSection::GetSectionObject()
+{ 
+	return &Section;
+}
+
+FText FCVarSection::GetSectionTitle() const
+{
+	return GetCVarText();
+}
+
+FText FCVarSection::GetCVarText() const
+{
+	UMovieSceneCVarSection* CVarSection = Cast<UMovieSceneCVarSection>(&Section);
+	if (CVarSection)
+	{
+		return FText::FromString(CVarSection->GetString());
+	}
+	return LOCTEXT("NoCVarTitleName", "No Variables");
+}
+
+FText FCVarSection::GetSectionToolTip() const
+{
+	return GetSectionTitle();
+}
+
+int32 FCVarSection::OnPaintSection( FSequencerSectionPainter& Painter ) const
+{
+	int32 LayerId = Painter.PaintSectionBackground();
+	return LayerId;
+}
+
+void FCVarSection::BeginResizeSection()
+{
+}
+
+void FCVarSection::ResizeSection(ESequencerSectionResizeMode ResizeMode, FFrameNumber ResizeTime)
+{
+	ISequencerSection::ResizeSection(ResizeMode, ResizeTime);
+}
+
+void FCVarSection::BeginSlipSection()
+{
+
+}
+
+void FCVarSection::SlipSection(FFrameNumber SlipTime)
+{
+	ISequencerSection::SlipSection(SlipTime);
+}
+
+#undef LOCTEXT_NAMESPACE
\ No newline at end of file
diff --git a/Engine/Source/Editor/MovieSceneTools/Private/TrackEditors/CVarTrackEditor.h b/Engine/Source/Editor/MovieSceneTools/Private/TrackEditors/CVarTrackEditor.h
new file mode 100644
--- /dev/null
+++ b/Engine/Source/Editor/MovieSceneTools/Private/TrackEditors/CVarTrackEditor.h
@@ -0,0 +1,102 @@
+// Copyright Epic Games, Inc. All Rights Reserved.
+
+#pragma once
+
+#include "CoreMinimal.h"
+#include "Misc/Guid.h"
+#include "Templates/SubclassOf.h"
+#include "Widgets/SWidget.h"
+#include "ISequencer.h"
+#include "MovieSceneTrack.h"
+#include "ISequencerSection.h"
+#include "ISequencerTrackEditor.h"
+#include "MovieSceneTrackEditor.h"
+#include "MovieSceneToolHelpers.h"
+
+class FMenuBuilder;
+class FCVarTrackEditor;
+class FCVarSection;
+
+/**
+ * Tools for Console Variable tracks
+ */
+class FCVarTrackEditor
+	: public FMovieSceneTrackEditor
+{
+public:
+
+	/**
+	 * Constructor
+	 *
+	 * @param InSequencer The sequencer instance to be used by this tool.
+	 */
+	FCVarTrackEditor(TSharedRef<ISequencer> InSequencer);
+
+	/** Virtual destructor. */
+	virtual ~FCVarTrackEditor() { }
+
+	/**
+	 * Creates an instance of this class.  Called by a sequencer .
+	 *
+	 * @param OwningSequencer The sequencer instance to be used by this tool.
+	 * @return The new instance of this class.
+	 */
+	static TSharedRef<ISequencerTrackEditor> CreateTrackEditor(TSharedRef<ISequencer> OwningSequencer);
+
+public:
+
+	// ISequencerTrackEditor interface
+
+
+	virtual void BuildAddTrackMenu(FMenuBuilder& MenuBuilder) override;
+	virtual TSharedPtr<SWidget> BuildOutlinerEditWidget(const FGuid& ObjectBinding, UMovieSceneTrack* Track, const FBuildEditWidgetParams& Params) override;
+	virtual TSharedRef<ISequencerSection> MakeSectionInterface(UMovieSceneSection& SectionObject, UMovieSceneTrack& Track, FGuid ObjectBinding) override;
+	virtual void OnRelease() override;
+	virtual bool SupportsType(TSubclassOf<UMovieSceneTrack> Type) const override;
+	virtual bool SupportsSequence(UMovieSceneSequence* InSequence) const override;
+
+	virtual const FSlateBrush* GetIconBrush() const override;
+private:
+
+	/** Callback for executing the "Add Console Variable Track" menu entry. */
+	void HandleAddCVarTrackMenuEntryExecute();
+};
+
+/**
+ * Class for Console Variable sections
+ */
+class FCVarSection
+	: public ISequencerSection
+	, public TSharedFromThis<FCVarSection>
+{
+public:
+
+	/** Constructor. */
+	FCVarSection(UMovieSceneSection& InSection, TWeakPtr<ISequencer> InSequencer);
+
+	/** Virtual destructor. */
+	virtual ~FCVarSection();
+
+public:
+
+	// ISequencerSection interface
+	virtual UMovieSceneSection* GetSectionObject() override;
+	virtual int32 OnPaintSection( FSequencerSectionPainter& Painter ) const override;
+	virtual FText GetSectionTitle() const override;
+	virtual FText GetSectionToolTip() const override;
+	virtual void BeginResizeSection() override;
+	virtual void ResizeSection(ESequencerSectionResizeMode ResizeMode, FFrameNumber ResizeTime) override;
+	virtual void BeginSlipSection() override;
+	virtual void SlipSection(FFrameNumber SlipTime) override;
+	virtual float GetSectionHeight() const { return SequencerSectionConstants::DefaultSectionHeight*3.f; }
+
+private:
+
+	FText GetCVarText() const;
+
+	/** The section we are visualizing. */
+	UMovieSceneSection& Section;
+
+	TWeakPtr<ISequencer> Sequencer;
+};
+
diff --git a/Engine/Source/Editor/MovieSceneTools/Public/MovieSceneToolsModule.h b/Engine/Source/Editor/MovieSceneTools/Public/MovieSceneToolsModule.h
--- a/Engine/Source/Editor/MovieSceneTools/Public/MovieSceneToolsModule.h
+++ b/Engine/Source/Editor/MovieSceneTools/Public/MovieSceneToolsModule.h
@@ -119,6 +119,7 @@ private:
 	FDelegateHandle MPCTrackCreateEditorHandle;
 	FDelegateHandle PrimitiveMaterialCreateEditorHandle;
 	FDelegateHandle CameraShakeSourceShakeCreateEditorHandle;
+	FDelegateHandle CVarTrackCreateEditorHandle;
 
 	FDelegateHandle GenerateEventEntryPointsHandle;
 	FDelegateHandle FixupPayloadParameterNameHandle;
diff --git a/Engine/Source/Editor/Sequencer/Private/SequencerUtilities.cpp b/Engine/Source/Editor/Sequencer/Private/SequencerUtilities.cpp
--- a/Engine/Source/Editor/Sequencer/Private/SequencerUtilities.cpp
+++ b/Engine/Source/Editor/Sequencer/Private/SequencerUtilities.cpp
@@ -32,7 +32,7 @@
 #include "Framework/Notifications/NotificationManager.h"
 #include "Widgets/Notifications/SNotificationList.h"
 
-
+PRAGMA_DISABLE_OPTIMIZATION
 
 #define LOCTEXT_NAMESPACE "FSequencerUtilities"
 
@@ -49,6 +49,19 @@ static EVisibility GetRolloverVisibility(TAttribute<bool> HoverState, TWeakPtr<S
 	}
 }
 
+static EVisibility GetRolloverVisibility(TAttribute<bool> HoverState, TWeakPtr<SButton> WeakButton)
+{
+	TSharedPtr<SButton> Button = WeakButton.Pin();
+	if (HoverState.Get())
+	{
+		return EVisibility::SelfHitTestInvisible;
+	}
+	else
+	{
+		return EVisibility::Collapsed;
+	}
+}
+
 TSharedRef<SWidget> FSequencerUtilities::MakeAddButton(FText HoverText, FOnGetContent MenuContent, const TAttribute<bool>& HoverState, TWeakPtr<ISequencer> InSequencer)
 {
 	FSlateFontInfo SmallLayoutFont = FCoreStyle::GetDefaultFontStyle("Regular", 8);
@@ -98,6 +111,101 @@ TSharedRef<SWidget> FSequencerUtilities::MakeAddButton(FText HoverText, FOnGetCo
 	return ComboButton;
 }
 
+TSharedRef<SWidget> FSequencerUtilities::MakeAddButton(FText HoverText, FOnClicked OnClicked, const TAttribute<bool>& HoverState, TWeakPtr<ISequencer> InSequencer)
+{
+	FSlateFontInfo SmallLayoutFont = FCoreStyle::GetDefaultFontStyle("Regular", 8);
+
+	TSharedRef<STextBlock> ButtonText = SNew(STextBlock)
+		.Text(HoverText)
+		.Font(SmallLayoutFont)
+		.ColorAndOpacity(FSlateColor::UseForeground());
+
+	TSharedRef<SButton> Button =
+
+		SNew(SButton)
+		.IsFocusable(true)
+		.ButtonStyle(FEditorStyle::Get(), "HoverHintOnly")
+		.ForegroundColor(FSlateColor::UseForeground())
+		.IsEnabled_Lambda([=]() { return InSequencer.IsValid() ? !InSequencer.Pin()->IsReadOnly() : false; })
+		.OnClicked(OnClicked)
+		.ContentPadding(FMargin(5, 2))
+		.HAlign(HAlign_Center)
+		.VAlign(VAlign_Center)
+		.Content()
+		[
+			SNew(SHorizontalBox)
+
+			+ SHorizontalBox::Slot()
+			.AutoWidth()
+			.VAlign(VAlign_Center)
+			.Padding(FMargin(0, 0, 2, 0))
+			[
+				SNew(SImage)
+				.ColorAndOpacity(FSlateColor::UseForeground())
+				.Image(FEditorStyle::GetBrush("Plus"))
+			]
+
+			+ SHorizontalBox::Slot()
+			.VAlign(VAlign_Center)
+			.AutoWidth()
+			[
+				ButtonText
+			]
+		];
+
+	TAttribute<EVisibility> Visibility = TAttribute<EVisibility>::Create(TAttribute<EVisibility>::FGetter::CreateStatic(GetRolloverVisibility, HoverState, TWeakPtr<SButton>(Button)));
+	ButtonText->SetVisibility(Visibility);
+
+	return Button;
+}
+
+void FSequencerUtilities::CreateNewSection(UMovieSceneTrack* InTrack, TWeakPtr<ISequencer> InSequencer, int32 InRowIndex, EMovieSceneBlendType InBlendType)
+{
+	TSharedPtr<ISequencer> Sequencer = InSequencer.Pin();
+	if (!Sequencer.IsValid())
+	{
+		return;
+	}
+
+	FQualifiedFrameTime CurrentTime = Sequencer->GetLocalTime();
+	FFrameNumber PlaybackEnd = UE::MovieScene::DiscreteExclusiveUpper(Sequencer->GetFocusedMovieSceneSequence()->GetMovieScene()->GetPlaybackRange());
+
+	FScopedTransaction Transaction(LOCTEXT("AddSectionTransactionText", "Add Section"));
+	if (UMovieSceneSection* NewSection = InTrack->CreateNewSection())
+	{
+		int32 OverlapPriority = 0;
+		for (UMovieSceneSection* Section : InTrack->GetAllSections())
+		{
+			OverlapPriority = FMath::Max(Section->GetOverlapPriority() + 1, OverlapPriority);
+
+			// Move existing sections on the same row or beyond so that they don't overlap with the new section
+			if (Section != NewSection && Section->GetRowIndex() >= InRowIndex)
+			{
+				Section->SetRowIndex(Section->GetRowIndex() + 1);
+			}
+		}
+
+		InTrack->Modify();
+
+		NewSection->SetRange(TRange<FFrameNumber>(CurrentTime.Time.FrameNumber, PlaybackEnd));
+		NewSection->SetOverlapPriority(OverlapPriority);
+		NewSection->SetRowIndex(InRowIndex);
+		NewSection->SetBlendType(InBlendType);
+
+		InTrack->AddSection(*NewSection);
+		InTrack->UpdateEasing();
+
+		Sequencer->NotifyMovieSceneDataChanged(EMovieSceneDataChangeType::MovieSceneStructureItemAdded);
+		Sequencer->EmptySelection();
+		Sequencer->SelectSection(NewSection);
+		Sequencer->ThrobSectionSelection();
+	}
+	else
+	{
+		Transaction.Cancel();
+	}
+}
+
 void FSequencerUtilities::PopulateMenu_CreateNewSection(FMenuBuilder& MenuBuilder, int32 RowIndex, UMovieSceneTrack* Track, TWeakPtr<ISequencer> InSequencer)
 {
 	if (!Track)
@@ -562,3 +670,4 @@ void FSequencerUtilities::ShowReadOnlyError()
 
 
 #undef LOCTEXT_NAMESPACE
+PRAGMA_ENABLE_OPTIMIZATION
\ No newline at end of file
diff --git a/Engine/Source/Editor/Sequencer/Public/SequencerUtilities.h b/Engine/Source/Editor/Sequencer/Public/SequencerUtilities.h
--- a/Engine/Source/Editor/Sequencer/Public/SequencerUtilities.h
+++ b/Engine/Source/Editor/Sequencer/Public/SequencerUtilities.h
@@ -13,10 +13,16 @@ class ISequencer;
 class FMenuBuilder;
 class ULevelSequence;
 class UMovieSceneCompiledDataManager;
+enum class EMovieSceneBlendType : uint8;
 
 struct SEQUENCER_API FSequencerUtilities
 {
+	/* Creates a button (used for +Section) that opens a ComboButton with a user-defined sub-menu content. */
 	static TSharedRef<SWidget> MakeAddButton(FText HoverText, FOnGetContent MenuContent, const TAttribute<bool>& HoverState, TWeakPtr<ISequencer> InSequencer);
+	/* Creates a button (used for +Section) that fires a user-defined OnClick response with no sub-menu. */
+	static TSharedRef<SWidget> MakeAddButton(FText HoverText, FOnClicked OnClicked, const TAttribute<bool>& HoverState, TWeakPtr<ISequencer> InSequencer);
+
+	static void CreateNewSection(UMovieSceneTrack* InTrack, TWeakPtr<ISequencer> InSequencer, int32 InRowIndex, EMovieSceneBlendType InBlendType);
 
 	static void PopulateMenu_CreateNewSection(FMenuBuilder& MenuBuilder, int32 RowIndex, UMovieSceneTrack* Track, TWeakPtr<ISequencer> InSequencer);
 
diff --git a/Engine/Source/Runtime/LevelSequence/Private/LevelSequence.cpp b/Engine/Source/Runtime/LevelSequence/Private/LevelSequence.cpp
--- a/Engine/Source/Runtime/LevelSequence/Private/LevelSequence.cpp
+++ b/Engine/Source/Runtime/LevelSequence/Private/LevelSequence.cpp
@@ -28,6 +28,7 @@
 #include "Tracks/MovieSceneSlomoTrack.h"
 #include "Tracks/MovieSceneSpawnTrack.h"
 #include "Tracks/MovieSceneSubTrack.h"
+#include "Tracks/MovieSceneCVarTrack.h"
 #include "Modules/ModuleManager.h"
 #include "LevelSequencePlayer.h"
 #include "Compilation/MovieSceneCompiledDataManager.h"
@@ -117,7 +118,8 @@ ETrackSupport ULevelSequence::IsTrackSupported(TSubclassOf<class UMovieSceneTrac
 		InTrackClass == UMovieSceneMaterialParameterCollectionTrack::StaticClass() ||
 		InTrackClass == UMovieSceneSlomoTrack::StaticClass() ||
 		InTrackClass == UMovieSceneSpawnTrack::StaticClass() ||
-		InTrackClass == UMovieSceneSubTrack::StaticClass())
+		InTrackClass == UMovieSceneSubTrack::StaticClass() ||
+		InTrackClass == UMovieSceneCVarTrack::StaticClass())
 	{
 		return ETrackSupport::Supported;
 	}
diff --git a/Engine/Source/Runtime/MovieSceneTracks/Private/Sections/MovieSceneCVarSection.cpp b/Engine/Source/Runtime/MovieSceneTracks/Private/Sections/MovieSceneCVarSection.cpp
new file mode 100644
--- /dev/null
+++ b/Engine/Source/Runtime/MovieSceneTracks/Private/Sections/MovieSceneCVarSection.cpp
@@ -0,0 +1,91 @@
+// Copyright Epic Games, Inc. All Rights Reserved.
+
+#include "Sections/MovieSceneCVarSection.h"
+#include "EntitySystem/MovieSceneEntityBuilder.h"
+#include "EntitySystem/BuiltInComponentTypes.h"
+#include "TrackInstances/MovieSceneCVarTrackInstance.h"
+
+void FMovieSceneCVarOverrides::SetFromString(const FString& InString)
+{
+	ValuesByCVar.Reset();
+
+	// Split the section's specified value apart by new lines, and then parse out the key/value pairs.
+	TArray<FString> AllLines;
+	const bool bCullEmpty = true;
+	InString.ParseIntoArray(AllLines, LINE_TERMINATOR, bCullEmpty);
+
+	// If there was no LINE_TERMINATOR or they used comma separated values, AllLines will have 1 line in it containing all of InString.
+	TArray<FString> SplitLines;
+	for (const FString& Line : AllLines)
+	{
+		TArray<FString> NewStrings;
+		Line.ParseIntoArray(NewStrings, TEXT(","), bCullEmpty);
+
+		SplitLines.Append(NewStrings);
+	}
+
+	// SplitLines should have at least 1 line in it by extension of AllLines having at least 1 line.
+	for (const FString& Line : SplitLines)
+	{
+		// We should now have things in the format "foo=bar" or "foo bar", so we still need to split to find the value.
+		TArray<FString> OutArray;
+		const TCHAR* Delimiters[] = { TEXT(","), TEXT(" ") };
+		Line.ParseIntoArray(OutArray, Delimiters, UE_ARRAY_COUNT(Delimiters), bCullEmpty);
+		if (OutArray.Num() != 2)
+		{
+			// Not sure what format it's in.
+			UE_LOG(LogMovieScene, Warning, TEXT("Could not parse \"%s\" into a key/value pair. Expected format: \"foo=5.5\" or \"foo 5.5\"!"), *Line);
+			continue;
+		}
+
+		ValuesByCVar.Add(OutArray[0], OutArray[1]);
+	}
+}
+
+FString FMovieSceneCVarOverrides::GetString() const
+{
+	TStringBuilder<256> CVarString;
+	for (const TPair<FString, FString>& Pair : ValuesByCVar)
+	{
+		if (CVarString.Len() != 0)
+		{
+			CVarString += TEXT(", ");
+		}
+
+		CVarString += Pair.Key;
+		CVarString += TEXT(" ");
+		CVarString += Pair.Value;
+	}
+	return CVarString.ToString();
+}
+
+UMovieSceneCVarSection::UMovieSceneCVarSection(const FObjectInitializer& Init)
+	: Super(Init)
+{
+	EvalOptions.EnableAndSetCompletionMode(EMovieSceneCompletionMode::ProjectDefault);
+
+	SetBlendType(EMovieSceneBlendType::Absolute);
+}
+
+void UMovieSceneCVarSection::ImportEntityImpl(UMovieSceneEntitySystemLinker* EntityLinker, const FEntityImportParams& Params, FImportedEntity* OutImportedEntity)
+{
+	using namespace UE::MovieScene;
+
+	FMovieSceneTrackInstanceComponent TrackInstance { decltype(FMovieSceneTrackInstanceComponent::Owner)(this), UMovieSceneCVarTrackInstance::StaticClass() };
+
+	OutImportedEntity->AddBuilder(
+		FEntityBuilder()
+		.AddTag(FBuiltInComponentTypes::Get()->Tags.Master)
+		.Add(FBuiltInComponentTypes::Get()->TrackInstance, TrackInstance)
+	);
+}
+
+void UMovieSceneCVarSection::SetFromString(const FString& InString)
+{
+	ConsoleVariables.SetFromString(InString);
+}
+
+FString UMovieSceneCVarSection::GetString() const
+{
+	return ConsoleVariables.GetString();
+}
diff --git a/Engine/Source/Runtime/MovieSceneTracks/Private/TrackInstances/MovieSceneCVarTrackInstance.cpp b/Engine/Source/Runtime/MovieSceneTracks/Private/TrackInstances/MovieSceneCVarTrackInstance.cpp
new file mode 100644
--- /dev/null
+++ b/Engine/Source/Runtime/MovieSceneTracks/Private/TrackInstances/MovieSceneCVarTrackInstance.cpp
@@ -0,0 +1,160 @@
+// Copyright Epic Games, Inc. All Rights Reserved.
+
+#include "TrackInstances/MovieSceneCVarTrackInstance.h"
+#include "Sections/MovieSceneCVarSection.h"
+#include "Tracks/MovieSceneCVarTrack.h"
+#include "EntitySystem/MovieSceneEntitySystemLinker.h"
+#include "EntitySystem/MovieSceneInstanceRegistry.h"
+#include "EntitySystem/TrackInstance/MovieSceneTrackInstanceSystem.h"
+#include "Evaluation/PreAnimatedState/MovieScenePreAnimatedStateTypes.h"
+#include "MovieSceneCommonHelpers.h"
+
+namespace UE
+{
+namespace MovieScene
+{
+
+struct FPreAnimatedCVarTraits
+{
+	using KeyType     = FString;
+	using StorageType = FString;
+
+	static void RestorePreAnimatedValue(const FString& InKey, const FString& CachedValue, const FRestoreStateParams& Params)
+	{
+		static IConsoleManager& ConsoleManager = IConsoleManager::Get();
+		IConsoleVariable* ConsoleVariable = ConsoleManager.FindConsoleVariable(*InKey);
+		if (ConsoleVariable)
+		{
+			ConsoleVariable->Set(*CachedValue);
+		}
+	}
+
+	static FString CachePreAnimatedValue(const FString& InKey)
+	{
+		static IConsoleManager& ConsoleManager = IConsoleManager::Get();
+
+		IConsoleVariable* ConsoleVariable = ConsoleManager.FindConsoleVariable(*InKey);
+		return ConsoleVariable ? ConsoleVariable->GetString() : FString();
+	}
+};
+struct FPreAnimatedCVarStorage : TPreAnimatedStateStorage<FPreAnimatedCVarTraits>
+{
+	static TAutoRegisterPreAnimatedStorageID<FPreAnimatedCVarStorage> StorageID;
+};
+TAutoRegisterPreAnimatedStorageID<FPreAnimatedCVarStorage> FPreAnimatedCVarStorage::StorageID;
+
+} // namespace MovieScene
+} // namespace UE
+
+void UMovieSceneCVarTrackInstance::OnBeginUpdateInputs()
+{
+	CVarsNeedingUpdate.Reset();
+}
+
+void UMovieSceneCVarTrackInstance::OnAnimate()
+{
+}
+
+void UMovieSceneCVarTrackInstance::OnInputAdded(const FMovieSceneTrackInstanceInput& InInput)
+{
+	using namespace UE::MovieScene;
+
+	UMovieSceneEntitySystemLinker* Linker  = GetLinker();
+	UMovieSceneCVarSection*        Section = CastChecked<UMovieSceneCVarSection>(InInput.Section);
+
+	FScopedPreAnimatedCaptureSource CaptureSource(Linker, InInput);
+
+	// Capture pre-animated state for all the cvars
+	if (Linker->PreAnimatedState.IsCapturingGlobalState() || CaptureSource.WantsRestoreState())
+	{
+		static const FPreAnimatedStorageGroupHandle UngroupedHandle;
+
+		TSharedPtr<FPreAnimatedCVarStorage> CVarStorage = Linker->PreAnimatedState.GetOrCreateStorage<FPreAnimatedCVarStorage>();
+
+		for (const TPair<FString, FString>& Pair : Section->ConsoleVariables.ValuesByCVar)
+		{
+			CVarStorage->SavePreAnimatedState(UngroupedHandle, Pair.Key, FPreAnimatedCVarTraits::CachePreAnimatedValue);
+		}
+	}
+
+	for (const TPair<FString, FString>& Pair : Section->ConsoleVariables.ValuesByCVar)
+	{
+		CVarsNeedingUpdate.Add(Pair.Key);
+	}
+}
+
+void UMovieSceneCVarTrackInstance::OnInputRemoved(const FMovieSceneTrackInstanceInput& InInput)
+{
+	using namespace UE::MovieScene;
+
+	UMovieSceneCVarSection* Section = CastChecked<UMovieSceneCVarSection>(InInput.Section);
+	
+	for (const TPair<FString, FString>& Pair : Section->ConsoleVariables.ValuesByCVar)
+	{
+		CVarsNeedingUpdate.Add(Pair.Key);
+	}
+}
+
+void UMovieSceneCVarTrackInstance::OnEndUpdateInputs()
+{
+	using namespace UE::MovieScene;
+
+	const FInstanceRegistry* InstanceRegistry = GetLinker()->GetInstanceRegistry();
+
+	struct FCVarInfo
+	{
+		FString Value;
+		int32 HierarchicalBias;
+	};
+	TMap<FString, FCVarInfo> CVarInfo;
+
+	for (const FMovieSceneTrackInstanceInput& Input : GetInputs())
+	{
+		const int32 HierarchicalBias = InstanceRegistry->GetInstance(Input.InstanceHandle).GetContext().GetHierarchicalBias();
+
+		UMovieSceneCVarSection* Section = CastChecked<UMovieSceneCVarSection>(Input.Section);
+		for (const TPair<FString, FString>& Pair : Section->ConsoleVariables.ValuesByCVar)
+		{
+			// If CVarsNeedingUpdate is empty that implies that a section was changed so we need to reapply everything
+			if (CVarsNeedingUpdate.Num() == 0 || CVarsNeedingUpdate.Contains(Pair.Key))
+			{
+				FCVarInfo* ExistingInfo = CVarInfo.Find(Pair.Key);
+				if (ExistingInfo)
+				{
+					// If this input has greater hbias, it should override the existing one
+					if (HierarchicalBias > ExistingInfo->HierarchicalBias)
+					{
+						ExistingInfo->Value = Pair.Value;
+					}
+					// Currently we do not support blending with cvars, so this is a warning
+					else if (HierarchicalBias == ExistingInfo->HierarchicalBias)
+					{
+						UE_LOG(LogMovieScene, Warning, TEXT("Multiple sections animating the same CVar '%s', this is not supported."), *Pair.Key);
+					}
+				}
+				else
+				{
+					CVarInfo.Add(Pair.Key, FCVarInfo{ Pair.Value, HierarchicalBias });
+				}
+			}
+		}
+	}
+
+	IConsoleManager& ConsoleManager = IConsoleManager::Get();
+	for (const TPair<FString, FCVarInfo>& Pair : CVarInfo)
+	{
+		IConsoleVariable* ConsoleVariable = ConsoleManager.FindConsoleVariable(*Pair.Key);
+		if (ConsoleVariable)
+		{
+			ConsoleVariable->Set(*Pair.Value.Value);
+		}
+		else
+		{
+			UE_LOG(LogMovieScene, Warning, TEXT("Attempting to animate undefined CVar '%s' with value '%s'."), *Pair.Key, *Pair.Value.Value);
+		}
+	}
+}
+
+void UMovieSceneCVarTrackInstance::OnDestroyed()
+{
+}
diff --git a/Engine/Source/Runtime/MovieSceneTracks/Private/TrackInstances/MovieSceneCVarTrackInstance.h b/Engine/Source/Runtime/MovieSceneTracks/Private/TrackInstances/MovieSceneCVarTrackInstance.h
new file mode 100644
--- /dev/null
+++ b/Engine/Source/Runtime/MovieSceneTracks/Private/TrackInstances/MovieSceneCVarTrackInstance.h
@@ -0,0 +1,26 @@
+// Copyright Epic Games, Inc. All Rights Reserved.
+
+#pragma once
+
+#include "CoreMinimal.h"
+#include "EntitySystem/BuiltInComponentTypes.h"
+#include "EntitySystem/TrackInstance/MovieSceneTrackInstance.h"
+#include "UObject/ObjectMacros.h"
+#include "Sections/MovieSceneCVarSection.h"
+#include "MovieSceneCVarTrackInstance.generated.h"
+
+UCLASS()
+class UMovieSceneCVarTrackInstance : public UMovieSceneTrackInstance
+{
+	GENERATED_BODY()
+
+private:
+	virtual void OnBeginUpdateInputs() override;
+	virtual void OnAnimate() override;
+	virtual void OnInputAdded(const FMovieSceneTrackInstanceInput& InInput) override;
+	virtual void OnInputRemoved(const FMovieSceneTrackInstanceInput& InInput) override;
+	virtual void OnEndUpdateInputs() override;
+	virtual void OnDestroyed() override;
+
+	TSet<FString> CVarsNeedingUpdate;
+};
diff --git a/Engine/Source/Runtime/MovieSceneTracks/Private/Tracks/MovieSceneCVarTrack.cpp b/Engine/Source/Runtime/MovieSceneTracks/Private/Tracks/MovieSceneCVarTrack.cpp
new file mode 100644
--- /dev/null
+++ b/Engine/Source/Runtime/MovieSceneTracks/Private/Tracks/MovieSceneCVarTrack.cpp
@@ -0,0 +1,97 @@
+// Copyright Epic Games, Inc. All Rights Reserved.
+
+#include "Tracks/MovieSceneCVarTrack.h"
+#include "MovieScene.h"
+#include "MovieSceneCommonHelpers.h"
+#include "Sections/MovieSceneCVarSection.h"
+
+
+#define LOCTEXT_NAMESPACE "MovieSceneCVarTrack"
+
+UMovieSceneCVarTrack::UMovieSceneCVarTrack( const FObjectInitializer& ObjectInitializer )
+	: Super( ObjectInitializer )
+{
+#if WITH_EDITORONLY_DATA
+	TrackTint = FColor(23, 89, 145, 65);
+#endif
+
+	// By default, don't evaluate cvar values in pre and postroll as that would conflict with
+	// the previous shot's desired values.
+	EvalOptions.bEvaluateInPreroll = EvalOptions.bEvaluateInPostroll = false;
+	// SupportedBlendTypes.Add(EMovieSceneBlendType::Absolute);
+}
+
+/* UMovieSceneTrack interface
+ *****************************************************************************/
+
+void UMovieSceneCVarTrack::AddSection(UMovieSceneSection& Section)
+{
+	if (UMovieSceneCVarSection* CVarSection = Cast<UMovieSceneCVarSection>(&Section))
+	{
+		Sections.Add(CVarSection);
+	}
+	
+	// todo: Should this call MovieSceneHelpers::SortConsecutiveSections(Sections); ?
+}
+
+bool UMovieSceneCVarTrack::SupportsType(TSubclassOf<UMovieSceneSection> SectionClass) const
+{
+	return SectionClass == UMovieSceneCVarSection::StaticClass();
+}
+
+UMovieSceneSection* UMovieSceneCVarTrack::CreateNewSection()
+{
+	return NewObject<UMovieSceneCVarSection>(this, NAME_None, RF_Transactional);
+}
+
+bool UMovieSceneCVarTrack::SupportsMultipleRows() const
+{
+	return false;
+}
+
+EMovieSceneTrackEasingSupportFlags UMovieSceneCVarTrack::SupportsEasing(FMovieSceneSupportsEasingParams& Params) const
+{
+	return EMovieSceneTrackEasingSupportFlags::None;
+}
+
+const TArray<UMovieSceneSection*>& UMovieSceneCVarTrack::GetAllSections() const
+{
+	return Sections;
+}
+
+bool UMovieSceneCVarTrack::HasSection(const UMovieSceneSection& Section) const
+{
+	return Sections.Contains(&Section);
+}
+
+
+bool UMovieSceneCVarTrack::IsEmpty() const
+{
+	return Sections.Num() == 0;
+}
+
+void UMovieSceneCVarTrack::RemoveSection(UMovieSceneSection& Section)
+{
+	Sections.Remove(&Section);
+}
+
+void UMovieSceneCVarTrack::RemoveSectionAt(int32 SectionIndex)
+{
+	UMovieSceneSection* SectionToDelete = Sections[SectionIndex];
+	Sections.RemoveAt(SectionIndex);
+	MovieSceneHelpers::SortConsecutiveSections(Sections);
+}
+
+void UMovieSceneCVarTrack::RemoveAllAnimationData()
+{
+	Sections.Empty();
+}
+
+#if WITH_EDITORONLY_DATA
+FText UMovieSceneCVarTrack::GetDefaultDisplayName() const
+{
+	return LOCTEXT("TrackName", "Console Variable");
+}
+#endif
+
+#undef LOCTEXT_NAMESPACE
diff --git a/Engine/Source/Runtime/MovieSceneTracks/Public/Sections/MovieSceneCVarSection.h b/Engine/Source/Runtime/MovieSceneTracks/Public/Sections/MovieSceneCVarSection.h
new file mode 100644
--- /dev/null
+++ b/Engine/Source/Runtime/MovieSceneTracks/Public/Sections/MovieSceneCVarSection.h
@@ -0,0 +1,56 @@
+// Copyright Epic Games, Inc. All Rights Reserved.
+
+#pragma once
+
+#include "CoreMinimal.h"
+#include "EntitySystem/IMovieSceneEntityProvider.h"
+#include "EntitySystem/MovieSceneEntityIDs.h"
+#include "MovieSceneSection.h"
+#include "MovieSceneCVarSection.generated.h"
+
+USTRUCT()
+struct FMovieSceneCVarOverrides
+{
+	GENERATED_BODY()
+
+	MOVIESCENETRACKS_API void SetFromString(const FString& InString);
+
+	MOVIESCENETRACKS_API FString GetString() const;
+
+	/** The name of the console variable and the value, separated by ' ' or '=', ie: "foo.bar=1" or "foo.bar 1". */
+	UPROPERTY(EditAnywhere, Category="Console Variables")
+	TMap<FString, FString> ValuesByCVar;
+};
+
+/**
+ * A CVar section is responsible for applying a user-supplied value to the specified cvar, and then restoring the previous value after the section ends.
+ */
+UCLASS(MinimalAPI)
+class UMovieSceneCVarSection 
+	: public UMovieSceneSection
+	, public IMovieSceneEntityProvider
+{
+	GENERATED_BODY()
+
+public:
+	/** Constructs a new console variable section */
+	UMovieSceneCVarSection(const FObjectInitializer& Init);
+
+	UFUNCTION(BlueprintCallable, Category="CVars")
+	MOVIESCENETRACKS_API void SetFromString(const FString& InString);
+
+	UFUNCTION(BlueprintCallable, Category="CVars")
+	MOVIESCENETRACKS_API FString GetString() const;
+
+private:
+
+	virtual void ImportEntityImpl(UMovieSceneEntitySystemLinker* EntityLinker, const FEntityImportParams& Params, FImportedEntity* OutImportedEntity) override;
+
+public:
+
+	/** The name of the console variable and the value, separated by ' ' or '=', ie: "foo.bar=1" or "foo.bar 1". */
+	UPROPERTY(EditAnywhere,  meta=(MultiLine=true), Category="Console Variables")
+	FMovieSceneCVarOverrides ConsoleVariables;
+
+	friend class UMovieSceneCVarTrackInstance;
+};
diff --git a/Engine/Source/Runtime/MovieSceneTracks/Public/Tracks/MovieSceneCVarTrack.h b/Engine/Source/Runtime/MovieSceneTracks/Public/Tracks/MovieSceneCVarTrack.h
new file mode 100644
--- /dev/null
+++ b/Engine/Source/Runtime/MovieSceneTracks/Public/Tracks/MovieSceneCVarTrack.h
@@ -0,0 +1,44 @@
+// Copyright Epic Games, Inc. All Rights Reserved.
+
+#pragma once
+
+#include "CoreMinimal.h"
+#include "UObject/ObjectMacros.h"
+#include "MovieSceneNameableTrack.h"
+#include "MovieSceneCVarTrack.generated.h"
+
+/**
+ * Track for setting (and restoring) Console Variables during playback.
+ */
+UCLASS(MinimalAPI)
+class UMovieSceneCVarTrack
+	: public UMovieSceneNameableTrack
+{
+	GENERATED_BODY()
+	UMovieSceneCVarTrack( const FObjectInitializer& ObjectInitializer );
+	
+public:
+
+	// UMovieSceneTrack interface
+	virtual void AddSection(UMovieSceneSection& Section) override;
+	virtual bool SupportsType(TSubclassOf<UMovieSceneSection> SectionClass) const override;
+	virtual UMovieSceneSection* CreateNewSection() override;
+	virtual bool SupportsMultipleRows() const override;
+	virtual EMovieSceneTrackEasingSupportFlags SupportsEasing(FMovieSceneSupportsEasingParams& Params) const override;
+	virtual const TArray<UMovieSceneSection*>& GetAllSections() const override;
+	virtual bool HasSection(const UMovieSceneSection& Section) const override;
+	virtual bool IsEmpty() const override;
+	virtual void RemoveSection(UMovieSceneSection& Section) override;
+	virtual void RemoveSectionAt(int32 SectionIndex) override;
+	virtual void RemoveAllAnimationData() override;
+
+#if WITH_EDITORONLY_DATA
+	virtual FText GetDefaultDisplayName() const override;
+#endif
+
+private:
+
+	/** All movie scene sections. */
+	UPROPERTY()
+	TArray<TObjectPtr<UMovieSceneSection>> Sections;
+};
