Adds MOVIESCENETRACKS_API to some MovieScene3DTransformSection functions so they can be used by the cache (now in a plugin)

diff --git a/Engine/Source/Runtime/MovieSceneTracks/Public/Sections/MovieScene3DTransformSection.h b/Engine/Source/Runtime/MovieSceneTracks/Public/Sections/MovieScene3DTransformSection.h
--- a/Engine/Source/Runtime/MovieSceneTracks/Public/Sections/MovieScene3DTransformSection.h
+++ b/Engine/Source/Runtime/MovieSceneTracks/Public/Sections/MovieScene3DTransformSection.h
@@ -222,8 +222,8 @@ public:
 
 	/* From UMovieSection*/
 	
-	virtual bool ShowCurveForChannel(const void *Channel) const override;
-	virtual void SetBlendType(EMovieSceneBlendType InBlendType) override;
+	MOVIESCENETRACKS_API virtual bool ShowCurveForChannel(const void *Channel) const override;
+	MOVIESCENETRACKS_API virtual void SetBlendType(EMovieSceneBlendType InBlendType) override;
 
 public:
 
@@ -254,14 +254,14 @@ public:
 
 protected:
 
-	virtual TSharedPtr<FStructOnScope> GetKeyStruct(TArrayView<const FKeyHandle> KeyHandles) override;
-	virtual EMovieSceneChannelProxyType CacheChannelProxy() override;
+	MOVIESCENETRACKS_API virtual TSharedPtr<FStructOnScope> GetKeyStruct(TArrayView<const FKeyHandle> KeyHandles) override;
+	MOVIESCENETRACKS_API virtual EMovieSceneChannelProxyType CacheChannelProxy() override;
 
 private:
 	
-	virtual bool PopulateEvaluationFieldImpl(const TRange<FFrameNumber>& EffectiveRange, const FMovieSceneEvaluationFieldEntityMetaData& InMetaData, FMovieSceneEntityComponentFieldBuilder* OutFieldBuilder) override;
-	virtual void ImportEntityImpl(UMovieSceneEntitySystemLinker* EntityLinker, const FEntityImportParams& Params, FImportedEntity* OutImportedEntity) override;
-	virtual void InterrogateEntityImpl(UMovieSceneEntitySystemLinker* EntityLinker, const FEntityImportParams& Params, FImportedEntity* OutImportedEntity) override;
+	MOVIESCENETRACKS_API virtual bool PopulateEvaluationFieldImpl(const TRange<FFrameNumber>& EffectiveRange, const FMovieSceneEvaluationFieldEntityMetaData& InMetaData, FMovieSceneEntityComponentFieldBuilder* OutFieldBuilder) override;
+	MOVIESCENETRACKS_API virtual void ImportEntityImpl(UMovieSceneEntitySystemLinker* EntityLinker, const FEntityImportParams& Params, FImportedEntity* OutImportedEntity) override;
+	MOVIESCENETRACKS_API virtual void InterrogateEntityImpl(UMovieSceneEntitySystemLinker* EntityLinker, const FEntityImportParams& Params, FImportedEntity* OutImportedEntity) override;
 
 	template<typename BaseBuilderType>
 	void BuildEntity(BaseBuilderType& InBaseBuilder, UMovieSceneEntitySystemLinker* Linker, const FEntityImportParams& Params, FImportedEntity* OutImportedEntity);
@@ -287,11 +287,11 @@ private:
 	UPROPERTY()
 	FMovieSceneFloatChannel ManualWeight;
 
-	/** Optional pointer to a “channels override” container object. This object would only be allocated if any channels are overridden with a non-standard channel 	*/
+	/** Optional pointer to a "channels override" container object. This object would only be allocated if any channels are overridden with a non-standard channel 	*/
 	UPROPERTY()
 	TObjectPtr<UMovieSceneSectionChannelOverrideRegistry> OverrideRegistry;
 
-	/** Whether to use a quaternion linear interpolation between keys. This finds the 'shortest' distance between keys */
+	/** Whether to use a quaternion linear interpolation between keys. This finds the 'shortest' rotation between keyed orientations. */
 	UPROPERTY(EditAnywhere, DisplayName = "Use Quaternion Interpolation", Category = "Rotation")
 	bool bUseQuaternionInterpolation;
 
