Sequencer: Support for FVector2f and FVector2D (which is now a double)

diff --git a/Engine/Source/Editor/MovieSceneTools/Private/TrackEditors/PropertyTrackEditors/VectorPropertyTrackEditor.cpp b/Engine/Source/Editor/MovieSceneTools/Private/TrackEditors/PropertyTrackEditors/VectorPropertyTrackEditor.cpp
--- a/Engine/Source/Editor/MovieSceneTools/Private/TrackEditors/PropertyTrackEditors/VectorPropertyTrackEditor.cpp
+++ b/Engine/Source/Editor/MovieSceneTools/Private/TrackEditors/PropertyTrackEditors/VectorPropertyTrackEditor.cpp
@@ -30,7 +30,7 @@ void FFloatVectorPropertyTrackEditor::GenerateKeysFromPropertyChanged( const FPr
 	}
 	FName StructName = StructProp->Struct->GetFName();
 
-	bool bIsVector2D = StructName == NAME_Vector2D;
+	bool bIsVector2f = StructName == NAME_Vector2f;
 	bool bIsVector = (StructName == NAME_Vector3f
 #if UE_LARGE_WORLD_COORDINATES_DISABLED
 			|| StructName == NAME_Vector
@@ -46,11 +46,11 @@ void FFloatVectorPropertyTrackEditor::GenerateKeysFromPropertyChanged( const FPr
 	FVector4f VectorValues;
 	int32 Channels;
 
-	if ( bIsVector2D )
+	if ( bIsVector2f )
 	{
-		FVector2D Vector2DValue = PropertyChangedParams.GetPropertyValue<FVector2D>();
-		VectorValues.X = Vector2DValue.X;
-		VectorValues.Y = Vector2DValue.Y;
+		FVector2f Vector2fValue = PropertyChangedParams.GetPropertyValue<FVector2f>();
+		VectorValues.X = Vector2fValue.X;
+		VectorValues.Y = Vector2fValue.Y;
 		Channels = 2;
 	}
 	else if ( bIsVector )
@@ -93,7 +93,7 @@ void FFloatVectorPropertyTrackEditor::InitializeNewTrack( UMovieSceneFloatVector
 	const FStructProperty* StructProp = CastField<const FStructProperty>( PropertyChangedParams.PropertyPath.GetLeafMostProperty().Property.Get() );
 	FName StructName = StructProp->Struct->GetFName();
 
-	if ( StructName == NAME_Vector2D )
+	if ( StructName == NAME_Vector2f )
 	{
 		NewTrack->SetNumChannelsUsed( 2 );
 	}
@@ -145,7 +145,7 @@ bool FFloatVectorPropertyTrackEditor::ModifyGeneratedKeysByCurrentAndWeight(UObj
 		{
 		case 2:
 			{
-				FVector2D Val(InterrogatedValues[0].X, InterrogatedValues[0].Y);
+				FVector2f Val(InterrogatedValues[0].X, InterrogatedValues[0].Y);
 				FMovieSceneChannelProxy& Proxy = SectionToKey->GetChannelProxy();
 				GeneratedTotalKeys[0]->ModifyByCurrentAndWeight(Proxy, KeyTime, (void *)&Val.X, Weight);
 				GeneratedTotalKeys[1]->ModifyByCurrentAndWeight(Proxy, KeyTime, (void *)&Val.Y, Weight);
@@ -200,6 +200,7 @@ void FDoubleVectorPropertyTrackEditor::GenerateKeysFromPropertyChanged( const FP
 	}
 	FName StructName = StructProp->Struct->GetFName();
 
+	bool bIsVector2D = StructName == NAME_Vector2D;
 	bool bIsVector = (StructName == NAME_Vector3d
 #if !UE_LARGE_WORLD_COORDINATES_DISABLED
 			|| StructName == NAME_Vector
@@ -212,12 +213,19 @@ void FDoubleVectorPropertyTrackEditor::GenerateKeysFromPropertyChanged( const FP
 #endif
 		);
 
-	ensure(bIsVector || bIsVector4);
+	ensure(bIsVector2D || bIsVector || bIsVector4);
 
 	FVector4d VectorValues;
 	int32 Channels;
 	
-	if (bIsVector)
+	if (bIsVector2D)
+	{
+		FVector2D Vector2DValue = PropertyChangedParams.GetPropertyValue<FVector2D>();
+		VectorValues.X = Vector2DValue.X;
+		VectorValues.Y = Vector2DValue.Y;
+		Channels = 2;
+	}
+	else if (bIsVector)
 	{
 		FVector3d Vector3DValue = PropertyChangedParams.GetPropertyValue<FVector3d>();
 		VectorValues.X = Vector3DValue.X;
@@ -236,11 +244,14 @@ void FDoubleVectorPropertyTrackEditor::GenerateKeysFromPropertyChanged( const FP
 
 	const bool bKeyX = ChannelName == NAME_None || ChannelName == XName;
 	const bool bKeyY = ChannelName == NAME_None || ChannelName == YName;
-	const bool bKeyZ = ChannelName == NAME_None || ChannelName == ZName;
 
 	OutGeneratedKeys.Add(FMovieSceneChannelValueSetter::Create<FMovieSceneDoubleChannel>(0, VectorValues.X, bKeyX));
 	OutGeneratedKeys.Add(FMovieSceneChannelValueSetter::Create<FMovieSceneDoubleChannel>(1, VectorValues.Y, bKeyY));
-	OutGeneratedKeys.Add(FMovieSceneChannelValueSetter::Create<FMovieSceneDoubleChannel>(2, VectorValues.Z, bKeyZ));
+
+	if ( Channels >= 3 )
+	{
+		OutGeneratedKeys.Add(FMovieSceneChannelValueSetter::Create<FMovieSceneDoubleChannel>(2, VectorValues.Z, ChannelName == NAME_None || ChannelName == ZName));
+	}
 
 	if ( Channels >= 4)
 	{
@@ -254,6 +265,10 @@ void FDoubleVectorPropertyTrackEditor::InitializeNewTrack( UMovieSceneDoubleVect
 	const FStructProperty* StructProp = CastField<const FStructProperty>( PropertyChangedParams.PropertyPath.GetLeafMostProperty().Property.Get() );
 	FName StructName = StructProp->Struct->GetFName();
 
+	if (StructName == NAME_Vector2D)
+	{
+		NewTrack->SetNumChannelsUsed( 2 );
+	}
 	if (StructName == NAME_Vector3d
 #if !UE_LARGE_WORLD_COORDINATES_DISABLED
 			|| StructName == NAME_Vector
@@ -295,6 +310,14 @@ bool FDoubleVectorPropertyTrackEditor::ModifyGeneratedKeysByCurrentAndWeight(UOb
 
 		switch (VectorTrack->GetNumChannelsUsed())
 		{
+		case 2:
+			{
+				FVector2D Val(InterrogatedValues[0].X, InterrogatedValues[0].Y);
+				FMovieSceneChannelProxy& Proxy = SectionToKey->GetChannelProxy();
+				GeneratedTotalKeys[0]->ModifyByCurrentAndWeight(Proxy, KeyTime, (void *)&Val.X, Weight);
+				GeneratedTotalKeys[1]->ModifyByCurrentAndWeight(Proxy, KeyTime, (void *)&Val.Y, Weight);
+			}
+			break;
 		case 3:
 			{
 				FVector3d Val(InterrogatedValues[0].X, InterrogatedValues[0].Y, InterrogatedValues[0].Z);
diff --git a/Engine/Source/Editor/MovieSceneTools/Private/TrackEditors/PropertyTrackEditors/VectorPropertyTrackEditor.h b/Engine/Source/Editor/MovieSceneTools/Private/TrackEditors/PropertyTrackEditors/VectorPropertyTrackEditor.h
--- a/Engine/Source/Editor/MovieSceneTools/Private/TrackEditors/PropertyTrackEditors/VectorPropertyTrackEditor.h
+++ b/Engine/Source/Editor/MovieSceneTools/Private/TrackEditors/PropertyTrackEditors/VectorPropertyTrackEditor.h
@@ -40,9 +40,9 @@ public:
 			FAnimatedPropertyKey::FromStructType(NAME_Vector),
 			FAnimatedPropertyKey::FromStructType(NAME_Vector4),
 #endif
+			FAnimatedPropertyKey::FromStructType(NAME_Vector2f),
 			FAnimatedPropertyKey::FromStructType(NAME_Vector3f),
-			FAnimatedPropertyKey::FromStructType(NAME_Vector4f),
-			FAnimatedPropertyKey::FromStructType(NAME_Vector2D)
+			FAnimatedPropertyKey::FromStructType(NAME_Vector4f)
 		});
 	}
 
@@ -97,6 +97,7 @@ public:
 			FAnimatedPropertyKey::FromStructType(NAME_Vector),
 			FAnimatedPropertyKey::FromStructType(NAME_Vector4),
 #endif
+			FAnimatedPropertyKey::FromStructType(NAME_Vector2D),
 			FAnimatedPropertyKey::FromStructType(NAME_Vector3d),
 			FAnimatedPropertyKey::FromStructType(NAME_Vector4d)
 		});
diff --git a/Engine/Source/Runtime/MovieSceneTracks/Private/MovieSceneTracksComponentTypes.cpp b/Engine/Source/Runtime/MovieSceneTracks/Private/MovieSceneTracksComponentTypes.cpp
--- a/Engine/Source/Runtime/MovieSceneTracks/Private/MovieSceneTracksComponentTypes.cpp
+++ b/Engine/Source/Runtime/MovieSceneTracks/Private/MovieSceneTracksComponentTypes.cpp
@@ -94,9 +94,9 @@ void ConvertOperationalProperty(const FSlateColor& InColor, FIntermediateColor&
 /* ---------------------------------------------------------------------------
  * Vector conversion functions
  * ---------------------------------------------------------------------------*/
-void ConvertOperationalProperty(const FFloatIntermediateVector& InVector, FVector2D& Out)
+void ConvertOperationalProperty(const FFloatIntermediateVector& InVector, FVector2f& Out)
 {
-	Out = FVector2D(InVector.X, InVector.Y);
+	Out = FVector2f(InVector.X, InVector.Y);
 }
 
 void ConvertOperationalProperty(const FFloatIntermediateVector& InVector, FVector3f& Out)
@@ -109,7 +109,7 @@ void ConvertOperationalProperty(const FFloatIntermediateVector& InVector, FVecto
 	Out = FVector4(InVector.X, InVector.Y, InVector.Z, InVector.W);
 }
 
-void ConvertOperationalProperty(const FVector2D& In, FFloatIntermediateVector& Out)
+void ConvertOperationalProperty(const FVector2f& In, FFloatIntermediateVector& Out)
 {
 	Out = FFloatIntermediateVector(In.X, In.Y);
 }
@@ -124,6 +124,11 @@ void ConvertOperationalProperty(const FVector4& In, FFloatIntermediateVector& Ou
 	Out = FFloatIntermediateVector(In.X, In.Y, In.Z, In.W);
 }
 
+void ConvertOperationalProperty(const FDoubleIntermediateVector& InVector, FVector2D& Out)
+{
+	Out = FVector2D(InVector.X, InVector.Y);
+}
+
 void ConvertOperationalProperty(const FDoubleIntermediateVector& InVector, FVector3d& Out)
 {
 	Out = FVector3d(InVector.X, InVector.Y, InVector.Z);
@@ -134,6 +139,11 @@ void ConvertOperationalProperty(const FDoubleIntermediateVector& InVector, FVect
 	Out = FVector4d(InVector.X, InVector.Y, InVector.Z, InVector.W);
 }
 
+void ConvertOperationalProperty(const FVector2D& In, FDoubleIntermediateVector& Out)
+{
+	Out = FDoubleIntermediateVector(In.X, In.Y);
+}
+
 void ConvertOperationalProperty(const FVector3d& In, FDoubleIntermediateVector& Out)
 {
 	Out = FDoubleIntermediateVector(In.X, In.Y, In.Z);
@@ -411,7 +421,7 @@ struct FFloatVectorHandler : TPropertyComponentHandler<FFloatVectorPropertyTrait
 			FStructProperty* BoundProperty = CastField<FStructProperty>(FTrackInstancePropertyBindings::FindProperty(Object, Binding.PropertyPath.ToString()));
 			if (ensure(BoundProperty && BoundProperty->Struct))
 			{
-				if (BoundProperty->Struct == TBaseStructure<FVector2D>::Get())
+				if (BoundProperty->Struct->GetFName() == NAME_Vector2f)
 				{
 					OutMetaData.NumChannels = 2;
 				}
@@ -458,7 +468,11 @@ struct FDoubleVectorHandler : TPropertyComponentHandler<FDoubleVectorPropertyTra
 			FStructProperty* BoundProperty = CastField<FStructProperty>(FTrackInstancePropertyBindings::FindProperty(Object, Binding.PropertyPath.ToString()));
 			if (ensure(BoundProperty && BoundProperty->Struct))
 			{
-				if (BoundProperty->Struct->GetFName() == NAME_Vector3d || BoundProperty->Struct->GetFName() == NAME_Vector)
+				if (BoundProperty->Struct->GetFName() == NAME_Vector2D)
+				{
+					OutMetaData.NumChannels = 2;
+				}
+				else if (BoundProperty->Struct->GetFName() == NAME_Vector3d || BoundProperty->Struct->GetFName() == NAME_Vector)
 				{
 					OutMetaData.NumChannels = 3;
 				}
diff --git a/Engine/Source/Runtime/MovieSceneTracks/Private/Sections/MovieSceneVectorSection.cpp b/Engine/Source/Runtime/MovieSceneTracks/Private/Sections/MovieSceneVectorSection.cpp
--- a/Engine/Source/Runtime/MovieSceneTracks/Private/Sections/MovieSceneVectorSection.cpp
+++ b/Engine/Source/Runtime/MovieSceneTracks/Private/Sections/MovieSceneVectorSection.cpp
@@ -34,7 +34,7 @@ struct TVectorSectionEditorDataTraits<float>
 	{
 		if (NumChannels == 2)
 		{
-			FVector2D Vector = Bindings.GetCurrentValue<FVector2D>(InObject);
+			FVector2f Vector = Bindings.GetCurrentValue<FVector2f>(InObject);
 			return FCompositeValue(Vector.X, Vector.Y, 0.f, 0.f);
 		}
 		else if (NumChannels == 3)
@@ -58,7 +58,12 @@ struct TVectorSectionEditorDataTraits<double>
 
 	static FCompositeValue GetPropertyValue(UObject& InObject, FTrackInstancePropertyBindings& Bindings, int32 NumChannels)
 	{
-		if (NumChannels == 3)
+		if (NumChannels == 2)
+		{
+			FVector2D Vector = Bindings.GetCurrentValue<FVector2D>(InObject);
+			return FCompositeValue(Vector.X, Vector.Y, 0.f, 0.f);
+		}
+		else if (NumChannels == 3)
 		{
 			FVector3d Vector = Bindings.GetCurrentValue<FVector3d>(InObject);
 			return FCompositeValue(Vector.X, Vector.Y, Vector.Z, 0.f);
@@ -201,7 +206,7 @@ TSharedPtr<FStructOnScope> UMovieSceneFloatVectorSection::GetKeyStruct(TArrayVie
 	TSharedPtr<FStructOnScope> KeyStruct;
 	if (ChannelsUsed == 2)
 	{
-		KeyStruct = MakeShareable(new FStructOnScope(FMovieSceneVector2DKeyStruct::StaticStruct()));
+		KeyStruct = MakeShareable(new FStructOnScope(FMovieSceneVector2fKeyStruct::StaticStruct()));
 	}
 	else if (ChannelsUsed == 3)
 	{
@@ -314,7 +319,11 @@ void UMovieSceneDoubleVectorSection::RecreateChannelProxy()
 TSharedPtr<FStructOnScope> UMovieSceneDoubleVectorSection::GetKeyStruct(TArrayView<const FKeyHandle> KeyHandles)
 {
 	TSharedPtr<FStructOnScope> KeyStruct;
-	if (ChannelsUsed == 3)
+	if (ChannelsUsed == 2)
+	{
+		KeyStruct = MakeShareable(new FStructOnScope(FMovieSceneVector2DKeyStruct::StaticStruct()));
+	}
+	else if (ChannelsUsed == 3)
 	{
 		KeyStruct = MakeShareable(new FStructOnScope(FMovieSceneVector3dKeyStruct::StaticStruct()));
 	}
diff --git a/Engine/Source/Runtime/MovieSceneTracks/Public/MovieSceneTracksComponentTypes.h b/Engine/Source/Runtime/MovieSceneTracks/Public/MovieSceneTracksComponentTypes.h
--- a/Engine/Source/Runtime/MovieSceneTracks/Public/MovieSceneTracksComponentTypes.h
+++ b/Engine/Source/Runtime/MovieSceneTracks/Public/MovieSceneTracksComponentTypes.h
@@ -201,7 +201,7 @@ struct FFloatVectorPropertyTraits
 
 		switch (MetaData.NumChannels)
 		{
-			case 2: TIndirectPropertyTraits<FVector2D, FFloatIntermediateVector>::GetObjectPropertyValue(InObject, PropertyOffset, OutValue); return;
+			case 2: TIndirectPropertyTraits<FVector2f, FFloatIntermediateVector>::GetObjectPropertyValue(InObject, PropertyOffset, OutValue); return;
 			case 3: TIndirectPropertyTraits<FVector3f, FFloatIntermediateVector>::GetObjectPropertyValue(InObject, PropertyOffset, OutValue); return;
 			case 4: TIndirectPropertyTraits<FVector4,  FFloatIntermediateVector>::GetObjectPropertyValue(InObject, PropertyOffset, OutValue); return;
 		}
@@ -218,7 +218,7 @@ struct FFloatVectorPropertyTraits
 
 		switch (MetaData.NumChannels)
 		{
-			case 2: TIndirectPropertyTraits<FVector2D, FFloatIntermediateVector>::GetObjectPropertyValue(InObject, PropertyBindings, OutValue); return;
+			case 2: TIndirectPropertyTraits<FVector2f, FFloatIntermediateVector>::GetObjectPropertyValue(InObject, PropertyBindings, OutValue); return;
 			case 3: TIndirectPropertyTraits<FVector3f, FFloatIntermediateVector>::GetObjectPropertyValue(InObject, PropertyBindings, OutValue); return;
 			case 4: TIndirectPropertyTraits<FVector4,  FFloatIntermediateVector>::GetObjectPropertyValue(InObject, PropertyBindings, OutValue); return;
 		}
@@ -235,7 +235,7 @@ struct FFloatVectorPropertyTraits
 
 		switch (MetaData.NumChannels)
 		{
-			case 2: TIndirectPropertyTraits<FVector2D, FFloatIntermediateVector>::GetObjectPropertyValue(InObject, PropertyPath, OutValue); return;
+			case 2: TIndirectPropertyTraits<FVector2f, FFloatIntermediateVector>::GetObjectPropertyValue(InObject, PropertyPath, OutValue); return;
 			case 3: TIndirectPropertyTraits<FVector3f, FFloatIntermediateVector>::GetObjectPropertyValue(InObject, PropertyPath, OutValue); return;
 			case 4: TIndirectPropertyTraits<FVector4,  FFloatIntermediateVector>::GetObjectPropertyValue(InObject, PropertyPath, OutValue); return;
 		}
@@ -257,7 +257,7 @@ struct FFloatVectorPropertyTraits
 
 		switch (MetaData.NumChannels)
 		{
-			case 2: TIndirectPropertyTraits<FVector2D, FFloatIntermediateVector>::SetObjectPropertyValue(InObject, PropertyOffset, InValue); return;
+			case 2: TIndirectPropertyTraits<FVector2f, FFloatIntermediateVector>::SetObjectPropertyValue(InObject, PropertyOffset, InValue); return;
 			case 3: TIndirectPropertyTraits<FVector3f, FFloatIntermediateVector>::SetObjectPropertyValue(InObject, PropertyOffset, InValue); return;
 			case 4: TIndirectPropertyTraits<FVector4,  FFloatIntermediateVector>::SetObjectPropertyValue(InObject, PropertyOffset, InValue); return;
 		}
@@ -275,7 +275,7 @@ struct FFloatVectorPropertyTraits
 
 		switch (MetaData.NumChannels)
 		{
-			case 2: TIndirectPropertyTraits<FVector2D, FFloatIntermediateVector>::SetObjectPropertyValue(InObject, PropertyBindings, InValue); return;
+			case 2: TIndirectPropertyTraits<FVector2f, FFloatIntermediateVector>::SetObjectPropertyValue(InObject, PropertyBindings, InValue); return;
 			case 3: TIndirectPropertyTraits<FVector3f, FFloatIntermediateVector>::SetObjectPropertyValue(InObject, PropertyBindings, InValue); return;
 			case 4: TIndirectPropertyTraits<FVector4,  FFloatIntermediateVector>::SetObjectPropertyValue(InObject, PropertyBindings, InValue); return;
 		}
@@ -303,6 +303,7 @@ struct FDoubleVectorPropertyTraits
 	{
 		switch (MetaData.NumChannels)
 		{
+			case 2: TIndirectPropertyTraits<FVector2D, FDoubleIntermediateVector>::GetObjectPropertyValue(InObject, PropertyOffset, OutValue); return;
 			case 3: TIndirectPropertyTraits<FVector3d, FDoubleIntermediateVector>::GetObjectPropertyValue(InObject, PropertyOffset, OutValue); return;
 			case 4: TIndirectPropertyTraits<FVector4d, FDoubleIntermediateVector>::GetObjectPropertyValue(InObject, PropertyOffset, OutValue); return;
 		}
@@ -311,6 +312,7 @@ struct FDoubleVectorPropertyTraits
 	{
 		switch (MetaData.NumChannels)
 		{
+			case 2: TIndirectPropertyTraits<FVector2D, FDoubleIntermediateVector>::GetObjectPropertyValue(InObject, PropertyBindings, OutValue); return;
 			case 3: TIndirectPropertyTraits<FVector3d, FDoubleIntermediateVector>::GetObjectPropertyValue(InObject, PropertyBindings, OutValue); return;
 			case 4: TIndirectPropertyTraits<FVector4d, FDoubleIntermediateVector>::GetObjectPropertyValue(InObject, PropertyBindings, OutValue); return;
 		}
@@ -319,6 +321,7 @@ struct FDoubleVectorPropertyTraits
 	{
 		switch (MetaData.NumChannels)
 		{
+			case 2: TIndirectPropertyTraits<FVector2D, FDoubleIntermediateVector>::GetObjectPropertyValue(InObject, PropertyPath, OutValue); return;
 			case 3: TIndirectPropertyTraits<FVector3d, FDoubleIntermediateVector>::GetObjectPropertyValue(InObject, PropertyPath, OutValue); return;
 			case 4: TIndirectPropertyTraits<FVector4d, FDoubleIntermediateVector>::GetObjectPropertyValue(InObject, PropertyPath, OutValue); return;
 		}
@@ -333,6 +336,7 @@ struct FDoubleVectorPropertyTraits
 	{
 		switch (MetaData.NumChannels)
 		{
+			case 2: TIndirectPropertyTraits<FVector2D, FDoubleIntermediateVector>::SetObjectPropertyValue(InObject, PropertyOffset, InValue); return;
 			case 3: TIndirectPropertyTraits<FVector3d, FDoubleIntermediateVector>::SetObjectPropertyValue(InObject, PropertyOffset, InValue); return;
 			case 4: TIndirectPropertyTraits<FVector4d, FDoubleIntermediateVector>::SetObjectPropertyValue(InObject, PropertyOffset, InValue); return;
 		}
@@ -343,6 +347,7 @@ struct FDoubleVectorPropertyTraits
 	{
 		switch (MetaData.NumChannels)
 		{
+			case 2: TIndirectPropertyTraits<FVector2D, FDoubleIntermediateVector>::SetObjectPropertyValue(InObject, PropertyBindings, InValue); return;
 			case 3: TIndirectPropertyTraits<FVector3d, FDoubleIntermediateVector>::SetObjectPropertyValue(InObject, PropertyBindings, InValue); return;
 			case 4: TIndirectPropertyTraits<FVector4d, FDoubleIntermediateVector>::SetObjectPropertyValue(InObject, PropertyBindings, InValue); return;
 		}
diff --git a/Engine/Source/Runtime/MovieSceneTracks/Public/MovieSceneTracksPropertyTypes.h b/Engine/Source/Runtime/MovieSceneTracks/Public/MovieSceneTracksPropertyTypes.h
--- a/Engine/Source/Runtime/MovieSceneTracks/Public/MovieSceneTracksPropertyTypes.h
+++ b/Engine/Source/Runtime/MovieSceneTracks/Public/MovieSceneTracksPropertyTypes.h
@@ -202,15 +202,17 @@ MOVIESCENETRACKS_API void ConvertOperationalProperty(const FColor& InColor, FInt
 MOVIESCENETRACKS_API void ConvertOperationalProperty(const FLinearColor& InColor, FIntermediateColor& OutIntermediate);
 MOVIESCENETRACKS_API void ConvertOperationalProperty(const FSlateColor& InColor, FIntermediateColor& OutIntermediate);
 
-MOVIESCENETRACKS_API void ConvertOperationalProperty(const FFloatIntermediateVector& InVector, FVector2D& Out);
+MOVIESCENETRACKS_API void ConvertOperationalProperty(const FFloatIntermediateVector& InVector, FVector2f& Out);
 MOVIESCENETRACKS_API void ConvertOperationalProperty(const FFloatIntermediateVector& InVector, FVector3f& Out);
 MOVIESCENETRACKS_API void ConvertOperationalProperty(const FFloatIntermediateVector& InVector, FVector4& Out);
-MOVIESCENETRACKS_API void ConvertOperationalProperty(const FVector2D& In, FFloatIntermediateVector& Out);
+MOVIESCENETRACKS_API void ConvertOperationalProperty(const FVector2f& In, FFloatIntermediateVector& Out);
 MOVIESCENETRACKS_API void ConvertOperationalProperty(const FVector3f& In, FFloatIntermediateVector& Out);
 MOVIESCENETRACKS_API void ConvertOperationalProperty(const FVector4& In, FFloatIntermediateVector& Out);
 
+MOVIESCENETRACKS_API void ConvertOperationalProperty(const FDoubleIntermediateVector& InVector, FVector2D& Out);
 MOVIESCENETRACKS_API void ConvertOperationalProperty(const FDoubleIntermediateVector& InVector, FVector3d& Out);
 MOVIESCENETRACKS_API void ConvertOperationalProperty(const FDoubleIntermediateVector& InVector, FVector4d& Out);
+MOVIESCENETRACKS_API void ConvertOperationalProperty(const FVector2D& In, FDoubleIntermediateVector& Out);
 MOVIESCENETRACKS_API void ConvertOperationalProperty(const FVector3d& In, FDoubleIntermediateVector& Out);
 MOVIESCENETRACKS_API void ConvertOperationalProperty(const FVector4d& In, FDoubleIntermediateVector& Out);
 
diff --git a/Engine/Source/Runtime/MovieSceneTracks/Public/Sections/MovieSceneVectorSection.h b/Engine/Source/Runtime/MovieSceneTracks/Public/Sections/MovieSceneVectorSection.h
--- a/Engine/Source/Runtime/MovieSceneTracks/Public/Sections/MovieSceneVectorSection.h
+++ b/Engine/Source/Runtime/MovieSceneTracks/Public/Sections/MovieSceneVectorSection.h
@@ -40,10 +40,10 @@ template<> struct TStructOpsTypeTraits<FMovieSceneFloatVectorKeyStructBase> : pu
 
 
 /**
- * Proxy structure for 2D vector section key data.
+ * Proxy structure for 2f vector section key data.
  */
 USTRUCT()
-struct FMovieSceneVector2DKeyStruct
+struct FMovieSceneVector2fKeyStruct
 	: public FMovieSceneFloatVectorKeyStructBase
 {
 	GENERATED_BODY()
@@ -55,7 +55,7 @@ struct FMovieSceneVector2DKeyStruct
 	//~ FMovieSceneFloatVectorKeyStructBase interface
 	virtual float* GetPropertyChannelByIndex(int32 Index) override { return &Vector[Index]; }
 };
-template<> struct TStructOpsTypeTraits<FMovieSceneVector2DKeyStruct> : public TStructOpsTypeTraitsBase2<FMovieSceneVector2DKeyStruct> { enum { WithCopy = false }; };
+template<> struct TStructOpsTypeTraits<FMovieSceneVector2fKeyStruct> : public TStructOpsTypeTraitsBase2<FMovieSceneVector2fKeyStruct> { enum { WithCopy = false }; };
 
 /**
 * Proxy structure for float vector section key data.
@@ -116,6 +116,24 @@ struct FMovieSceneDoubleVectorKeyStructBase
 };
 template<> struct TStructOpsTypeTraits<FMovieSceneDoubleVectorKeyStructBase> : public TStructOpsTypeTraitsBase2<FMovieSceneDoubleVectorKeyStructBase> { enum { WithCopy = false, WithPureVirtual = true, }; };
 
+/**
+ * Proxy structure for 2D vector section key data.
+ */
+USTRUCT()
+struct FMovieSceneVector2DKeyStruct
+	: public FMovieSceneDoubleVectorKeyStructBase
+{
+	GENERATED_BODY()
+
+	/** They key's vector value. */
+	UPROPERTY(EditAnywhere, Category=Key)
+	FVector2D Vector = FVector2D::ZeroVector;
+
+	//~ FMovieSceneDoubleVectorKeyStructBase interface
+	virtual double* GetPropertyChannelByIndex(int32 Index) override { return &Vector[Index]; }
+};
+template<> struct TStructOpsTypeTraits<FMovieSceneVector2DKeyStruct> : public TStructOpsTypeTraitsBase2<FMovieSceneVector2DKeyStruct> { enum { WithCopy = false }; };
+
 /**
 * Proxy structure for double vector section key data.
 */
