Sequencer: Fix non-unity

diff --git a/Engine/Source/Runtime/MovieScene/Private/Evaluation/MovieScenePreAnimatedState.cpp b/Engine/Source/Runtime/MovieScene/Private/Evaluation/MovieScenePreAnimatedState.cpp
--- a/Engine/Source/Runtime/MovieScene/Private/Evaluation/MovieScenePreAnimatedState.cpp
+++ b/Engine/Source/Runtime/MovieScene/Private/Evaluation/MovieScenePreAnimatedState.cpp
@@ -3,6 +3,7 @@
 #include "Evaluation/MovieScenePreAnimatedState.h"
 #include "GameFramework/Actor.h"
 #include "Components/ActorComponent.h"
+#include "EntitySystem/MovieSceneEntitySystemLinker.h"
 #include "Evaluation/PreAnimatedState/MovieScenePreAnimatedStateExtension.h"
 #include "Evaluation/PreAnimatedState/MovieSceneRestoreStateParams.h"
 #include "Evaluation/PreAnimatedState/MovieScenePreAnimatedCaptureSources.h"
diff --git a/Engine/Source/Runtime/MovieSceneTracks/Private/TrackInstances/MovieSceneCVarTrackInstance.cpp b/Engine/Source/Runtime/MovieSceneTracks/Private/TrackInstances/MovieSceneCVarTrackInstance.cpp
--- a/Engine/Source/Runtime/MovieSceneTracks/Private/TrackInstances/MovieSceneCVarTrackInstance.cpp
+++ b/Engine/Source/Runtime/MovieSceneTracks/Private/TrackInstances/MovieSceneCVarTrackInstance.cpp
@@ -6,7 +6,11 @@
 #include "EntitySystem/MovieSceneEntitySystemLinker.h"
 #include "EntitySystem/MovieSceneInstanceRegistry.h"
 #include "EntitySystem/TrackInstance/MovieSceneTrackInstanceSystem.h"
+#include "Evaluation/PreAnimatedState/MovieScenePreAnimatedCaptureSource.h"
+#include "Evaluation/PreAnimatedState/MovieScenePreAnimatedStateExtension.h"
+#include "Evaluation/PreAnimatedState/MovieScenePreAnimatedStateStorage.h"
 #include "Evaluation/PreAnimatedState/MovieScenePreAnimatedStateTypes.h"
+#include "Evaluation/PreAnimatedState/MovieScenePreAnimatedStorageID.inl"
 #include "MovieSceneCommonHelpers.h"
 
 namespace UE
