Sequencer: CVar track supports multiple rows

diff --git a/Engine/Source/Runtime/MovieSceneTracks/Private/Tracks/MovieSceneCVarTrack.cpp b/Engine/Source/Runtime/MovieSceneTracks/Private/Tracks/MovieSceneCVarTrack.cpp
--- a/Engine/Source/Runtime/MovieSceneTracks/Private/Tracks/MovieSceneCVarTrack.cpp
+++ b/Engine/Source/Runtime/MovieSceneTracks/Private/Tracks/MovieSceneCVarTrack.cpp
@@ -46,7 +46,7 @@ UMovieSceneSection* UMovieSceneCVarTrack::CreateNewSection()
 
 bool UMovieSceneCVarTrack::SupportsMultipleRows() const
 {
-	return false;
+	return true;
 }
 
 EMovieSceneTrackEasingSupportFlags UMovieSceneCVarTrack::SupportsEasing(FMovieSceneSupportsEasingParams& Params) const
