Sequencer: Support for FVector2f and FVector2D (which is now a double)

diff --git a/Engine/Source/Runtime/MovieSceneTracks/Public/Sections/MovieSceneVectorSection.h b/Engine/Source/Runtime/MovieSceneTracks/Public/Sections/MovieSceneVectorSection.h
--- a/Engine/Source/Runtime/MovieSceneTracks/Public/Sections/MovieSceneVectorSection.h
+++ b/Engine/Source/Runtime/MovieSceneTracks/Public/Sections/MovieSceneVectorSection.h
@@ -130,7 +130,7 @@ struct FMovieSceneVector2DKeyStruct
 	FVector2D Vector = FVector2D::ZeroVector;
 
 	//~ FMovieSceneDoubleVectorKeyStructBase interface
-	virtual double* GetPropertyChannelByIndex(int32 Index) override { return &Vector[Index]; }
+	virtual double* GetPropertyChannelByIndex(int32 Index) override { return (double*)&Vector[Index]; }
 };
 template<> struct TStructOpsTypeTraits<FMovieSceneVector2DKeyStruct> : public TStructOpsTypeTraitsBase2<FMovieSceneVector2DKeyStruct> { enum { WithCopy = false }; };
 
