Expose MovieScene::SetComponentTransformAndVelocity so external code can make use of it

diff --git a/Engine/Source/Runtime/MovieSceneTracks/Public/MovieSceneTracksPropertyTypes.h b/Engine/Source/Runtime/MovieSceneTracks/Public/MovieSceneTracksPropertyTypes.h
--- a/Engine/Source/Runtime/MovieSceneTracks/Public/MovieSceneTracksPropertyTypes.h
+++ b/Engine/Source/Runtime/MovieSceneTracks/Public/MovieSceneTracksPropertyTypes.h
@@ -189,6 +189,7 @@ struct FIntermediate3DTransform
 
 MOVIESCENETRACKS_API FIntermediate3DTransform GetComponentTransform(const UObject* Object);
 MOVIESCENETRACKS_API void SetComponentTransform(USceneComponent* SceneComponent, const FIntermediate3DTransform& InTransform);
+MOVIESCENETRACKS_API void SetComponentTransformAndVelocity(UObject* Object, const FIntermediate3DTransform& InTransform);
 
 MOVIESCENETRACKS_API void ConvertOperationalProperty(const FIntermediate3DTransform& In, FEulerTransform& Out);
 MOVIESCENETRACKS_API void ConvertOperationalProperty(const FEulerTransform& In, FIntermediate3DTransform& Out);
