Sequencer: Fix conversion of float-to-double data for pre-LWC sequences.

diff --git a/Engine/Source/Runtime/MovieSceneTracks/Private/MovieSceneTracksComponentTypes.cpp b/Engine/Source/Runtime/MovieSceneTracks/Private/MovieSceneTracksComponentTypes.cpp
--- a/Engine/Source/Runtime/MovieSceneTracks/Private/MovieSceneTracksComponentTypes.cpp
+++ b/Engine/Source/Runtime/MovieSceneTracks/Private/MovieSceneTracksComponentTypes.cpp
@@ -104,9 +104,9 @@ void ConvertOperationalProperty(const FFloatIntermediateVector& InVector, FVecto
 	Out = FVector3f(InVector.X, InVector.Y, InVector.Z);
 }
 
-void ConvertOperationalProperty(const FFloatIntermediateVector& InVector, FVector4& Out)
+void ConvertOperationalProperty(const FFloatIntermediateVector& InVector, FVector4f& Out)
 {
-	Out = FVector4(InVector.X, InVector.Y, InVector.Z, InVector.W);
+	Out = FVector4f(InVector.X, InVector.Y, InVector.Z, InVector.W);
 }
 
 void ConvertOperationalProperty(const FVector2f& In, FFloatIntermediateVector& Out)
@@ -119,14 +119,14 @@ void ConvertOperationalProperty(const FVector3f& In, FFloatIntermediateVector& O
 	Out = FFloatIntermediateVector(In.X, In.Y, In.Z);
 }
 
-void ConvertOperationalProperty(const FVector4& In, FFloatIntermediateVector& Out)
+void ConvertOperationalProperty(const FVector4f& In, FFloatIntermediateVector& Out)
 {
 	Out = FFloatIntermediateVector(In.X, In.Y, In.Z, In.W);
 }
 
-void ConvertOperationalProperty(const FDoubleIntermediateVector& InVector, FVector2D& Out)
+void ConvertOperationalProperty(const FDoubleIntermediateVector& InVector, FVector2d& Out)
 {
-	Out = FVector2D(InVector.X, InVector.Y);
+	Out = FVector2d(InVector.X, InVector.Y);
 }
 
 void ConvertOperationalProperty(const FDoubleIntermediateVector& InVector, FVector3d& Out)
@@ -139,7 +139,7 @@ void ConvertOperationalProperty(const FDoubleIntermediateVector& InVector, FVect
 	Out = FVector4d(InVector.X, InVector.Y, InVector.Z, InVector.W);
 }
 
-void ConvertOperationalProperty(const FVector2D& In, FDoubleIntermediateVector& Out)
+void ConvertOperationalProperty(const FVector2d& In, FDoubleIntermediateVector& Out)
 {
 	Out = FDoubleIntermediateVector(In.X, In.Y);
 }
@@ -421,23 +421,21 @@ struct FFloatVectorHandler : TPropertyComponentHandler<FFloatVectorPropertyTrait
 			FStructProperty* BoundProperty = CastField<FStructProperty>(FTrackInstancePropertyBindings::FindProperty(Object, Binding.PropertyPath.ToString()));
 			if (ensure(BoundProperty && BoundProperty->Struct))
 			{
-				if (BoundProperty->Struct == TBaseStructure<FVector2D>::Get())
+				if (BoundProperty->Struct == TBaseStructure<FVector2D>::Get() || BoundProperty->Struct == TVariantStructure<FVector2f>::Get() || BoundProperty->Struct == TVariantStructure<FVector2d>::Get())
 				{
 					OutMetaData.NumChannels = 2;
+					OutMetaData.bIsDouble = (BoundProperty->Struct == TBaseStructure<FVector2D>::Get() || BoundProperty->Struct == TVariantStructure<FVector2d>::Get());
 				}
-				else if (BoundProperty->Struct->GetFName() == NAME_Vector3f || BoundProperty->Struct->GetFName() == NAME_Vector)
+				else if (BoundProperty->Struct == TBaseStructure<FVector>::Get() || BoundProperty->Struct == TVariantStructure<FVector3f>::Get() || BoundProperty->Struct == TVariantStructure<FVector3d>::Get())
 				{
 					OutMetaData.NumChannels = 3;
-
-					if (BoundProperty->Struct->GetFName() == NAME_Vector)
-					{
-						OutMetaData.bIsDouble = true;
-					}
+					OutMetaData.bIsDouble = (BoundProperty->Struct == TBaseStructure<FVector>::Get() || BoundProperty->Struct == TVariantStructure<FVector3d>::Get());
 				}
 				else
 				{
-					ensure(BoundProperty->Struct == TBaseStructure<FVector4>::Get() || (BoundProperty->Struct->GetFName() == NAME_Vector4f) || (BoundProperty->Struct->GetFName() == NAME_Vector4));
+					ensure(BoundProperty->Struct == TBaseStructure<FVector4>::Get() || BoundProperty->Struct == TVariantStructure<FVector4f>::Get() || BoundProperty->Struct == TVariantStructure<FVector4d>::Get());
 					OutMetaData.NumChannels = 4;
+					OutMetaData.bIsDouble = (BoundProperty->Struct == TBaseStructure<FVector4>::Get() || BoundProperty->Struct == TVariantStructure<FVector4d>::Get());
 				}
 			}
 			else
diff --git a/Engine/Source/Runtime/MovieSceneTracks/Public/MovieSceneTracksComponentTypes.h b/Engine/Source/Runtime/MovieSceneTracks/Public/MovieSceneTracksComponentTypes.h
--- a/Engine/Source/Runtime/MovieSceneTracks/Public/MovieSceneTracksComponentTypes.h
+++ b/Engine/Source/Runtime/MovieSceneTracks/Public/MovieSceneTracksComponentTypes.h
@@ -179,116 +179,6 @@
 	}
 };
 
-struct FFloatVectorPropertyTraits
-{
-	using StorageType  = FFloatIntermediateVector;
-	using MetaDataType = TPropertyMetaData<FVectorPropertyMetaData>;
-
-	static void GetObjectPropertyValue(const UObject* InObject, FVectorPropertyMetaData MetaData, const FCustomPropertyAccessor& BaseCustomAccessor, FFloatIntermediateVector& OutValue)
-	{
-		const TCustomPropertyAccessor<FFloatVectorPropertyTraits>& CustomAccessor = static_cast<const TCustomPropertyAccessor<FFloatVectorPropertyTraits>&>(BaseCustomAccessor);
-		OutValue = (*CustomAccessor.Functions.Getter)(InObject, MetaData);
-	}
-	static void GetObjectPropertyValue(const UObject* InObject, FVectorPropertyMetaData MetaData, uint16 PropertyOffset, FFloatIntermediateVector& OutValue)
-	{
-		if (MetaData.bIsDouble)
-		{
-			FDoubleIntermediateVector TempDoubleValue;
-			TIndirectPropertyTraits<FVector3d, FDoubleIntermediateVector>::GetObjectPropertyValue(InObject, PropertyOffset, TempDoubleValue); 
-			OutValue = FFloatIntermediateVector(TempDoubleValue.X, TempDoubleValue.Y, TempDoubleValue.Z, TempDoubleValue.W);
-			return;
-		}
-
-		switch (MetaData.NumChannels)
-		{
-			case 2: TIndirectPropertyTraits<FVector2f, FFloatIntermediateVector>::GetObjectPropertyValue(InObject, PropertyOffset, OutValue); return;
-			case 3: TIndirectPropertyTraits<FVector3f, FFloatIntermediateVector>::GetObjectPropertyValue(InObject, PropertyOffset, OutValue); return;
-			case 4: TIndirectPropertyTraits<FVector4,  FFloatIntermediateVector>::GetObjectPropertyValue(InObject, PropertyOffset, OutValue); return;
-		}
-	}
-	static void GetObjectPropertyValue(const UObject* InObject, FVectorPropertyMetaData MetaData, FTrackInstancePropertyBindings* PropertyBindings, FFloatIntermediateVector& OutValue)
-	{
-		if (MetaData.bIsDouble)
-		{
-			FDoubleIntermediateVector TempDoubleValue;
-			TIndirectPropertyTraits<FVector3d, FDoubleIntermediateVector>::GetObjectPropertyValue(InObject, PropertyBindings, TempDoubleValue); 
-			OutValue = FFloatIntermediateVector(TempDoubleValue.X, TempDoubleValue.Y, TempDoubleValue.Z, TempDoubleValue.W);
-			return;
-		}
-
-		switch (MetaData.NumChannels)
-		{
-			case 2: TIndirectPropertyTraits<FVector2f, FFloatIntermediateVector>::GetObjectPropertyValue(InObject, PropertyBindings, OutValue); return;
-			case 3: TIndirectPropertyTraits<FVector3f, FFloatIntermediateVector>::GetObjectPropertyValue(InObject, PropertyBindings, OutValue); return;
-			case 4: TIndirectPropertyTraits<FVector4,  FFloatIntermediateVector>::GetObjectPropertyValue(InObject, PropertyBindings, OutValue); return;
-		}
-	}
-	static void GetObjectPropertyValue(const UObject* InObject, FVectorPropertyMetaData MetaData, const FName& PropertyPath, StorageType& OutValue)
-	{
-		if (MetaData.bIsDouble)
-		{
-			FDoubleIntermediateVector TempDoubleValue;
-			TIndirectPropertyTraits<FVector3d, FDoubleIntermediateVector>::GetObjectPropertyValue(InObject, PropertyPath, TempDoubleValue); 
-			OutValue = FFloatIntermediateVector(TempDoubleValue.X, TempDoubleValue.Y, TempDoubleValue.Z, TempDoubleValue.W);
-			return;
-		}
-
-		switch (MetaData.NumChannels)
-		{
-			case 2: TIndirectPropertyTraits<FVector2f, FFloatIntermediateVector>::GetObjectPropertyValue(InObject, PropertyPath, OutValue); return;
-			case 3: TIndirectPropertyTraits<FVector3f, FFloatIntermediateVector>::GetObjectPropertyValue(InObject, PropertyPath, OutValue); return;
-			case 4: TIndirectPropertyTraits<FVector4,  FFloatIntermediateVector>::GetObjectPropertyValue(InObject, PropertyPath, OutValue); return;
-		}
-	}
-
-	static void SetObjectPropertyValue(UObject* InObject, FVectorPropertyMetaData MetaData, const FCustomPropertyAccessor& BaseCustomAccessor, const FFloatIntermediateVector& InValue)
-	{
-		const TCustomPropertyAccessor<FFloatVectorPropertyTraits>& CustomAccessor = static_cast<const TCustomPropertyAccessor<FFloatVectorPropertyTraits>&>(BaseCustomAccessor);
-		(*CustomAccessor.Functions.Setter)(InObject, MetaData, InValue);
-	}
-	static void SetObjectPropertyValue(UObject* InObject, FVectorPropertyMetaData MetaData, uint16 PropertyOffset, const FFloatIntermediateVector& InValue)
-	{
-		if (MetaData.bIsDouble)
-		{
-			FDoubleIntermediateVector TempDoubleValue(InValue.X, InValue.Y, InValue.Z, InValue.W);
-			TIndirectPropertyTraits<FVector3d, FDoubleIntermediateVector>::SetObjectPropertyValue(InObject, PropertyOffset, TempDoubleValue);
-			return;
-		}
-
-		switch (MetaData.NumChannels)
-		{
-			case 2: TIndirectPropertyTraits<FVector2f, FFloatIntermediateVector>::SetObjectPropertyValue(InObject, PropertyOffset, InValue); return;
-			case 3: TIndirectPropertyTraits<FVector3f, FFloatIntermediateVector>::SetObjectPropertyValue(InObject, PropertyOffset, InValue); return;
-			case 4: TIndirectPropertyTraits<FVector4,  FFloatIntermediateVector>::SetObjectPropertyValue(InObject, PropertyOffset, InValue); return;
-		}
-
-		checkf(false, TEXT("Invalid number of channels"));
-	}
-	static void SetObjectPropertyValue(UObject* InObject, FVectorPropertyMetaData MetaData, FTrackInstancePropertyBindings* PropertyBindings, const FFloatIntermediateVector& InValue)
-	{
-		if (MetaData.bIsDouble)
-		{
-			FDoubleIntermediateVector TempDoubleValue(InValue.X, InValue.Y, InValue.Z, InValue.W);
-			TIndirectPropertyTraits<FVector3d, FDoubleIntermediateVector>::SetObjectPropertyValue(InObject, PropertyBindings, TempDoubleValue);
-			return;
-		}
-
-		switch (MetaData.NumChannels)
-		{
-			case 2: TIndirectPropertyTraits<FVector2f, FFloatIntermediateVector>::SetObjectPropertyValue(InObject, PropertyBindings, InValue); return;
-			case 3: TIndirectPropertyTraits<FVector3f, FFloatIntermediateVector>::SetObjectPropertyValue(InObject, PropertyBindings, InValue); return;
-			case 4: TIndirectPropertyTraits<FVector4,  FFloatIntermediateVector>::SetObjectPropertyValue(InObject, PropertyBindings, InValue); return;
-		}
-
-		checkf(false, TEXT("Invalid number of channels"));
-	}
-
-	static FFloatIntermediateVector CombineComposites(FVectorPropertyMetaData MetaData, float InX, float InY, float InZ, float InW)
-	{
-		return FFloatIntermediateVector(InX, InY, InZ, InW);
-	}
-};
-
 struct FDoubleVectorPropertyTraits
 {
 	using StorageType  = FDoubleIntermediateVector;
@@ -303,7 +193,7 @@
 	{
 		switch (MetaData.NumChannels)
 		{
-			case 2: TIndirectPropertyTraits<FVector2D, FDoubleIntermediateVector>::GetObjectPropertyValue(InObject, PropertyOffset, OutValue); return;
+			case 2: TIndirectPropertyTraits<FVector2d, FDoubleIntermediateVector>::GetObjectPropertyValue(InObject, PropertyOffset, OutValue); return;
 			case 3: TIndirectPropertyTraits<FVector3d, FDoubleIntermediateVector>::GetObjectPropertyValue(InObject, PropertyOffset, OutValue); return;
 			case 4: TIndirectPropertyTraits<FVector4d, FDoubleIntermediateVector>::GetObjectPropertyValue(InObject, PropertyOffset, OutValue); return;
 		}
@@ -312,7 +202,7 @@
 	{
 		switch (MetaData.NumChannels)
 		{
-			case 2: TIndirectPropertyTraits<FVector2D, FDoubleIntermediateVector>::GetObjectPropertyValue(InObject, PropertyBindings, OutValue); return;
+			case 2: TIndirectPropertyTraits<FVector2d, FDoubleIntermediateVector>::GetObjectPropertyValue(InObject, PropertyBindings, OutValue); return;
 			case 3: TIndirectPropertyTraits<FVector3d, FDoubleIntermediateVector>::GetObjectPropertyValue(InObject, PropertyBindings, OutValue); return;
 			case 4: TIndirectPropertyTraits<FVector4d, FDoubleIntermediateVector>::GetObjectPropertyValue(InObject, PropertyBindings, OutValue); return;
 		}
@@ -321,7 +211,7 @@
 	{
 		switch (MetaData.NumChannels)
 		{
-			case 2: TIndirectPropertyTraits<FVector2D, FDoubleIntermediateVector>::GetObjectPropertyValue(InObject, PropertyPath, OutValue); return;
+			case 2: TIndirectPropertyTraits<FVector2d, FDoubleIntermediateVector>::GetObjectPropertyValue(InObject, PropertyPath, OutValue); return;
 			case 3: TIndirectPropertyTraits<FVector3d, FDoubleIntermediateVector>::GetObjectPropertyValue(InObject, PropertyPath, OutValue); return;
 			case 4: TIndirectPropertyTraits<FVector4d, FDoubleIntermediateVector>::GetObjectPropertyValue(InObject, PropertyPath, OutValue); return;
 		}
@@ -336,7 +226,7 @@
 	{
 		switch (MetaData.NumChannels)
 		{
-			case 2: TIndirectPropertyTraits<FVector2D, FDoubleIntermediateVector>::SetObjectPropertyValue(InObject, PropertyOffset, InValue); return;
+			case 2: TIndirectPropertyTraits<FVector2d, FDoubleIntermediateVector>::SetObjectPropertyValue(InObject, PropertyOffset, InValue); return;
 			case 3: TIndirectPropertyTraits<FVector3d, FDoubleIntermediateVector>::SetObjectPropertyValue(InObject, PropertyOffset, InValue); return;
 			case 4: TIndirectPropertyTraits<FVector4d, FDoubleIntermediateVector>::SetObjectPropertyValue(InObject, PropertyOffset, InValue); return;
 		}
@@ -347,7 +237,7 @@
 	{
 		switch (MetaData.NumChannels)
 		{
-			case 2: TIndirectPropertyTraits<FVector2D, FDoubleIntermediateVector>::SetObjectPropertyValue(InObject, PropertyBindings, InValue); return;
+			case 2: TIndirectPropertyTraits<FVector2d, FDoubleIntermediateVector>::SetObjectPropertyValue(InObject, PropertyBindings, InValue); return;
 			case 3: TIndirectPropertyTraits<FVector3d, FDoubleIntermediateVector>::SetObjectPropertyValue(InObject, PropertyBindings, InValue); return;
 			case 4: TIndirectPropertyTraits<FVector4d, FDoubleIntermediateVector>::SetObjectPropertyValue(InObject, PropertyBindings, InValue); return;
 		}
@@ -361,3 +251,115 @@
 	}
 };
 
+struct FFloatVectorPropertyTraits
+{
+	using StorageType = FFloatIntermediateVector;
+	using MetaDataType = TPropertyMetaData<FVectorPropertyMetaData>;
+
+	static void GetObjectPropertyValue(const UObject* InObject, FVectorPropertyMetaData MetaData, const FCustomPropertyAccessor& BaseCustomAccessor, FFloatIntermediateVector& OutValue)
+	{
+		checkf(!MetaData.bIsDouble, TEXT("Type mismatch between FVectorXf and FVectorXd. Please check for any custom accessors defined on the wrong vector property type."));
+		const TCustomPropertyAccessor<FFloatVectorPropertyTraits>& CustomAccessor = static_cast<const TCustomPropertyAccessor<FFloatVectorPropertyTraits>&>(BaseCustomAccessor);
+		OutValue = (*CustomAccessor.Functions.Getter)(InObject, MetaData);
+	}
+	static void GetObjectPropertyValue(const UObject* InObject, FVectorPropertyMetaData MetaData, uint16 PropertyOffset, FFloatIntermediateVector& OutValue)
+	{
+		if (MetaData.bIsDouble)
+		{
+			FDoubleIntermediateVector TempDoubleValue;
+			FDoubleVectorPropertyTraits::GetObjectPropertyValue(InObject, MetaData, PropertyOffset, TempDoubleValue);
+			OutValue = FFloatIntermediateVector(TempDoubleValue.X, TempDoubleValue.Y, TempDoubleValue.Z, TempDoubleValue.W);
+			return;
+		}
+
+		switch (MetaData.NumChannels)
+		{
+		case 2: TIndirectPropertyTraits<FVector2f, FFloatIntermediateVector>::GetObjectPropertyValue(InObject, PropertyOffset, OutValue); return;
+		case 3: TIndirectPropertyTraits<FVector3f, FFloatIntermediateVector>::GetObjectPropertyValue(InObject, PropertyOffset, OutValue); return;
+		case 4: TIndirectPropertyTraits<FVector4f, FFloatIntermediateVector>::GetObjectPropertyValue(InObject, PropertyOffset, OutValue); return;
+		}
+	}
+	static void GetObjectPropertyValue(const UObject* InObject, FVectorPropertyMetaData MetaData, FTrackInstancePropertyBindings* PropertyBindings, FFloatIntermediateVector& OutValue)
+	{
+		if (MetaData.bIsDouble)
+		{
+			FDoubleIntermediateVector TempDoubleValue;
+			FDoubleVectorPropertyTraits::GetObjectPropertyValue(InObject, MetaData, PropertyBindings, TempDoubleValue);
+			OutValue = FFloatIntermediateVector(TempDoubleValue.X, TempDoubleValue.Y, TempDoubleValue.Z, TempDoubleValue.W);
+			return;
+		}
+
+		switch (MetaData.NumChannels)
+		{
+		case 2: TIndirectPropertyTraits<FVector2f, FFloatIntermediateVector>::GetObjectPropertyValue(InObject, PropertyBindings, OutValue); return;
+		case 3: TIndirectPropertyTraits<FVector3f, FFloatIntermediateVector>::GetObjectPropertyValue(InObject, PropertyBindings, OutValue); return;
+		case 4: TIndirectPropertyTraits<FVector4f, FFloatIntermediateVector>::GetObjectPropertyValue(InObject, PropertyBindings, OutValue); return;
+		}
+	}
+	static void GetObjectPropertyValue(const UObject* InObject, FVectorPropertyMetaData MetaData, const FName& PropertyPath, StorageType& OutValue)
+	{
+		if (MetaData.bIsDouble)
+		{
+			FDoubleIntermediateVector TempDoubleValue;
+			FDoubleVectorPropertyTraits::GetObjectPropertyValue(InObject, MetaData, PropertyPath, TempDoubleValue);
+			OutValue = FFloatIntermediateVector(TempDoubleValue.X, TempDoubleValue.Y, TempDoubleValue.Z, TempDoubleValue.W);
+			return;
+		}
+
+		switch (MetaData.NumChannels)
+		{
+		case 2: TIndirectPropertyTraits<FVector2f, FFloatIntermediateVector>::GetObjectPropertyValue(InObject, PropertyPath, OutValue); return;
+		case 3: TIndirectPropertyTraits<FVector3f, FFloatIntermediateVector>::GetObjectPropertyValue(InObject, PropertyPath, OutValue); return;
+		case 4: TIndirectPropertyTraits<FVector4f, FFloatIntermediateVector>::GetObjectPropertyValue(InObject, PropertyPath, OutValue); return;
+		}
+	}
+
+	static void SetObjectPropertyValue(UObject* InObject, FVectorPropertyMetaData MetaData, const FCustomPropertyAccessor& BaseCustomAccessor, const FFloatIntermediateVector& InValue)
+	{
+		checkf(!MetaData.bIsDouble, TEXT("Type mismatch between FVectorXf and FVectorXd. Please check for any custom accessors defined on the wrong vector property type."));
+		const TCustomPropertyAccessor<FFloatVectorPropertyTraits>& CustomAccessor = static_cast<const TCustomPropertyAccessor<FFloatVectorPropertyTraits>&>(BaseCustomAccessor);
+		(*CustomAccessor.Functions.Setter)(InObject, MetaData, InValue);
+	}
+	static void SetObjectPropertyValue(UObject* InObject, FVectorPropertyMetaData MetaData, uint16 PropertyOffset, const FFloatIntermediateVector& InValue)
+	{
+		if (MetaData.bIsDouble)
+		{
+			FDoubleIntermediateVector TempDoubleValue(InValue.X, InValue.Y, InValue.Z, InValue.W);
+			FDoubleVectorPropertyTraits::SetObjectPropertyValue(InObject, MetaData, PropertyOffset, TempDoubleValue);
+			return;
+		}
+
+		switch (MetaData.NumChannels)
+		{
+		case 2: TIndirectPropertyTraits<FVector2f, FFloatIntermediateVector>::SetObjectPropertyValue(InObject, PropertyOffset, InValue); return;
+		case 3: TIndirectPropertyTraits<FVector3f, FFloatIntermediateVector>::SetObjectPropertyValue(InObject, PropertyOffset, InValue); return;
+		case 4: TIndirectPropertyTraits<FVector4f, FFloatIntermediateVector>::SetObjectPropertyValue(InObject, PropertyOffset, InValue); return;
+		}
+
+		checkf(false, TEXT("Invalid number of channels"));
+	}
+	static void SetObjectPropertyValue(UObject* InObject, FVectorPropertyMetaData MetaData, FTrackInstancePropertyBindings* PropertyBindings, const FFloatIntermediateVector& InValue)
+	{
+		if (MetaData.bIsDouble)
+		{
+			FDoubleIntermediateVector TempDoubleValue(InValue.X, InValue.Y, InValue.Z, InValue.W);
+			FDoubleVectorPropertyTraits::SetObjectPropertyValue(InObject, MetaData, PropertyBindings, TempDoubleValue);
+			return;
+		}
+
+		switch (MetaData.NumChannels)
+		{
+		case 2: TIndirectPropertyTraits<FVector2f, FFloatIntermediateVector>::SetObjectPropertyValue(InObject, PropertyBindings, InValue); return;
+		case 3: TIndirectPropertyTraits<FVector3f, FFloatIntermediateVector>::SetObjectPropertyValue(InObject, PropertyBindings, InValue); return;
+		case 4: TIndirectPropertyTraits<FVector4f, FFloatIntermediateVector>::SetObjectPropertyValue(InObject, PropertyBindings, InValue); return;
+		}
+
+		checkf(false, TEXT("Invalid number of channels"));
+	}
+
+	static FFloatIntermediateVector CombineComposites(FVectorPropertyMetaData MetaData, float InX, float InY, float InZ, float InW)
+	{
+		return FFloatIntermediateVector(InX, InY, InZ, InW);
+	}
+};
+
diff --git a/Engine/Source/Runtime/MovieSceneTracks/Public/MovieSceneTracksPropertyTypes.h b/Engine/Source/Runtime/MovieSceneTracks/Public/MovieSceneTracksPropertyTypes.h
--- a/Engine/Source/Runtime/MovieSceneTracks/Public/MovieSceneTracksPropertyTypes.h
+++ b/Engine/Source/Runtime/MovieSceneTracks/Public/MovieSceneTracksPropertyTypes.h
@@ -205,15 +205,15 @@ MOVIESCENETRACKS_API void ConvertOperationalProperty(const FSlateColor& InColor,
 
 MOVIESCENETRACKS_API void ConvertOperationalProperty(const FFloatIntermediateVector& InVector, FVector2f& Out);
 MOVIESCENETRACKS_API void ConvertOperationalProperty(const FFloatIntermediateVector& InVector, FVector3f& Out);
-MOVIESCENETRACKS_API void ConvertOperationalProperty(const FFloatIntermediateVector& InVector, FVector4& Out);
+MOVIESCENETRACKS_API void ConvertOperationalProperty(const FFloatIntermediateVector& InVector, FVector4f& Out);
 MOVIESCENETRACKS_API void ConvertOperationalProperty(const FVector2f& In, FFloatIntermediateVector& Out);
 MOVIESCENETRACKS_API void ConvertOperationalProperty(const FVector3f& In, FFloatIntermediateVector& Out);
-MOVIESCENETRACKS_API void ConvertOperationalProperty(const FVector4& In, FFloatIntermediateVector& Out);
+MOVIESCENETRACKS_API void ConvertOperationalProperty(const FVector4f& In, FFloatIntermediateVector& Out);
 
-MOVIESCENETRACKS_API void ConvertOperationalProperty(const FDoubleIntermediateVector& InVector, FVector2D& Out);
+MOVIESCENETRACKS_API void ConvertOperationalProperty(const FDoubleIntermediateVector& InVector, FVector2d& Out);
 MOVIESCENETRACKS_API void ConvertOperationalProperty(const FDoubleIntermediateVector& InVector, FVector3d& Out);
 MOVIESCENETRACKS_API void ConvertOperationalProperty(const FDoubleIntermediateVector& InVector, FVector4d& Out);
-MOVIESCENETRACKS_API void ConvertOperationalProperty(const FVector2D& In, FDoubleIntermediateVector& Out);
+MOVIESCENETRACKS_API void ConvertOperationalProperty(const FVector2d& In, FDoubleIntermediateVector& Out);
 MOVIESCENETRACKS_API void ConvertOperationalProperty(const FVector3d& In, FDoubleIntermediateVector& Out);
 MOVIESCENETRACKS_API void ConvertOperationalProperty(const FVector4d& In, FDoubleIntermediateVector& Out);
 
