Sequencer: More fixes for LWC type discrepancies

- BP "real" properties used to be floats, but are now doubles, so we need to support widening values on the fly before applying them to objects.
- Because of FVector2d/FVector2D name collisions, we actually can't specifically test for the former in code. But testing for the latter should be enough for our needs.

diff --git a/Engine/Plugins/MovieScene/TemplateSequence/Source/TemplateSequence/Private/CameraAnimationSequencePlayer.cpp b/Engine/Plugins/MovieScene/TemplateSequence/Source/TemplateSequence/Private/CameraAnimationSequencePlayer.cpp
--- a/Engine/Plugins/MovieScene/TemplateSequence/Source/TemplateSequence/Private/CameraAnimationSequencePlayer.cpp
+++ b/Engine/Plugins/MovieScene/TemplateSequence/Source/TemplateSequence/Private/CameraAnimationSequencePlayer.cpp
@@ -33,22 +33,36 @@ void SetCameraStandInTransform(UObject* Object, const FIntermediate3DTransform&
 	CameraStandIn->SetTransform(Result);
 }
 
-template<typename PropertyTraits>
-void UpdateInitialPropertyValues(UMovieSceneEntitySystemLinker* Linker, const TPropertyComponents<PropertyTraits>& PropertyComponents)
+template<typename PropertyTraits, typename MetaDataIndices>
+struct TUpdateInitialPropertyValuesImpl;
+
+template<typename PropertyTraits, int ...MetaDataIndices>
+struct TUpdateInitialPropertyValuesImpl<PropertyTraits, TIntegerSequence<int, MetaDataIndices...>>
 {
-	const FBuiltInComponentTypes* const BuiltInComponents = FBuiltInComponentTypes::Get();
+	void Update(UMovieSceneEntitySystemLinker* Linker, const TPropertyComponents<PropertyTraits>& PropertyComponents)
+	{
+		const FBuiltInComponentTypes* const BuiltInComponents = FBuiltInComponentTypes::Get();
 
-	const FPropertyDefinition& PropertyDefinition = BuiltInComponents->PropertyRegistry.GetDefinition(PropertyComponents.CompositeID);
+		const FPropertyDefinition& PropertyDefinition = BuiltInComponents->PropertyRegistry.GetDefinition(PropertyComponents.CompositeID);
 
-	TGetPropertyValues<PropertyTraits> GetProperties(PropertyDefinition.CustomPropertyRegistration);
+		TGetPropertyValues<PropertyTraits> GetProperties(PropertyDefinition.CustomPropertyRegistration);
 
-	FEntityTaskBuilder()
-	.Read(BuiltInComponents->BoundObject)
-	.ReadOneOf(BuiltInComponents->CustomPropertyIndex, BuiltInComponents->FastPropertyOffset, BuiltInComponents->SlowProperty)
-	.Write(PropertyComponents.InitialValue)
-	.FilterAll({ PropertyComponents.PropertyTag })
-	.SetDesiredThread(Linker->EntityManager.GetGatherThread())
-	.RunInline_PerAllocation(&Linker->EntityManager, GetProperties);
+		FEntityTaskBuilder()
+			.Read(BuiltInComponents->BoundObject)
+			.ReadOneOf(BuiltInComponents->CustomPropertyIndex, BuiltInComponents->FastPropertyOffset, BuiltInComponents->SlowProperty)
+			.ReadAllOf(PropertyComponents.MetaDataComponents.template GetType<MetaDataIndices>()...)
+			.Write(PropertyComponents.InitialValue)
+			.FilterAll({ PropertyComponents.PropertyTag })
+			.SetDesiredThread(Linker->EntityManager.GetGatherThread())
+			.RunInline_PerAllocation(&Linker->EntityManager, GetProperties);
+	}
+};
+
+template<typename PropertyTraits>
+void UpdateInitialPropertyValues(UMovieSceneEntitySystemLinker* Linker, const TPropertyComponents<PropertyTraits>& PropertyComponents)
+{
+	TUpdateInitialPropertyValuesImpl<PropertyTraits, TMakeIntegerSequence<int, PropertyTraits::MetaDataType::Num>>()
+		.Update(Linker, PropertyComponents);
 }
 
 }
diff --git a/Engine/Source/Runtime/MovieSceneTracks/Private/MovieSceneTracksComponentTypes.cpp b/Engine/Source/Runtime/MovieSceneTracks/Private/MovieSceneTracksComponentTypes.cpp
--- a/Engine/Source/Runtime/MovieSceneTracks/Private/MovieSceneTracksComponentTypes.cpp
+++ b/Engine/Source/Runtime/MovieSceneTracks/Private/MovieSceneTracksComponentTypes.cpp
@@ -227,38 +227,44 @@ void SetSkyLightComponentLightColor(UObject* Object, EColorPropertyType InColorT
 	SkyLightComponent->SetLightColor(InColor.GetLinearColor());
 }
 
-float GetSecondFogDataFogDensity(const UObject* Object)
+float GetSecondFogDataFogDensity(const UObject* Object, bool bIsDouble)
 {
+	check(!bIsDouble);
 	const UExponentialHeightFogComponent* ExponentialHeightFogComponent = CastChecked<const UExponentialHeightFogComponent>(Object);
 	return ExponentialHeightFogComponent->SecondFogData.FogDensity;
 }
 
-void SetSecondFogDataFogDensity(UObject* Object, float InFogDensity)
+void SetSecondFogDataFogDensity(UObject* Object, bool bIsDouble, float InFogDensity)
 {
+	check(!bIsDouble);
 	UExponentialHeightFogComponent* ExponentialHeightFogComponent = CastChecked<UExponentialHeightFogComponent>(Object);
 	ExponentialHeightFogComponent->SecondFogData.FogDensity = InFogDensity;
 }
 
-float GetSecondFogDataFogHeightFalloff(const UObject* Object)
+float GetSecondFogDataFogHeightFalloff(const UObject* Object, bool bIsDouble)
 {
+	check(!bIsDouble);
 	const UExponentialHeightFogComponent* ExponentialHeightFogComponent = CastChecked<const UExponentialHeightFogComponent>(Object);
 	return ExponentialHeightFogComponent->SecondFogData.FogHeightFalloff;
 }
 
-void SetSecondFogDataFogHeightFalloff(UObject* Object, float InFogHeightFalloff)
+void SetSecondFogDataFogHeightFalloff(UObject* Object, bool bIsDouble, float InFogHeightFalloff)
 {
+	check(!bIsDouble);
 	UExponentialHeightFogComponent* ExponentialHeightFogComponent = CastChecked<UExponentialHeightFogComponent>(Object);
 	ExponentialHeightFogComponent->SecondFogData.FogHeightFalloff = InFogHeightFalloff;
 }
 
-float GetSecondFogDataFogHeightOffset(const UObject* Object)
+float GetSecondFogDataFogHeightOffset(const UObject* Object, bool bIsDouble)
 {
+	check(!bIsDouble);
 	const UExponentialHeightFogComponent* ExponentialHeightFogComponent = CastChecked<const UExponentialHeightFogComponent>(Object);
 	return ExponentialHeightFogComponent->SecondFogData.FogHeightOffset;
 }
 
-void SetSecondFogDataFogHeightOffset(UObject* Object, float InFogHeightOffset)
+void SetSecondFogDataFogHeightOffset(UObject* Object, bool bIsDouble, float InFogHeightOffset)
 {
+	check(!bIsDouble);
 	UExponentialHeightFogComponent* ExponentialHeightFogComponent = CastChecked<UExponentialHeightFogComponent>(Object);
 	ExponentialHeightFogComponent->SecondFogData.FogHeightOffset = InFogHeightOffset;
 }
@@ -361,6 +367,33 @@ void FComponentDetachParams::ApplyDetach(USceneComponent* ChildComponentToAttach
 static bool GMovieSceneTracksComponentTypesDestroyed = false;
 static TUniquePtr<FMovieSceneTracksComponentTypes> GMovieSceneTracksComponentTypes;
 
+struct FFloatHandler : TPropertyComponentHandler<FFloatPropertyTraits, float>
+{
+	virtual void DispatchInitializePropertyMetaDataTasks(const FPropertyDefinition& Definition, FSystemTaskPrerequisites& InPrerequisites, FSystemSubsequentTasks& Subsequents, UMovieSceneEntitySystemLinker* Linker) override
+	{
+		FBuiltInComponentTypes* BuiltInComponents = FBuiltInComponentTypes::Get();
+		FMovieSceneTracksComponentTypes* TrackComponents = FMovieSceneTracksComponentTypes::Get();
+
+		FEntityTaskBuilder()
+		.Read(BuiltInComponents->BoundObject)
+		.Read(BuiltInComponents->PropertyBinding)
+		.Write(TrackComponents->Float.MetaDataComponents.GetType<0>())
+		.FilterAll({ BuiltInComponents->Tags.NeedsLink })
+		.Iterate_PerEntity(&Linker->EntityManager, [](UObject* Object, const FMovieScenePropertyBinding& Binding, bool& OutIsDouble)
+		{
+			FProperty* BoundProperty = FTrackInstancePropertyBindings::FindProperty(Object, Binding.PropertyPath.ToString());
+			if (ensure(BoundProperty))
+			{
+				OutIsDouble = BoundProperty->IsA<FDoubleProperty>();
+			}
+			else
+			{
+				OutIsDouble = false;
+			}
+		});
+	}
+};
+
 struct FColorHandler : TPropertyComponentHandler<FColorPropertyTraits, float, float, float, float>
 {
 	virtual void DispatchInitializePropertyMetaDataTasks(const FPropertyDefinition& Definition, FSystemTaskPrerequisites& InPrerequisites, FSystemSubsequentTasks& Subsequents, UMovieSceneEntitySystemLinker* Linker) override
@@ -421,10 +454,10 @@ struct FFloatVectorHandler : TPropertyComponentHandler<FFloatVectorPropertyTrait
 			FStructProperty* BoundProperty = CastField<FStructProperty>(FTrackInstancePropertyBindings::FindProperty(Object, Binding.PropertyPath.ToString()));
 			if (ensure(BoundProperty && BoundProperty->Struct))
 			{
-				if (BoundProperty->Struct == TBaseStructure<FVector2D>::Get() || BoundProperty->Struct == TVariantStructure<FVector2f>::Get() || BoundProperty->Struct == TVariantStructure<FVector2d>::Get())
+				if (BoundProperty->Struct == TBaseStructure<FVector2D>::Get() || BoundProperty->Struct == TVariantStructure<FVector2f>::Get()) // LWC_TODO: Fix LWC type FName case
 				{
 					OutMetaData.NumChannels = 2;
-					OutMetaData.bIsDouble = (BoundProperty->Struct == TBaseStructure<FVector2D>::Get() || BoundProperty->Struct == TVariantStructure<FVector2d>::Get());
+					OutMetaData.bIsDouble = (BoundProperty->Struct == TBaseStructure<FVector2D>::Get()); // LWC_TODO: Fix LWC type FName case
 				}
 				else if (BoundProperty->Struct == TBaseStructure<FVector>::Get() || BoundProperty->Struct == TVariantStructure<FVector3f>::Get() || BoundProperty->Struct == TVariantStructure<FVector3d>::Get())
 				{
@@ -513,6 +546,7 @@ FMovieSceneTracksComponentTypes::FMovieSceneTracksComponentTypes()
 	ComponentRegistry->NewPropertyType(EulerTransform, TEXT("FEulerTransform"));
 	ComponentRegistry->NewPropertyType(ComponentTransform, TEXT("Component Transform"));
 
+	Float.MetaDataComponents.Initialize(ComponentRegistry, TEXT("Is Double"));
 	Color.MetaDataComponents.Initialize(ComponentRegistry, TEXT("Color Type"));
 	FloatVector.MetaDataComponents.Initialize(ComponentRegistry, TEXT("Num Float Vector Channels"));
 	DoubleVector.MetaDataComponents.Initialize(ComponentRegistry, TEXT("Num Double Vector Channels"));
@@ -584,7 +618,7 @@ FMovieSceneTracksComponentTypes::FMovieSceneTracksComponentTypes()
 	.AddSoleChannel(BuiltInComponents->FloatResult[0])
 	.SetBlenderSystem<UMovieScenePiecewiseFloatBlenderSystem>()
 	.SetCustomAccessors(&Accessors.Float)
-	.Commit();
+	.Commit(FFloatHandler());
 
 	// --------------------------------------------------------------------------------------------
 	// Set up double properties
diff --git a/Engine/Source/Runtime/MovieSceneTracks/Public/MovieSceneTracksComponentTypes.h b/Engine/Source/Runtime/MovieSceneTracks/Public/MovieSceneTracksComponentTypes.h
--- a/Engine/Source/Runtime/MovieSceneTracks/Public/MovieSceneTracksComponentTypes.h
+++ b/Engine/Source/Runtime/MovieSceneTracks/Public/MovieSceneTracksComponentTypes.h
@@ -111,6 +111,95 @@ struct FAttachmentComponent
 	FComponentDetachParams DetachParams;
 };
 
+struct FFloatPropertyTraits
+{
+	using StorageType  = float;
+	using MetaDataType = TPropertyMetaData<bool>;
+
+	using FloatTraitsImpl = TDirectPropertyTraits<float>;
+	using DoubleTraitsImpl = TDirectPropertyTraits<double>;
+
+	static void GetObjectPropertyValue(const UObject* InObject, bool bIsDouble, const FCustomPropertyAccessor& BaseCustomAccessor, float& OutValue)
+	{
+		checkf(!bIsDouble, TEXT("Type mismatch between float and double. Please check for any custom accessors defined on the wrong property type."));
+		FloatTraitsImpl::GetObjectPropertyValue(InObject, BaseCustomAccessor, OutValue);
+	}
+	static void GetObjectPropertyValue(const UObject* InObject, bool bIsDouble, uint16 PropertyOffset, float& OutValue)
+	{
+		if (bIsDouble)
+		{
+			double TempValue;
+			DoubleTraitsImpl::GetObjectPropertyValue(InObject, PropertyOffset, TempValue);
+			OutValue = TempValue;
+		}
+		else
+		{
+			FloatTraitsImpl::GetObjectPropertyValue(InObject, PropertyOffset, OutValue);
+		}
+	}
+	static void GetObjectPropertyValue(const UObject* InObject, bool bIsDouble, FTrackInstancePropertyBindings* PropertyBindings, float& OutValue)
+	{
+		if (bIsDouble)
+		{
+			double TempValue;
+			DoubleTraitsImpl::GetObjectPropertyValue(InObject, PropertyBindings, TempValue);
+			OutValue = TempValue;
+		}
+		else
+		{
+			FloatTraitsImpl::GetObjectPropertyValue(InObject, PropertyBindings, OutValue);
+		}
+	}
+	static void GetObjectPropertyValue(const UObject* InObject, bool bIsDouble, const FName& PropertyPath, float& OutValue)
+	{
+		if (bIsDouble)
+		{
+			double TempValue;
+			DoubleTraitsImpl::GetObjectPropertyValue(InObject, PropertyPath, TempValue);
+			OutValue = TempValue;
+		}
+		else
+		{
+			FloatTraitsImpl::GetObjectPropertyValue(InObject, PropertyPath, OutValue);
+		}
+	}
+
+	static void SetObjectPropertyValue(UObject* InObject, bool bIsDouble, const FCustomPropertyAccessor& BaseCustomAccessor, float InValue)
+	{
+		checkf(!bIsDouble, TEXT("Type mismatch between float and double. Please check for any custom accessors defined on the wrong vector property type."));
+		FloatTraitsImpl::SetObjectPropertyValue(InObject, BaseCustomAccessor, InValue);
+	}
+	static void SetObjectPropertyValue(UObject* InObject, bool bIsDouble, uint16 PropertyOffset, float InValue)
+	{
+		if (bIsDouble)
+		{
+			double TempValue = (float)InValue;
+			DoubleTraitsImpl::SetObjectPropertyValue(InObject, PropertyOffset, TempValue);
+		}
+		else
+		{
+			FloatTraitsImpl::SetObjectPropertyValue(InObject, PropertyOffset, InValue);
+		}
+	}
+	static void SetObjectPropertyValue(UObject* InObject, bool bIsDouble, FTrackInstancePropertyBindings* PropertyBindings, float InValue)
+	{
+		if (bIsDouble)
+		{
+			double TempValue = (float)InValue;
+			DoubleTraitsImpl::SetObjectPropertyValue(InObject, PropertyBindings, TempValue);
+		}
+		else
+		{
+			FloatTraitsImpl::SetObjectPropertyValue(InObject, PropertyBindings, InValue);
+		}
+	}
+
+	static float CombineComposites(bool bIsDouble, float InValue)
+	{
+		return InValue;
+	}
+};
+
 struct FColorPropertyTraits
 {
 	using StorageType  = FIntermediateColor;
@@ -367,7 +456,6 @@ using FBoolPropertyTraits               = TDirectPropertyTraits<bool>;
 using FBytePropertyTraits               = TDirectPropertyTraits<uint8>;
 using FEnumPropertyTraits               = TDirectPropertyTraits<uint8>;
 using FIntPropertyTraits                = TDirectPropertyTraits<int32>;
-using FFloatPropertyTraits              = TDirectPropertyTraits<float>;
 using FDoublePropertyTraits             = TDirectPropertyTraits<double>;
 using FTransformPropertyTraits          = TIndirectPropertyTraits<FTransform, FIntermediate3DTransform>;
 using FEulerTransformPropertyTraits     = TIndirectPropertyTraits<FEulerTransform, FIntermediate3DTransform>;
diff --git a/Engine/Source/Runtime/UMG/Private/Animation/MovieSceneUMGComponentTypes.cpp b/Engine/Source/Runtime/UMG/Private/Animation/MovieSceneUMGComponentTypes.cpp
--- a/Engine/Source/Runtime/UMG/Private/Animation/MovieSceneUMGComponentTypes.cpp
+++ b/Engine/Source/Runtime/UMG/Private/Animation/MovieSceneUMGComponentTypes.cpp
@@ -40,13 +40,15 @@ void ConvertOperationalProperty(const FWidgetTransform& In, FIntermediateWidgetT
 	Out.ShearY = In.Shear.Y;
 }
 
-static float GetRenderOpacity(const UObject* Object)
+static float GetRenderOpacity(const UObject* Object, bool bIsDouble)
 {
+	check(!bIsDouble);
 	return CastChecked<const UWidget>(Object)->GetRenderOpacity();
 }
 
-static void SetRenderOpacity(UObject* Object, float InRenderOpacity)
+static void SetRenderOpacity(UObject* Object, bool bIsDouble, float InRenderOpacity)
 {
+	check(!bIsDouble);
 	CastChecked<UWidget>(Object)->SetRenderOpacity(InRenderOpacity);
 }
 
