Sequencer: Fix wrong forwarding of float property getter/setter for LWC compatibility

diff --git a/Engine/Source/Runtime/MovieSceneTracks/Public/MovieSceneTracksComponentTypes.h b/Engine/Source/Runtime/MovieSceneTracks/Public/MovieSceneTracksComponentTypes.h
--- a/Engine/Source/Runtime/MovieSceneTracks/Public/MovieSceneTracksComponentTypes.h
+++ b/Engine/Source/Runtime/MovieSceneTracks/Public/MovieSceneTracksComponentTypes.h
@@ -122,7 +122,8 @@ struct FFloatPropertyTraits
 	static void GetObjectPropertyValue(const UObject* InObject, bool bIsDouble, const FCustomPropertyAccessor& BaseCustomAccessor, float& OutValue)
 	{
 		checkf(!bIsDouble, TEXT("Type mismatch between float and double. Please check for any custom accessors defined on the wrong property type."));
-		FloatTraitsImpl::GetObjectPropertyValue(InObject, BaseCustomAccessor, OutValue);
+		const TCustomPropertyAccessor<FFloatPropertyTraits>& CustomAccessor = static_cast<const TCustomPropertyAccessor<FFloatPropertyTraits>&>(BaseCustomAccessor);
+		OutValue = (*CustomAccessor.Functions.Getter)(InObject, bIsDouble);
 	}
 	static void GetObjectPropertyValue(const UObject* InObject, bool bIsDouble, uint16 PropertyOffset, float& OutValue)
 	{
@@ -167,7 +168,8 @@ struct FFloatPropertyTraits
 	static void SetObjectPropertyValue(UObject* InObject, bool bIsDouble, const FCustomPropertyAccessor& BaseCustomAccessor, float InValue)
 	{
 		checkf(!bIsDouble, TEXT("Type mismatch between float and double. Please check for any custom accessors defined on the wrong vector property type."));
-		FloatTraitsImpl::SetObjectPropertyValue(InObject, BaseCustomAccessor, InValue);
+		const TCustomPropertyAccessor<FFloatPropertyTraits>& CustomAccessor = static_cast<const TCustomPropertyAccessor<FFloatPropertyTraits>&>(BaseCustomAccessor);
+		(*CustomAccessor.Functions.Setter)(InObject, bIsDouble, InValue);
 	}
 	static void SetObjectPropertyValue(UObject* InObject, bool bIsDouble, uint16 PropertyOffset, float InValue)
 	{
