Sequencer: Suppress subtitles if evaluation if set in the audio section

diff --git a/Engine/Source/Runtime/MovieSceneTracks/Private/Evaluation/MovieSceneAudioTemplate.cpp b/Engine/Source/Runtime/MovieSceneTracks/Private/Evaluation/MovieSceneAudioTemplate.cpp
--- a/Engine/Source/Runtime/MovieSceneTracks/Private/Evaluation/MovieSceneAudioTemplate.cpp
+++ b/Engine/Source/Runtime/MovieSceneTracks/Private/Evaluation/MovieSceneAudioTemplate.cpp
@@ -414,6 +414,8 @@ struct FAudioSectionExecutionToken : IMovieSceneExecutionToken
 			AudioComponent.SetPitchMultiplier(PitchMultiplier);
 		}
 
+		AudioComponent.bSuppressSubtitles = AudioSection->GetSuppressSubtitles();			
+
 		// Evaluate inputs and apply the params.
 		EvaluateAllAndSetParameters<FMovieSceneFloatChannel, float>(AudioComponent, Context.GetTime());
 		EvaluateAllAndSetParameters<FMovieSceneBoolChannel, bool> (AudioComponent, Context.GetTime());
