Sequencer: Trigger camera cut changes between camera cut sections even if the camera actor is the same

This restores previous edge-case behaviour that some people relied upon.

#UDN 00403791

diff --git a/Engine/Source/Runtime/MovieSceneTracks/Private/TrackInstances/MovieSceneCameraCutTrackInstance.cpp b/Engine/Source/Runtime/MovieSceneTracks/Private/TrackInstances/MovieSceneCameraCutTrackInstance.cpp
--- a/Engine/Source/Runtime/MovieSceneTracks/Private/TrackInstances/MovieSceneCameraCutTrackInstance.cpp
+++ b/Engine/Source/Runtime/MovieSceneTracks/Private/TrackInstances/MovieSceneCameraCutTrackInstance.cpp
@@ -45,6 +45,7 @@ namespace MovieScene
 	struct FBlendedCameraCut
 	{
 		FInstanceHandle InstanceHandle;
+		TObjectPtr<UMovieSceneSection> Section;
 
 		FMovieSceneObjectBindingID CameraBindingID;
 		FMovieSceneSequenceID OperandSequenceID;
@@ -67,8 +68,9 @@ namespace MovieScene
 
 		FBlendedCameraCut()
 		{}
-		FBlendedCameraCut(FInstanceHandle InInstanceHandle, FMovieSceneObjectBindingID InCameraBindingID, FMovieSceneSequenceID InOperandSequenceID) 
-			: InstanceHandle(InInstanceHandle)
+		FBlendedCameraCut(const FMovieSceneTrackInstanceInput& InInput, FMovieSceneObjectBindingID InCameraBindingID, FMovieSceneSequenceID InOperandSequenceID) 
+			: InstanceHandle(InInput.InstanceHandle)
+			, Section(InInput.Section)
 			, CameraBindingID(InCameraBindingID)
 			, OperandSequenceID(InOperandSequenceID)
 		{}
@@ -91,7 +93,7 @@ namespace MovieScene
 			return TMovieSceneAnimTypeID<FCameraCutPreAnimatedToken>();
 		}
 
-		virtual void RestoreState(const UE::MovieScene::FRestoreStateParams& RestoreParams) override
+		virtual void RestoreState(const FRestoreStateParams& RestoreParams) override
 		{
 			IMovieScenePlayer* Player = RestoreParams.GetTerminalPlayer();
 			if (!ensure(Player))
@@ -125,7 +127,21 @@ namespace MovieScene
 			return nullptr;
 		}
 
-		static void AnimatePreRoll(const UE::MovieScene::FPreRollCameraCut& Params, const FMovieSceneContext& Context, const FMovieSceneSequenceID& SequenceID, IMovieScenePlayer& Player)
+		static bool MatchesCameraCutCache(UObject* CameraActor, const FBlendedCameraCut& Params, const UMovieSceneCameraCutTrackInstance::FCameraCutCache& CameraCutCache)
+		{
+			return CameraActor == CameraCutCache.LastLockedCamera.Get() &&
+				Params.InstanceHandle == CameraCutCache.LastInstanceHandle &&
+				Params.Section == CameraCutCache.LastSection;
+		}
+
+		static void UpdateCameraCutCache(UObject* CameraActor, const FBlendedCameraCut& Params, UMovieSceneCameraCutTrackInstance::FCameraCutCache& OutCameraCutCache)
+		{
+			OutCameraCutCache.LastLockedCamera = CameraActor;
+			OutCameraCutCache.LastInstanceHandle = Params.InstanceHandle;
+			OutCameraCutCache.LastSection = Params.Section;
+		}
+
+		static void AnimatePreRoll(const FPreRollCameraCut& Params, const FMovieSceneContext& Context, const FMovieSceneSequenceID& SequenceID, IMovieScenePlayer& Player)
 		{
 			if (Params.bHasCutTransform)
 			{
@@ -144,10 +160,8 @@ namespace MovieScene
 			}
 		}
 
-		static bool AnimateBlendedCameraCut(const UE::MovieScene::FBlendedCameraCut& Params, UMovieSceneCameraCutTrackInstance::FCameraCutCache& CameraCutCache, const FMovieSceneContext& Context, IMovieScenePlayer& Player)
+		static bool AnimateBlendedCameraCut(const FBlendedCameraCut& Params, UMovieSceneCameraCutTrackInstance::FCameraCutCache& CameraCutCache, const FMovieSceneContext& Context, IMovieScenePlayer& Player)
 		{
-			using namespace UE::MovieScene;
-
 			UObject* CameraActor = FindBoundObject(Params.CameraBindingID, Params.OperandSequenceID, Player);
 
 			EMovieSceneCameraCutParams CameraCutParams;
@@ -165,13 +179,14 @@ namespace MovieScene
 
 			static const FMovieSceneAnimTypeID CameraAnimTypeID = FMovieSceneAnimTypeID::Unique();
 
-			if (CameraCutCache.LastLockedCamera.Get() != CameraActor)
+			const bool bMatchesCache = MatchesCameraCutCache(CameraActor, Params, CameraCutCache);
+			if (!bMatchesCache)
 			{
 				Player.SavePreAnimatedState(CameraAnimTypeID, FCameraCutPreAnimatedTokenProducer());
 
 				CameraCutParams.UnlockIfCameraObject = CameraCutCache.LastLockedCamera.Get();
 				Player.UpdateCameraCut(CameraActor, CameraCutParams);
-				CameraCutCache.LastLockedCamera = CameraActor;
+				UpdateCameraCutCache(CameraActor, Params, CameraCutCache);
 				return true;
 			}
 			else if (CameraActor || CameraCutParams.BlendTime > 0.f)
@@ -231,7 +246,7 @@ void UMovieSceneCameraCutTrackInstance::OnAnimate()
 			const UMovieSceneCameraCutTrack* Track = Section->GetTypedOuter<UMovieSceneCameraCutTrack>();
 			const FMovieSceneTimeTransform SequenceToRootTransform = Context.GetSequenceToRootTransform();
 
-			FBlendedCameraCut Params(Input.InstanceHandle, CameraBindingID, SequenceInstance.GetSequenceID());
+			FBlendedCameraCut Params(Input, CameraBindingID, SequenceInstance.GetSequenceID());
 			Params.bCanBlend = Track->bCanBlend;
 
 			// Get start/current/end time.
diff --git a/Engine/Source/Runtime/MovieSceneTracks/Private/TrackInstances/MovieSceneCameraCutTrackInstance.h b/Engine/Source/Runtime/MovieSceneTracks/Private/TrackInstances/MovieSceneCameraCutTrackInstance.h
--- a/Engine/Source/Runtime/MovieSceneTracks/Private/TrackInstances/MovieSceneCameraCutTrackInstance.h
+++ b/Engine/Source/Runtime/MovieSceneTracks/Private/TrackInstances/MovieSceneCameraCutTrackInstance.h
@@ -27,6 +27,8 @@ private:
 	struct FCameraCutCache
 	{
 		TWeakObjectPtr<> LastLockedCamera;
+		UE::MovieScene::FInstanceHandle LastInstanceHandle;
+		TObjectPtr<UMovieSceneSection> LastSection;
 	};
 
 	struct FCameraCutInputInfo
