Sequencer: Fixed component attachment assigning pre-animated-state to non-root instances

diff --git a/Engine/Source/Runtime/MovieSceneTracks/Private/Systems/MovieSceneComponentAttachmentSystem.cpp b/Engine/Source/Runtime/MovieSceneTracks/Private/Systems/MovieSceneComponentAttachmentSystem.cpp
--- a/Engine/Source/Runtime/MovieSceneTracks/Private/Systems/MovieSceneComponentAttachmentSystem.cpp
+++ b/Engine/Source/Runtime/MovieSceneTracks/Private/Systems/MovieSceneComponentAttachmentSystem.cpp
@@ -273,7 +273,7 @@ void UMovieSceneComponentAttachmentSystem::SavePreAnimatedState(const FPreAnimat
 	// Start tracking all attachments to the component transform storage
 	FEntityTaskBuilder()
 	.ReadEntityIDs()
-	.Read(BuiltInComponents->InstanceHandle)
+	.Read(BuiltInComponents->RootInstanceHandle)
 	.Read(BuiltInComponents->BoundObject)
 	.FilterAll(FilterMask)
 	.Iterate_PerAllocation(&Linker->EntityManager,
