Create camera shake instances using the world the parent sequence is running in.

If the shake contains a camera animation sequence, we need to create a moviescene linker in a valid world.

diff --git a/Engine/Source/Runtime/MovieSceneTracks/Private/Evaluation/MovieSceneCameraAnimTemplate.cpp b/Engine/Source/Runtime/MovieSceneTracks/Private/Evaluation/MovieSceneCameraAnimTemplate.cpp
--- a/Engine/Source/Runtime/MovieSceneTracks/Private/Evaluation/MovieSceneCameraAnimTemplate.cpp
+++ b/Engine/Source/Runtime/MovieSceneTracks/Private/Evaluation/MovieSceneCameraAnimTemplate.cpp
@@ -521,7 +521,8 @@ struct FCameraShakeExecutionToken : TAccumulateCameraAnimExecutionToken<FCameraS
 				return false;
 			}
 
-			CameraShakeInstance = NewObject<UCameraShakeBase>(GetTransientPackage(), SourceData.ShakeClass);
+			UObject* OuterObject = Player.GetPlaybackContext() ? Player.GetPlaybackContext() : GetTransientPackage();
+			CameraShakeInstance = NewObject<UCameraShakeBase>(OuterObject, SourceData.ShakeClass);
 			if (CameraShakeInstance)
 			{
 				// Store the anim instance with the section and always remove it when we've finished evaluating
