Properly set default value when binding ID changes for actor reference sections.

diff --git a/Engine/Source/Runtime/MovieSceneTracks/Private/Sections/MovieSceneActorReferenceSection.cpp b/Engine/Source/Runtime/MovieSceneTracks/Private/Sections/MovieSceneActorReferenceSection.cpp
--- a/Engine/Source/Runtime/MovieSceneTracks/Private/Sections/MovieSceneActorReferenceSection.cpp
+++ b/Engine/Source/Runtime/MovieSceneTracks/Private/Sections/MovieSceneActorReferenceSection.cpp
@@ -144,7 +144,10 @@ void UMovieSceneActorReferenceSection::OnBindingIDsUpdated(const TMap<UE::MovieS
 	{
 		Modify();
 
-		ActorReferenceData.GetDefault().Object = OldFixedToNewFixedMap[DefaultFixedBindingID].ConvertToRelative(LocalSequenceID, Hierarchy);
+		FMovieSceneActorReferenceKey NewDefaultValue = ActorReferenceData.GetDefault();
+		NewDefaultValue.Object = OldFixedToNewFixedMap[DefaultFixedBindingID].ConvertToRelative(LocalSequenceID, Hierarchy);
+
+		ActorReferenceData.SetDefault(NewDefaultValue);
 	}
 
 	for (FMovieSceneActorReferenceKey& Key : ActorReferenceData.GetData().GetValues())
