Sequencer: Added a new option, "Swap Root Bone With Component Root", to put the root motion on the skeletal mesh component.

When this option is enabled, the root bone is evaluated as usual, cached and then set to 0. The cached root bone transform is then set on the skeletal mesh component in FAnimSequencerInstanceProxy::PostEvaluate. If there is an existing transform track on the skeletal mesh component, however, the cached root bone transform is applied additively to that. Without this, the skeletal mesh wouldn't be able to be globally transformed by Sequencer.

Deprecated UpdateAnimTrackWithRootMotion so that it now takes FAnimSequencerData instead of each individual parameter separately.

diff --git a/Engine/Source/Runtime/AnimGraphRuntime/Private/AnimSequencerInstance.cpp b/Engine/Source/Runtime/AnimGraphRuntime/Private/AnimSequencerInstance.cpp
--- a/Engine/Source/Runtime/AnimGraphRuntime/Private/AnimSequencerInstance.cpp
+++ b/Engine/Source/Runtime/AnimGraphRuntime/Private/AnimSequencerInstance.cpp
@@ -37,12 +37,13 @@ void UAnimSequencerInstance::UpdateAnimTrack(UAnimSequenceBase* InAnimSequence,
 
 void UAnimSequencerInstance::UpdateAnimTrackWithRootMotion(UAnimSequenceBase* InAnimSequence, int32 SequenceId, const TOptional<FRootMotionOverride>& RootMotion, float InFromPosition, float InToPosition, float Weight, bool bFireNotifies)
 {
-	UpdateAnimTrackWithRootMotion(InAnimSequence, SequenceId, RootMotion, InFromPosition, InToPosition, Weight, bFireNotifies, nullptr);
+	FAnimSequencerData AnimSequencerData(InAnimSequence, SequenceId, RootMotion, InFromPosition, InToPosition, Weight, bFireNotifies, false, FTransform(), nullptr);
+	GetProxyOnGameThread<FAnimSequencerInstanceProxy>().UpdateAnimTrackWithRootMotion(AnimSequencerData);
 }
 
-void UAnimSequencerInstance::UpdateAnimTrackWithRootMotion(UAnimSequenceBase* InAnimSequence, int32 SequenceId, const TOptional<FRootMotionOverride>& RootMotion, float InFromPosition, float InToPosition, float Weight, bool bFireNotifies, UMirrorDataTable* InMirrorDataTable)
+void UAnimSequencerInstance::UpdateAnimTrackWithRootMotion(const FAnimSequencerData& InAnimSequencerData)
 {
-	GetProxyOnGameThread<FAnimSequencerInstanceProxy>().UpdateAnimTrackWithRootMotion(InAnimSequence, SequenceId, RootMotion, InFromPosition, InToPosition, Weight, bFireNotifies, InMirrorDataTable);
+	GetProxyOnGameThread<FAnimSequencerInstanceProxy>().UpdateAnimTrackWithRootMotion(InAnimSequencerData);
 }
 
 void UAnimSequencerInstance::ConstructNodes()
diff --git a/Engine/Source/Runtime/AnimGraphRuntime/Private/AnimSequencerInstanceProxy.cpp b/Engine/Source/Runtime/AnimGraphRuntime/Private/AnimSequencerInstanceProxy.cpp
--- a/Engine/Source/Runtime/AnimGraphRuntime/Private/AnimSequencerInstanceProxy.cpp
+++ b/Engine/Source/Runtime/AnimGraphRuntime/Private/AnimSequencerInstanceProxy.cpp
@@ -60,9 +60,42 @@ bool FAnimSequencerInstanceProxy::Evaluate(FPoseContext& Output)
 			}
 		}
 	}
+	
+	RootBoneTransform.Reset();
+
+	if (bSwapRootBoneWithComponentRoot)
+	{
+		for (const FCompactPoseBoneIndex BoneIndex : Output.Pose.ForEachBoneIndex())
+		{
+			if (BoneIndex.IsRootBone())
+			{
+				RootBoneTransform = Output.Pose[BoneIndex];
+				Output.Pose[BoneIndex] = FTransform::Identity;
+				break;
+			}
+		}
+	}
+
 	return true;
 }
 
+void FAnimSequencerInstanceProxy::PostEvaluate(UAnimInstance* InAnimInstance)
+{
+	if (GetSkelMeshComponent() && bSwapRootBoneWithComponentRoot)
+	{
+		if (RootBoneTransform.IsSet())
+		{
+			FTransform RelativeTransform = RootBoneTransform.GetValue();
+
+			if (InitialTransform.IsSet())
+			{
+				RelativeTransform = RootBoneTransform.GetValue() * InitialTransform.GetValue();
+			}
+			GetSkelMeshComponent()->SetRelativeLocationAndRotation(RelativeTransform.GetLocation(), RelativeTransform.GetRotation().Rotator());
+		}
+	}
+}
+
 void FAnimSequencerInstanceProxy::UpdateAnimationNode(const FAnimationUpdateContext& InContext)
 {
 	UpdateCounter.Increment();
@@ -211,6 +244,13 @@ void FAnimSequencerInstanceProxy::UpdateAnimTrackWithRootMotion(UAnimSequenceBas
 	UpdateAnimTrack(InAnimSequence, SequenceId, RootMotion, InFromPosition, InToPosition, Weight, bFireNotifies, InMirrorDataTable);
 }
 
+void FAnimSequencerInstanceProxy::UpdateAnimTrackWithRootMotion(const FAnimSequencerData& InAnimSequencerData)
+{
+	bSwapRootBoneWithComponentRoot = InAnimSequencerData.bSwapRootBoneWithComponentRoot;
+	InitialTransform = InAnimSequencerData.InitialTransform;
+	UpdateAnimTrack(InAnimSequencerData.AnimSequence, InAnimSequencerData.SequenceId, InAnimSequencerData.RootMotion, InAnimSequencerData.FromPosition, InAnimSequencerData.ToPosition, InAnimSequencerData.Weight, InAnimSequencerData.bFireNotifies, InAnimSequencerData.MirrorDataTable);
+}
+
 void FAnimSequencerInstanceProxy::UpdateAnimTrack(UAnimSequenceBase* InAnimSequence, uint32 SequenceId, const TOptional<FRootMotionOverride>& InRootMotionOverride, TOptional<float> InFromPosition, float InToPosition, float Weight, bool bFireNotifies, UMirrorDataTable* InMirrorDataTable)
 {
 	EnsureAnimTrack(InAnimSequence, SequenceId);
diff --git a/Engine/Source/Runtime/AnimGraphRuntime/Public/AnimSequencerInstance.h b/Engine/Source/Runtime/AnimGraphRuntime/Public/AnimSequencerInstance.h
--- a/Engine/Source/Runtime/AnimGraphRuntime/Public/AnimSequencerInstance.h
+++ b/Engine/Source/Runtime/AnimGraphRuntime/Public/AnimSequencerInstance.h
@@ -12,6 +12,7 @@
 #include "AnimSequencerInstance.generated.h"
 
 struct FRootMotionOverride;
+struct FAnimSequencerData;
 
 UCLASS(transient, NotBlueprintable)
 class ANIMGRAPHRUNTIME_API UAnimSequencerInstance : public UAnimInstance, public ISequencerAnimationSupport
@@ -28,7 +29,10 @@ public:
 	UE_DEPRECATED(5.0, "Please use UpdateAnimTrackWithRootMotion that takes a MirrorDataTable")
 	void UpdateAnimTrackWithRootMotion(UAnimSequenceBase* InAnimSequence, int32 SequenceId, const TOptional<FRootMotionOverride>& RootMotion, float InFromPosition, float InToPosition, float Weight, bool bFireNotifies);
 
+	UE_DEPRECATED(5.1, "Please use the UpdateAnimTrackWithRootMotion that takes FAnimSequencerData")
 	void UpdateAnimTrackWithRootMotion(UAnimSequenceBase* InAnimSequence, int32 SequenceId,const TOptional<FRootMotionOverride>& RootMotion, float InFromPosition, float InToPosition, float Weight, bool bFireNotifies, UMirrorDataTable* InMirrorDataTable);
+	
+	void UpdateAnimTrackWithRootMotion(const FAnimSequencerData& InAnimSequencerData);
 
 	/** Construct all nodes in this instance */
 	virtual void ConstructNodes() override;
diff --git a/Engine/Source/Runtime/AnimGraphRuntime/Public/AnimSequencerInstanceProxy.h b/Engine/Source/Runtime/AnimGraphRuntime/Public/AnimSequencerInstanceProxy.h
--- a/Engine/Source/Runtime/AnimGraphRuntime/Public/AnimSequencerInstanceProxy.h
+++ b/Engine/Source/Runtime/AnimGraphRuntime/Public/AnimSequencerInstanceProxy.h
@@ -56,6 +56,34 @@ struct FRootMotionOverride
 	FTransform RootMotion;
 };
 
+struct ANIMGRAPHRUNTIME_API FAnimSequencerData
+{
+	FAnimSequencerData(UAnimSequenceBase* InAnimSequence, int32 InSequenceId, const TOptional<FRootMotionOverride>& InRootMotion, float InFromPosition, float InToPosition, float InWeight, bool bInFireNotifies, bool bInSwapRootBoneWithComponentRoot, TOptional<FTransform> InInitialTransform, UMirrorDataTable* InMirrorDataTable)
+		: AnimSequence(InAnimSequence)
+		, SequenceId(InSequenceId)
+		, RootMotion(InRootMotion)
+		, FromPosition(InFromPosition)
+		, ToPosition(InToPosition)
+		, Weight(InWeight)
+		, bFireNotifies(bInFireNotifies)
+		, bSwapRootBoneWithComponentRoot(bInSwapRootBoneWithComponentRoot)
+		, InitialTransform(InInitialTransform)
+		, MirrorDataTable(InMirrorDataTable)
+	{
+	}
+
+	UAnimSequenceBase* AnimSequence;
+	int32 SequenceId;
+	const TOptional<FRootMotionOverride>& RootMotion;
+	float FromPosition;
+	float ToPosition;
+	float Weight;
+	bool bFireNotifies;
+	bool bSwapRootBoneWithComponentRoot;
+	TOptional<FTransform> InitialTransform;
+	UMirrorDataTable* MirrorDataTable;
+};
+
 /** Quick n dirty RTTI to allow for derived classes to insert nodes of different types */
 #define SEQUENCER_INSTANCE_PLAYER_TYPE(TYPE, BASE) \
 	static const FName& GetTypeId() { static FName Type(TEXT(#TYPE)); return Type; } \
@@ -91,6 +119,7 @@ public:
 	// FAnimInstanceProxy interface
 	virtual void Initialize(UAnimInstance* InAnimInstance) override;
 	virtual bool Evaluate(FPoseContext& Output) override;
+	virtual void PostEvaluate(UAnimInstance* InAnimInstance) override;
 	virtual void UpdateAnimationNode(const FAnimationUpdateContext& InContext) override;
 
 	/** Update an animation sequence player in this instance */
@@ -100,8 +129,11 @@ public:
 	UE_DEPRECATED(5.0, "Please use the UpdateAnimTrackWithRootMotion that takes a MirrorDataTable")
 	void UpdateAnimTrackWithRootMotion(UAnimSequenceBase* InAnimSequence, int32 SequenceId, const TOptional<FRootMotionOverride>& RootMotion, float InFromPosition, float InToPosition, float Weight, bool bFireNotifies);
 	
+	UE_DEPRECATED(5.1, "Please use the UpdateAnimTrackWithRootMotion that takes FAnimSequencerData")
 	void UpdateAnimTrackWithRootMotion(UAnimSequenceBase* InAnimSequence, int32 SequenceId, const TOptional<FRootMotionOverride>& RootMotion, float InFromPosition, float InToPosition, float Weight, bool bFireNotifies, UMirrorDataTable* InMirrorDataTable);
 
+	void UpdateAnimTrackWithRootMotion(const FAnimSequencerData& InAnimSequencerData);
+
 	/** Reset all nodes in this instance */
 	virtual void ResetNodes();
 
@@ -146,4 +178,8 @@ protected:
 	void InitAnimTrack(UAnimSequenceBase* InAnimSequence, uint32 SequenceId);
 	void EnsureAnimTrack(UAnimSequenceBase* InAnimSequence, uint32 SequenceId);
 	void ClearSequencePlayerAndMirrorMaps();
+
+	bool bSwapRootBoneWithComponentRoot = false;
+	TOptional<FTransform> InitialTransform;
+	TOptional<FTransform> RootBoneTransform;
 };
diff --git a/Engine/Source/Runtime/MovieSceneTracks/Private/Evaluation/MovieSceneSkeletalAnimationTemplate.cpp b/Engine/Source/Runtime/MovieSceneTracks/Private/Evaluation/MovieSceneSkeletalAnimationTemplate.cpp
--- a/Engine/Source/Runtime/MovieSceneTracks/Private/Evaluation/MovieSceneSkeletalAnimationTemplate.cpp
+++ b/Engine/Source/Runtime/MovieSceneTracks/Private/Evaluation/MovieSceneSkeletalAnimationTemplate.cpp
@@ -1,6 +1,7 @@
 // Copyright Epic Games, Inc. All Rights Reserved.
 
 #include "Evaluation/MovieSceneSkeletalAnimationTemplate.h"
+#include "Evaluation/MovieSceneTemplateCommon.h"
 #include "Compilation/MovieSceneCompilerRules.h"
 #include "Animation/AnimMontage.h"
 #include "Animation/AnimSingleNodeInstance.h"
@@ -14,6 +15,9 @@
 #include "Rendering/MotionVectorSimulation.h"
 #include "Systems/MovieSceneMotionVectorSimulationSystem.h"
 #include "EntitySystem/MovieSceneEntitySystemLinker.h"
+#include "EntitySystem/MovieSceneEntitySystemTask.h"
+#include "Evaluation/MovieSceneEvaluationTemplateInstance.h"
+#include "MovieSceneTracksComponentTypes.h"
 #include "UObject/StrongObjectPtr.h"
 #include "SkeletalMeshRestoreState.h"
 #include "AnimSequencerInstanceProxy.h"
@@ -440,6 +444,36 @@ namespace MovieScene
 			}
 		}
 		
+		// Determines whether the skeletal mesh component has a component transform property tag
+		bool ContainsTransform(IMovieScenePlayer& Player, USkeletalMeshComponent* SkeletalMeshComponent) const
+		{
+			using namespace UE::MovieScene;
+			UMovieSceneEntitySystemLinker* Linker = Player.GetEvaluationTemplate().GetEntitySystemLinker();
+
+			bool bContainsTransform = false;
+
+			auto HarvestTransforms = [SkeletalMeshComponent, &bContainsTransform](UObject* BoundObject)
+			{
+				if (BoundObject == SkeletalMeshComponent)
+				{
+					bContainsTransform = true;
+				}
+			};
+					
+			FBuiltInComponentTypes* BuiltInComponents = FBuiltInComponentTypes::Get();
+
+			FMovieSceneTracksComponentTypes* Components = FMovieSceneTracksComponentTypes::Get();
+			FEntityTaskBuilder()
+				.Read(BuiltInComponents->BoundObject)
+				// Only include component transforms
+				.FilterAll({ Components->ComponentTransform.PropertyTag })
+				// Only read things with the resolved properties on - this ensures we do not read any intermediate component transforms for blended properties
+				.FilterAny({ BuiltInComponents->CustomPropertyIndex, BuiltInComponents->FastPropertyOffset, BuiltInComponents->SlowProperty })
+				.Iterate_PerEntity(&Linker->EntityManager, HarvestTransforms);
+
+			return bContainsTransform;
+		}
+		
 		void SetAnimPosition(FPersistentEvaluationData& PersistentData, IMovieScenePlayer& Player, USkeletalMeshComponent* SkeletalMeshComponent, FName SlotName, FObjectKey Section, UAnimSequenceBase* InAnimSequence, float InFromPosition, float InToPosition, float Weight, 
 			bool bPlaying, bool bFireNotifies, bool bForceCustomMode, const UMovieSceneSkeletalAnimationSection* AnimSection, FFrameTime CurrentTime)
 		{
@@ -469,7 +503,13 @@ namespace MovieScene
 					RootMotion.GetValue().bBlendFirstChildOfRoot = bBlendFirstChildOfRoot;
 				}
 
-				SequencerInst->UpdateAnimTrackWithRootMotion(InAnimSequence, GetTypeHash(AnimTypeID),RootMotion, InFromPosition, InToPosition, Weight, bFireNotifies, AnimSection->Params.MirrorDataTable.Get());
+				// If Sequencer has a transform track, we want to set the initial transform so that root motion (if it exists) can be applied relative to that.
+				const bool bContainsTransform = ContainsTransform(Player, SkeletalMeshComponent);
+				TOptional<FTransform> CurrentTransform = bContainsTransform ? SkeletalMeshComponent->GetRelativeTransform() : TOptional<FTransform>();
+								
+				const bool bSwapRootBoneWithComponentRoot = AnimSection->Params.bSwapRootBoneWithComponentRoot;
+				FAnimSequencerData AnimSequencerData(InAnimSequence, GetTypeHash(AnimTypeID), RootMotion, InFromPosition, InToPosition, Weight, bFireNotifies, bSwapRootBoneWithComponentRoot, CurrentTransform, AnimSection->Params.MirrorDataTable.Get());
+				SequencerInst->UpdateAnimTrackWithRootMotion(AnimSequencerData);
 			}
 			else if (UAnimInstance* AnimInst = GetSourceAnimInstance(SkeletalMeshComponent))
 			{
@@ -501,6 +541,8 @@ namespace MovieScene
 			bool bFireNotifies, float DeltaTime, bool bPlaying, bool bResetDynamics, bool bForceCustomMode, const UMovieSceneSkeletalAnimationSection* AnimSection,
 			FFrameTime CurrentTime)
 		{
+			using namespace UE::MovieScene;
+
 			static const bool bLooping = false;
 
 			if (!CanPlayAnimation(SkeletalMeshComponent, InAnimSequence))
@@ -528,8 +570,14 @@ namespace MovieScene
 					RootMotion.GetValue().RootMotion = RootMotionTransform.GetValue();
 					RootMotion.GetValue().bBlendFirstChildOfRoot = bBlendFirstChildOfRoot;
 				}
-
-				SequencerInst->UpdateAnimTrackWithRootMotion(InAnimSequence, GetTypeHash(AnimTypeID), RootMotion, InFromPosition, InToPosition, Weight, bFireNotifies, AnimSection->Params.MirrorDataTable.Get());
+							
+				// If Sequencer has a transform track for the skeletal mesh component, we want to set the initial transform so that root motion (if it exists) can be applied relative to that.
+				const bool bContainsTransform = ContainsTransform(Player, SkeletalMeshComponent);
+				TOptional<FTransform> CurrentTransform = bContainsTransform ? SkeletalMeshComponent->GetRelativeTransform() : TOptional<FTransform>();
+
+				const bool bSwapRootBoneWithComponentRoot = AnimSection->Params.bSwapRootBoneWithComponentRoot;
+				FAnimSequencerData AnimSequencerData(InAnimSequence, GetTypeHash(AnimTypeID), RootMotion, InFromPosition, InToPosition, Weight, bFireNotifies, bSwapRootBoneWithComponentRoot, CurrentTransform, AnimSection->Params.MirrorDataTable.Get());
+				SequencerInst->UpdateAnimTrackWithRootMotion(AnimSequencerData);
 			}
 			else if (UAnimInstance* AnimInst = GetSourceAnimInstance(SkeletalMeshComponent))
 			{
diff --git a/Engine/Source/Runtime/MovieSceneTracks/Private/Sections/MovieSceneSkeletalAnimationSection.cpp b/Engine/Source/Runtime/MovieSceneTracks/Private/Sections/MovieSceneSkeletalAnimationSection.cpp
--- a/Engine/Source/Runtime/MovieSceneTracks/Private/Sections/MovieSceneSkeletalAnimationSection.cpp
+++ b/Engine/Source/Runtime/MovieSceneTracks/Private/Sections/MovieSceneSkeletalAnimationSection.cpp
@@ -34,6 +34,7 @@ FMovieSceneSkeletalAnimationParams::FMovieSceneSkeletalAnimationParams()
 	Weight.SetDefault(1.f);
 	bSkipAnimNotifiers = false;
 	bForceCustomMode = false;
+	bSwapRootBoneWithComponentRoot = false;
 }
 
 UMovieSceneSkeletalAnimationSection::UMovieSceneSkeletalAnimationSection( const FObjectInitializer& ObjectInitializer )
diff --git a/Engine/Source/Runtime/MovieSceneTracks/Public/Sections/MovieSceneSkeletalAnimationSection.h b/Engine/Source/Runtime/MovieSceneTracks/Public/Sections/MovieSceneSkeletalAnimationSection.h
--- a/Engine/Source/Runtime/MovieSceneTracks/Public/Sections/MovieSceneSkeletalAnimationSection.h
+++ b/Engine/Source/Runtime/MovieSceneTracks/Public/Sections/MovieSceneSkeletalAnimationSection.h
@@ -69,6 +69,10 @@ struct FMovieSceneSkeletalAnimationParams
 	UPROPERTY(EditAnywhere, Category = "Animation")
 	bool bForceCustomMode;
 
+	/** If on the root bone transform will be swapped to the component transform*/
+	UPROPERTY(EditAnywhere, Category = "Animation")
+	bool bSwapRootBoneWithComponentRoot;
+
 	UPROPERTY()
 	float StartOffset_DEPRECATED;
 
