Sequencer: Follow up to Swap Root Bone so that it gets the relative transform of the actor's root component.

diff --git a/Engine/Source/Runtime/MovieSceneTracks/Private/Evaluation/MovieSceneSkeletalAnimationTemplate.cpp b/Engine/Source/Runtime/MovieSceneTracks/Private/Evaluation/MovieSceneSkeletalAnimationTemplate.cpp
--- a/Engine/Source/Runtime/MovieSceneTracks/Private/Evaluation/MovieSceneSkeletalAnimationTemplate.cpp
+++ b/Engine/Source/Runtime/MovieSceneTracks/Private/Evaluation/MovieSceneSkeletalAnimationTemplate.cpp
@@ -444,17 +444,17 @@ namespace MovieScene
 			}
 		}
 		
-		// Determines whether the skeletal mesh component has a component transform property tag
-		bool ContainsTransform(IMovieScenePlayer& Player, USkeletalMeshComponent* SkeletalMeshComponent) const
+		// Determines whether the bound object has a component transform property tag
+		bool ContainsTransform(IMovieScenePlayer& Player, UObject* InBoundObject) const
 		{
 			using namespace UE::MovieScene;
 			UMovieSceneEntitySystemLinker* Linker = Player.GetEvaluationTemplate().GetEntitySystemLinker();
 
 			bool bContainsTransform = false;
 
-			auto HarvestTransforms = [SkeletalMeshComponent, &bContainsTransform](UObject* BoundObject)
+			auto HarvestTransforms = [InBoundObject, &bContainsTransform](UObject* BoundObject)
 			{
-				if (BoundObject == SkeletalMeshComponent)
+				if (BoundObject == InBoundObject)
 				{
 					bContainsTransform = true;
 				}
@@ -474,6 +474,35 @@ namespace MovieScene
 			return bContainsTransform;
 		}
 		
+
+		// Get the current transform for the component that the root bone will be swaped to
+		TOptional<FTransform> GetCurrentTransform(IMovieScenePlayer& Player, ESwapRootBone SwapRootBone, USkeletalMeshComponent* SkeletalMeshComponent) const
+		{
+			TOptional<FTransform> CurrentTransform;
+			if (SwapRootBone == ESwapRootBone::SwapRootBone_Component)
+			{
+				if (ContainsTransform(Player, SkeletalMeshComponent))
+				{
+					CurrentTransform = SkeletalMeshComponent->GetRelativeTransform();
+				}
+			}
+			else if (SwapRootBone == ESwapRootBone::SwapRootBone_Actor)
+			{
+				if (AActor* Actor = SkeletalMeshComponent->GetOwner())
+				{
+					if (USceneComponent* RootComponent = Actor->GetRootComponent())
+					{
+						if (ContainsTransform(Player, RootComponent))
+						{
+							CurrentTransform = RootComponent->GetRelativeTransform();
+						}
+					}
+				}
+			}
+
+			return CurrentTransform;
+		}
+
 		void SetAnimPosition(FPersistentEvaluationData& PersistentData, IMovieScenePlayer& Player, USkeletalMeshComponent* SkeletalMeshComponent, FName SlotName, FObjectKey Section, UAnimSequenceBase* InAnimSequence, float InFromPosition, float InToPosition, float Weight, 
 			bool bPlaying, bool bFireNotifies, bool bForceCustomMode, const UMovieSceneSkeletalAnimationSection* AnimSection, FFrameTime CurrentTime)
 		{
@@ -504,9 +533,8 @@ namespace MovieScene
 				}
 
 				// If Sequencer has a transform track, we want to set the initial transform so that root motion (if it exists) can be applied relative to that.
-				const bool bContainsTransform = ContainsTransform(Player, SkeletalMeshComponent);
-				TOptional<FTransform> CurrentTransform = bContainsTransform ? SkeletalMeshComponent->GetRelativeTransform() : TOptional<FTransform>();
-								
+				TOptional<FTransform> CurrentTransform = GetCurrentTransform(Player, AnimSection->Params.SwapRootBone, SkeletalMeshComponent);
+
 				FAnimSequencerData AnimSequencerData(InAnimSequence, GetTypeHash(AnimTypeID), RootMotion, InFromPosition, InToPosition, Weight, bFireNotifies, AnimSection->Params.SwapRootBone, CurrentTransform, AnimSection->Params.MirrorDataTable.Get());
 				SequencerInst->UpdateAnimTrackWithRootMotion(AnimSequencerData);
 			}
@@ -570,9 +598,8 @@ namespace MovieScene
 					RootMotion.GetValue().bBlendFirstChildOfRoot = bBlendFirstChildOfRoot;
 				}
 							
-				// If Sequencer has a transform track for the skeletal mesh component, we want to set the initial transform so that root motion (if it exists) can be applied relative to that.
-				const bool bContainsTransform = ContainsTransform(Player, SkeletalMeshComponent);
-				TOptional<FTransform> CurrentTransform = bContainsTransform ? SkeletalMeshComponent->GetRelativeTransform() : TOptional<FTransform>();
+				// If Sequencer has a transform track, we want to set the initial transform so that root motion (if it exists) can be applied relative to that.
+				TOptional<FTransform> CurrentTransform = GetCurrentTransform(Player, AnimSection->Params.SwapRootBone, SkeletalMeshComponent);
 
 				FAnimSequencerData AnimSequencerData(InAnimSequence, GetTypeHash(AnimTypeID), RootMotion, InFromPosition, InToPosition, Weight, bFireNotifies, AnimSection->Params.SwapRootBone, CurrentTransform, AnimSection->Params.MirrorDataTable.Get());
 				SequencerInst->UpdateAnimTrackWithRootMotion(AnimSequencerData);
