Sequencer: Don't fire notifies if animation has looped around.

diff --git a/Engine/Source/Runtime/MovieSceneTracks/Private/Evaluation/MovieSceneSkeletalAnimationTemplate.cpp b/Engine/Source/Runtime/MovieSceneTracks/Private/Evaluation/MovieSceneSkeletalAnimationTemplate.cpp
--- a/Engine/Source/Runtime/MovieSceneTracks/Private/Evaluation/MovieSceneSkeletalAnimationTemplate.cpp
+++ b/Engine/Source/Runtime/MovieSceneTracks/Private/Evaluation/MovieSceneSkeletalAnimationTemplate.cpp
@@ -423,13 +423,28 @@ namespace MovieScene
 
 			for (const FMinimalAnimParameters& AnimParams : Parameters)
 			{
+				// Don't fire notifies if looping around
+				bool bLooped = false;
+				if (AnimParams.AnimSection->Params.bReverse)
+				{
+					if (AnimParams.FromEvalTime <= AnimParams.ToEvalTime)
+					{
+						bLooped = true;
+					}
+				}
+				else if (AnimParams.FromEvalTime >= AnimParams.ToEvalTime)
+				{
+					bLooped = true;
+					UE_LOG(LogMovieScene, Log, TEXT("Animation looped, skipping notifies: %f %f"), AnimParams.FromEvalTime, AnimParams.ToEvalTime);
+				}
+
 				FScopedPreAnimatedCaptureSource CaptureSource(&Player.PreAnimatedState, AnimParams.EvaluationScope.Key, AnimParams.EvaluationScope.CompletionMode == EMovieSceneCompletionMode::RestoreState);
 
 				if (bPreviewPlayback)
 				{
 					PreviewSetAnimPosition(PersistentData, Player, SkeletalMeshComponent,
 						AnimParams.SlotName, AnimParams.Section, AnimParams.Animation, AnimParams.FromEvalTime, AnimParams.ToEvalTime, AnimParams.BlendWeight,
-						bFireNotifies && !AnimParams.bSkipAnimNotifiers, DeltaTime, PlayerStatus == EMovieScenePlayerStatus::Playing,
+						bFireNotifies && !AnimParams.bSkipAnimNotifiers && !bLooped, DeltaTime, PlayerStatus == EMovieScenePlayerStatus::Playing,
 						bResetDynamics, AnimParams.bForceCustomMode, AnimParams.AnimSection, AnimParams.SectionTime
 					);
 				}
@@ -437,7 +452,7 @@ namespace MovieScene
 				{
 					SetAnimPosition(PersistentData, Player, SkeletalMeshComponent,
 						AnimParams.SlotName, AnimParams.Section, AnimParams.Animation, AnimParams.FromEvalTime, AnimParams.ToEvalTime, AnimParams.BlendWeight,
-						PlayerStatus == EMovieScenePlayerStatus::Playing, bFireNotifies && !AnimParams.bSkipAnimNotifiers,
+						PlayerStatus == EMovieScenePlayerStatus::Playing, bFireNotifies && !AnimParams.bSkipAnimNotifiers && !bLooped,
 						AnimParams.bForceCustomMode, AnimParams.AnimSection, AnimParams.SectionTime
 					);
 				}
