Sequencer: Remove debugging log

diff --git a/Engine/Source/Runtime/MovieSceneTracks/Private/Evaluation/MovieSceneSkeletalAnimationTemplate.cpp b/Engine/Source/Runtime/MovieSceneTracks/Private/Evaluation/MovieSceneSkeletalAnimationTemplate.cpp
--- a/Engine/Source/Runtime/MovieSceneTracks/Private/Evaluation/MovieSceneSkeletalAnimationTemplate.cpp
+++ b/Engine/Source/Runtime/MovieSceneTracks/Private/Evaluation/MovieSceneSkeletalAnimationTemplate.cpp
@@ -435,7 +435,6 @@ namespace MovieScene
 				else if (AnimParams.FromEvalTime >= AnimParams.ToEvalTime)
 				{
 					bLooped = true;
-					UE_LOG(LogMovieScene, Log, TEXT("Animation looped, skipping notifies: %f %f"), AnimParams.FromEvalTime, AnimParams.ToEvalTime);
 				}
 
 				FScopedPreAnimatedCaptureSource CaptureSource(&Player.PreAnimatedState, AnimParams.EvaluationScope.Key, AnimParams.EvaluationScope.CompletionMode == EMovieSceneCompletionMode::RestoreState);
