Sequencer: Add more descriptive stat for each property type (instead of Apply Properties for all property types)

diff --git a/Engine/Source/Runtime/MovieScene/Public/EntitySystem/MovieScenePropertyComponentHandler.h b/Engine/Source/Runtime/MovieScene/Public/EntitySystem/MovieScenePropertyComponentHandler.h
--- a/Engine/Source/Runtime/MovieScene/Public/EntitySystem/MovieScenePropertyComponentHandler.h
+++ b/Engine/Source/Runtime/MovieScene/Public/EntitySystem/MovieScenePropertyComponentHandler.h
@@ -24,10 +24,7 @@
 namespace MovieScene
 {
 
-DECLARE_CYCLE_STAT(TEXT("Apply properties"), MovieSceneEval_ApplyProperties,  STATGROUP_MovieSceneECS);
 
-
-
 template<typename PropertyTraits, typename MetaDatatype, typename MetaDataIndices, typename CompositeIndices, typename ...CompositeTypes>
 struct TPropertyComponentHandlerImpl;
 
@@ -346,7 +343,7 @@
 		.ReadAllOf(Definition.GetMetaDataComponent<MetaDataTypes>(MetaDataIndices)...)
 		.ReadAllOf(Composites[CompositeIndices].ComponentTypeID.ReinterpretCast<CompositeTypes>()...)
 		.FilterAll({ Definition.PropertyType })
-		.SetStat(GET_STATID(MovieSceneEval_ApplyProperties))
+		.SetStat(Definition.StatID)
 		.SetDesiredThread(Linker->EntityManager.GetGatherThread())
 		.template Dispatch_PerAllocation<CompleteSetterTask>(&Linker->EntityManager, InPrerequisites, &Subsequents, Definition.CustomPropertyRegistration);
 
@@ -366,7 +363,7 @@
 			.FilterAny({ CompletePropertyMask })
 			.FilterAll({ Definition.PropertyType })
 			.FilterOut(CompletePropertyMask)
-			.SetStat(GET_STATID(MovieSceneEval_ApplyProperties))
+			.SetStat(Definition.StatID)
 			.SetDesiredThread(Linker->EntityManager.GetGatherThread())
 			.template Dispatch_PerAllocation<PartialSetterTask>(&Linker->EntityManager, InPrerequisites, &Subsequents, Definition.CustomPropertyRegistration, Composites);
 		}
@@ -607,6 +604,11 @@
 		return *this;
 	}
 
+	TPropertyDefinitionBuilder<PropertyTraits>& SetStat(TStatId InStatID)
+	{
+		Definition->StatID = InStatID;
+	}
+
 	template<typename BlenderSystemType>
 	TPropertyDefinitionBuilder<PropertyTraits>& SetBlenderSystem()
 	{
diff --git a/Engine/Source/Runtime/MovieScene/Public/EntitySystem/MovieScenePropertyRegistry.h b/Engine/Source/Runtime/MovieScene/Public/EntitySystem/MovieScenePropertyRegistry.h
--- a/Engine/Source/Runtime/MovieScene/Public/EntitySystem/MovieScenePropertyRegistry.h
+++ b/Engine/Source/Runtime/MovieScene/Public/EntitySystem/MovieScenePropertyRegistry.h
@@ -12,6 +12,7 @@
 
 #include "Misc/InlineValue.h"
 #include "Misc/TVariant.h"
+#include "Misc/GeneratedTypeName.h"
 #include <initializer_list>
 
 struct FMovieScenePropertyBinding;
@@ -29,6 +30,7 @@
 struct FFloatDecompositionParams;
 struct FPropertyCompositeDefinition;
 
+DECLARE_CYCLE_STAT(TEXT("Apply properties"), MovieSceneEval_ApplyProperties,  STATGROUP_MovieSceneECS);
 
 /**
  * Stats pertaining to a given type of property including how many properties exist in the linker,
@@ -83,6 +85,9 @@
 	/** Pointer to a custom getter/setter registry for short circuiting the UObject VM. Must outlive this definitions lifetime (usually these are static or singletons) */
 	ICustomPropertyRegistration* CustomPropertyRegistration = nullptr;
 
+	/** Stat ID for this property type */
+	TStatId StatID;
+
 	/** A mask of which composite indices pertain to floats */
 	uint32 FloatCompositeMask = 0;
 
@@ -179,9 +184,9 @@
 	 * @return A builder class that should be used to define the composites that contribute to this property
 	 */
 	template<typename PropertyTraits>
-	TCompositePropertyDefinitionBuilder<PropertyTraits> DefineCompositeProperty(TPropertyComponents<PropertyTraits>& InOutPropertyComponents)
+	TCompositePropertyDefinitionBuilder<PropertyTraits> DefineCompositeProperty(TPropertyComponents<PropertyTraits>& InOutPropertyComponents, const TCHAR* InStatName)
 	{
-		DefinePropertyImpl(InOutPropertyComponents);
+		DefinePropertyImpl(InOutPropertyComponents, InStatName);
 		FPropertyDefinition* Property = &Properties[InOutPropertyComponents.CompositeID.AsIndex()];
 		return TCompositePropertyDefinitionBuilder<PropertyTraits>(Property, this);
 	}
@@ -193,9 +198,9 @@
 	 * @return A builder class that should be used to define the composites that contribute to this property
 	 */
 	template<typename PropertyTraits>
-	TPropertyDefinitionBuilder<PropertyTraits> DefineProperty(TPropertyComponents<PropertyTraits>& InOutPropertyComponents)
+	TPropertyDefinitionBuilder<PropertyTraits> DefineProperty(TPropertyComponents<PropertyTraits>& InOutPropertyComponents, const TCHAR* InStatName)
 	{
-		DefinePropertyImpl(InOutPropertyComponents);
+		DefinePropertyImpl(InOutPropertyComponents, InStatName);
 		FPropertyDefinition* Property = &Properties[InOutPropertyComponents.CompositeID.AsIndex()];
 		return TPropertyDefinitionBuilder<PropertyTraits>(Property, this);
 	}
@@ -249,7 +254,7 @@
 	 * @return A builder class that should be used to define the composites that contribute to this property
 	 */
 	template<typename PropertyTraits>
-	void DefinePropertyImpl(TPropertyComponents<PropertyTraits>& InOutPropertyComponents)
+	void DefinePropertyImpl(TPropertyComponents<PropertyTraits>& InOutPropertyComponents, const TCHAR* InStatName)
 	{
 		using StorageType = typename PropertyTraits::StorageType;
 		static_assert( TIsBitwiseConstructible<StorageType, StorageType>::Value && TIsTriviallyDestructible<StorageType>::Value, "StorageType must be trivially TIsTriviallyCopyConstructible" );
@@ -257,12 +262,34 @@
 		const int32 CompositeOffset = CompositeDefinitions.Num();
 		checkf(CompositeOffset <= MAX_uint16, TEXT("Maximum number of composite definitions reached"));
 
+		TStatId StatID;
+
+#if STATS || ENABLE_STATNAMEDEVENTS
+
+	#if STATS
+		// Use FDynamicStats to create the stat in the right stat group if possible
+		StatID = FDynamicStats::CreateStatId<STAT_GROUP_TO_FStatGroup(STATGROUP_MovieSceneECS)>( FName(InStatName) );
+	#else
+		// Otherwise just make a named stat
+		const auto& ConversionData = StringCast<PROFILER_CHAR>(InStatName);
+		const int32 NumStorageChars = (ConversionData.Length() + 1);	//length doesn't include null terminator
+
+		// We leak this string
+		PROFILER_CHAR* StoragePtr = new PROFILER_CHAR[NumStorageChars];
+		FMemory::Memcpy(StoragePtr, ConversionData.Get(), NumStorageChars * sizeof(PROFILER_CHAR));
+
+		StatID = TStatId(StoragePtr);
+	#endif
+#endif
+
 		FPropertyDefinition NewDefinition(
 			CompositeOffset, 
 			sizeof(StorageType), alignof(StorageType),
 			InOutPropertyComponents.PropertyTag,
 			InOutPropertyComponents.InitialValue);
 
+		NewDefinition.StatID = StatID;
+
 		NewDefinition.MetaDataTypes = InOutPropertyComponents.MetaDataComponents.GetTypes();
 		checkf(!NewDefinition.MetaDataTypes.Contains(FComponentTypeID()), TEXT("Property meta-data component is not defined"));
 
diff --git a/Engine/Source/Runtime/MovieSceneTracks/Private/MovieSceneTracksComponentTypes.cpp b/Engine/Source/Runtime/MovieSceneTracks/Private/MovieSceneTracksComponentTypes.cpp
--- a/Engine/Source/Runtime/MovieSceneTracks/Private/MovieSceneTracksComponentTypes.cpp
+++ b/Engine/Source/Runtime/MovieSceneTracks/Private/MovieSceneTracksComponentTypes.cpp
@@ -568,14 +568,14 @@
 
 	// --------------------------------------------------------------------------------------------
 	// Set up bool properties
-	BuiltInComponents->PropertyRegistry.DefineProperty(Bool)
+	BuiltInComponents->PropertyRegistry.DefineProperty(Bool, TEXT("Apply bool Properties"))
 	.AddSoleChannel(BuiltInComponents->BoolResult)
 	.SetBlenderSystem<UMovieScenePiecewiseBoolBlenderSystem>()
 	.SetCustomAccessors(&Accessors.Bool)
 	.Commit();
 
 	// Set up FTransform properties
-	BuiltInComponents->PropertyRegistry.DefineCompositeProperty(Transform)
+	BuiltInComponents->PropertyRegistry.DefineCompositeProperty(Transform, TEXT("Apply FTransform Properties"))
 	.AddComposite(BuiltInComponents->DoubleResult[0], &FIntermediate3DTransform::T_X)
 	.AddComposite(BuiltInComponents->DoubleResult[1], &FIntermediate3DTransform::T_Y)
 	.AddComposite(BuiltInComponents->DoubleResult[2], &FIntermediate3DTransform::T_Z)
@@ -590,7 +590,7 @@
 
 	// --------------------------------------------------------------------------------------------
 	// Set up byte properties
-	BuiltInComponents->PropertyRegistry.DefineProperty(Byte)
+	BuiltInComponents->PropertyRegistry.DefineProperty(Byte, TEXT("Apply Byte Properties"))
 	.AddSoleChannel(BuiltInComponents->ByteResult)
 	.SetBlenderSystem<UMovieScenePiecewiseByteBlenderSystem>()
 	.SetCustomAccessors(&Accessors.Byte)
@@ -598,7 +598,7 @@
 
 	// --------------------------------------------------------------------------------------------
 	// Set up enum properties
-	BuiltInComponents->PropertyRegistry.DefineProperty(Enum)
+	BuiltInComponents->PropertyRegistry.DefineProperty(Enum, TEXT("Apply Enum Properties"))
 	.AddSoleChannel(BuiltInComponents->ByteResult)
 	.SetBlenderSystem<UMovieScenePiecewiseEnumBlenderSystem>()
 	.SetCustomAccessors(&Accessors.Enum)
@@ -606,7 +606,7 @@
 
 	// --------------------------------------------------------------------------------------------
 	// Set up integer properties
-	BuiltInComponents->PropertyRegistry.DefineProperty(Integer)
+	BuiltInComponents->PropertyRegistry.DefineProperty(Integer, TEXT("Apply Integer Properties"))
 	.AddSoleChannel(BuiltInComponents->IntegerResult)
 	.SetBlenderSystem<UMovieScenePiecewiseIntegerBlenderSystem>()
 	.SetCustomAccessors(&Accessors.Integer)
@@ -614,7 +614,7 @@
 
 	// --------------------------------------------------------------------------------------------
 	// Set up float properties
-	BuiltInComponents->PropertyRegistry.DefineProperty(Float)
+	BuiltInComponents->PropertyRegistry.DefineProperty(Float, TEXT("Apply float Properties"))
 	.AddSoleChannel(BuiltInComponents->FloatResult[0])
 	.SetBlenderSystem<UMovieScenePiecewiseFloatBlenderSystem>()
 	.SetCustomAccessors(&Accessors.Float)
@@ -622,7 +622,7 @@
 
 	// --------------------------------------------------------------------------------------------
 	// Set up double properties
-	BuiltInComponents->PropertyRegistry.DefineProperty(Double)
+	BuiltInComponents->PropertyRegistry.DefineProperty(Double, TEXT("Apply Double Properties"))
 	.AddSoleChannel(BuiltInComponents->DoubleResult[0])
 	.SetBlenderSystem<UMovieScenePiecewiseDoubleBlenderSystem>()
 	.SetCustomAccessors(&Accessors.Double)
@@ -630,7 +630,7 @@
 
 	// --------------------------------------------------------------------------------------------
 	// Set up color properties
-	BuiltInComponents->PropertyRegistry.DefineCompositeProperty(Color)
+	BuiltInComponents->PropertyRegistry.DefineCompositeProperty(Color, TEXT("Apply Color Properties"))
 	.AddComposite(BuiltInComponents->FloatResult[0], &FIntermediateColor::R)
 	.AddComposite(BuiltInComponents->FloatResult[1], &FIntermediateColor::G)
 	.AddComposite(BuiltInComponents->FloatResult[2], &FIntermediateColor::B)
@@ -662,7 +662,7 @@
 
 	// --------------------------------------------------------------------------------------------
 	// Set up vector properties
-	BuiltInComponents->PropertyRegistry.DefineCompositeProperty(FloatVector)
+	BuiltInComponents->PropertyRegistry.DefineCompositeProperty(FloatVector, TEXT("Apply FloatVector Properties"))
 	.AddComposite(BuiltInComponents->FloatResult[0], &FFloatIntermediateVector::X)
 	.AddComposite(BuiltInComponents->FloatResult[1], &FFloatIntermediateVector::Y)
 	.AddComposite(BuiltInComponents->FloatResult[2], &FFloatIntermediateVector::Z)
@@ -671,7 +671,7 @@
 	.SetCustomAccessors(&Accessors.FloatVector)
 	.Commit(FFloatVectorHandler());
 
-	BuiltInComponents->PropertyRegistry.DefineCompositeProperty(DoubleVector)
+	BuiltInComponents->PropertyRegistry.DefineCompositeProperty(DoubleVector, TEXT("Apply DoubleVector Properties"))
 	.AddComposite(BuiltInComponents->DoubleResult[0], &FDoubleIntermediateVector::X)
 	.AddComposite(BuiltInComponents->DoubleResult[1], &FDoubleIntermediateVector::Y)
 	.AddComposite(BuiltInComponents->DoubleResult[2], &FDoubleIntermediateVector::Z)
@@ -682,7 +682,7 @@
 
 	// --------------------------------------------------------------------------------------------
 	// Set up FEulerTransform properties
-	BuiltInComponents->PropertyRegistry.DefineCompositeProperty(EulerTransform)
+	BuiltInComponents->PropertyRegistry.DefineCompositeProperty(EulerTransform, TEXT("Apply FEulerTransform Properties"))
 	.AddComposite(BuiltInComponents->DoubleResult[0], &FIntermediate3DTransform::T_X)
 	.AddComposite(BuiltInComponents->DoubleResult[1], &FIntermediate3DTransform::T_Y)
 	.AddComposite(BuiltInComponents->DoubleResult[2], &FIntermediate3DTransform::T_Z)
@@ -700,7 +700,7 @@
 	{
 		Accessors.ComponentTransform.Add(USceneComponent::StaticClass(), "Transform", &GetComponentTransform, &SetComponentTransformAndVelocity);
 
-		BuiltInComponents->PropertyRegistry.DefineCompositeProperty(ComponentTransform)
+		BuiltInComponents->PropertyRegistry.DefineCompositeProperty(ComponentTransform, TEXT("Call USceneComponent::SetRelativeTransform"))
 		.AddComposite(BuiltInComponents->DoubleResult[0], &FIntermediate3DTransform::T_X)
 		.AddComposite(BuiltInComponents->DoubleResult[1], &FIntermediate3DTransform::T_Y)
 		.AddComposite(BuiltInComponents->DoubleResult[2], &FIntermediate3DTransform::T_Z)
diff --git a/Engine/Source/Runtime/UMG/Private/Animation/MovieSceneUMGComponentTypes.cpp b/Engine/Source/Runtime/UMG/Private/Animation/MovieSceneUMGComponentTypes.cpp
--- a/Engine/Source/Runtime/UMG/Private/Animation/MovieSceneUMGComponentTypes.cpp
+++ b/Engine/Source/Runtime/UMG/Private/Animation/MovieSceneUMGComponentTypes.cpp
@@ -83,7 +83,7 @@
 
 	CustomWidgetTransformAccessors.Add(UWidget::StaticClass(), "RenderTransform", &GetRenderTransform, &SetRenderTransform);
 
-	BuiltInComponents->PropertyRegistry.DefineCompositeProperty(Margin)
+	BuiltInComponents->PropertyRegistry.DefineCompositeProperty(Margin, TEXT("Apply FMargin Properties"))
 	.AddComposite(BuiltInComponents->FloatResult[0], &FMargin::Left)
 	.AddComposite(BuiltInComponents->FloatResult[1], &FMargin::Top)
 	.AddComposite(BuiltInComponents->FloatResult[2], &FMargin::Right)
@@ -91,7 +91,7 @@
 	.SetBlenderSystem<UMovieScenePiecewiseFloatBlenderSystem>()
 	.Commit();
 
-	BuiltInComponents->PropertyRegistry.DefineCompositeProperty(WidgetTransform)
+	BuiltInComponents->PropertyRegistry.DefineCompositeProperty(WidgetTransform, TEXT("Call UUserWidget::SetRenderTransform"))
 	.AddComposite(BuiltInComponents->FloatResult[0], &FIntermediateWidgetTransform::TranslationX)
 	.AddComposite(BuiltInComponents->FloatResult[1], &FIntermediateWidgetTransform::TranslationY)
 	.AddComposite(BuiltInComponents->FloatResult[2], &FIntermediateWidgetTransform::Rotation)
