PR #9243: Added missing implementation to UMovieScene3DConstraintTrack::RemoveAllAnimationData() (Contributed by [REDACTED_PERSON])

diff --git a/Engine/Source/Runtime/MovieSceneTracks/Private/Tracks/MovieScene3DConstraintTrack.cpp b/Engine/Source/Runtime/MovieSceneTracks/Private/Tracks/MovieScene3DConstraintTrack.cpp
--- a/Engine/Source/Runtime/MovieSceneTracks/Private/Tracks/MovieScene3DConstraintTrack.cpp
+++ b/Engine/Source/Runtime/MovieSceneTracks/Private/Tracks/MovieScene3DConstraintTrack.cpp
@@ -35,7 +35,7 @@ const TArray<UMovieSceneSection*>& UMovieScene3DConstraintTrack::GetAllSections(
 
 void UMovieScene3DConstraintTrack::RemoveAllAnimationData()
 {
-	// do nothing
+	ConstraintSections.Empty();
 }
 
 
