Sequencer: Fixed restore state for Actor properties controlled by [REDACTED_PERSON] that were previously null

diff --git a/Engine/Source/Runtime/MovieSceneTracks/Private/Evaluation/MovieSceneActorReferenceTemplate.cpp b/Engine/Source/Runtime/MovieSceneTracks/Private/Evaluation/MovieSceneActorReferenceTemplate.cpp
--- a/Engine/Source/Runtime/MovieSceneTracks/Private/Evaluation/MovieSceneActorReferenceTemplate.cpp
+++ b/Engine/Source/Runtime/MovieSceneTracks/Private/Evaluation/MovieSceneActorReferenceTemplate.cpp
@@ -35,11 +35,6 @@ namespace PropertyTemplate
 		return InWeakPtr.Get();
 	}
 
-	static bool IsValueValid(UObject* InValue)
-	{
-		return InValue != nullptr;
-	}
-
 	template<> IMovieScenePreAnimatedTokenPtr CacheExistingState<UObject*, FMovieSceneObjectBindingID>(UObject& Object, FTrackInstancePropertyBindings& PropertyBindings)
 	{
 		return TCachedState<UObject*, TWeakObjectPtr<>>(PropertyBindings.GetCurrentValue<UObject*>(Object), PropertyBindings);
