Sequencer: Expand support for blender systems to operate on non-property entities

diff --git a/Engine/Source/Runtime/MovieSceneTracks/Private/Systems/MovieScenePiecewiseBlenderSystemImpl.cpp b/Engine/Source/Runtime/MovieSceneTracks/Private/Systems/MovieScenePiecewiseBlenderSystemImpl.cpp
--- a/Engine/Source/Runtime/MovieSceneTracks/Private/Systems/MovieScenePiecewiseBlenderSystemImpl.cpp
+++ b/Engine/Source/Runtime/MovieSceneTracks/Private/Systems/MovieScenePiecewiseBlenderSystemImpl.cpp
@@ -205,16 +205,35 @@ private:
 	}
 };
 
+template<typename ValueType>
+void BlendResults(const TAccumulationResult<ValueType>& Results, uint16 BlendID, ValueType& OutFinalBlendResult)
+{
+	TBlendResult<ValueType> AbsoluteResult = Results.GetAbsoluteResult(BlendID);
+	TBlendResult<ValueType> AdditiveResult = Results.GetAdditiveResult(BlendID);
+	TBlendResult<ValueType> AdditiveFromBaseResult = Results.GetAdditiveFromBaseResult(BlendID);
+
+#if DO_GUARD_SLOW
+	ensureMsgf(AbsoluteResult.Weight != 0.f, TEXT("Default blend combine being used for an entity that has no absolute weight. This should have an initial value and should be handled by each system, and excluded by default with UMovieSceneBlenderSystem::FinalCombineExclusionFilter ."));
+#endif
+
+	const float TotalWeight = AbsoluteResult.Weight;
+	if (TotalWeight != 0)
+	{
+		const ValueType Value = AbsoluteResult.Total / AbsoluteResult.Weight + AdditiveResult.Total + AdditiveFromBaseResult.Total;
+		OutFinalBlendResult = Value;
+	}
+}
+
 /** Task that combines all accumulated blends for any tracked property type that has blend inputs/outputs */
 template<typename ValueType>
-struct TCombineBlends
+struct TCombineBlendsForProperties
 {
 	using FBlendResult = TBlendResult<ValueType>;
 	using FAccumulationResult = TAccumulationResult<ValueType>;
 	using FAccumulationBuffers = TAccumulationBuffers<ValueType>;
 	using FPiecewiseBlendableValueTraits = TPiecewiseBlendableValueTraits<ValueType>;
 
-	explicit TCombineBlends(const TBitArray<>& InCachedRelevantProperties, const FAccumulationBuffers* InAccumulationBuffers, FEntityAllocationWriteContext InWriteContext)
+	explicit TCombineBlendsForProperties(const TBitArray<>& InCachedRelevantProperties, const FAccumulationBuffers* InAccumulationBuffers, FEntityAllocationWriteContext InWriteContext)
 		: CachedRelevantProperties(InCachedRelevantProperties)
 		, AccumulationBuffers(InAccumulationBuffers)
 		, PropertyRegistry(&FBuiltInComponentTypes::Get()->PropertyRegistry)
@@ -223,6 +242,8 @@ struct TCombineBlends
 
 	void ForEachAllocation(FEntityAllocationIteratorItem InItem, TRead<FMovieSceneBlendChannelID> BlendIDs)
 	{
+		static const FMovieSceneBlenderSystemID BlenderSystemID = TPiecewiseBlendableValueTraits<ValueType>::GetBlenderSystemID();
+
 		FEntityAllocation* Allocation = InItem;
 		const FComponentMask& AllocationType = InItem;
 
@@ -332,29 +353,65 @@ private:
 		}
 	}
 
-	void BlendResults(const FAccumulationResult& Results, uint16 BlendID, ValueType& OutFinalBlendResult)
+private:
+
+	TBitArray<> CachedRelevantProperties;
+	const FAccumulationBuffers* AccumulationBuffers;
+	const FPropertyRegistry* PropertyRegistry;
+	FEntityAllocationWriteContext WriteContext;
+};
+
+
+/** Task that combines all accumulated blends for any tracked non-property type that has blend inputs/outputs */
+template<typename ValueType>
+struct TCombineBlends
+{
+	using FBlendResult = TBlendResult<ValueType>;
+	using FAccumulationResult = TAccumulationResult<ValueType>;
+	using FAccumulationBuffers = TAccumulationBuffers<ValueType>;
+	using FPiecewiseBlendableValueTraits = TPiecewiseBlendableValueTraits<ValueType>;
+
+	explicit TCombineBlends(const FAccumulationBuffers* InAccumulationBuffers, FEntityAllocationWriteContext InWriteContext)
+		: AccumulationBuffers(InAccumulationBuffers)
+		, WriteContext(InWriteContext)
+	{}
+
+	void ForEachAllocation(FEntityAllocationIteratorItem InItem, TRead<FMovieSceneBlendChannelID> BlendIDs)
 	{
-		FBlendResult AbsoluteResult = Results.GetAbsoluteResult(BlendID);
-		FBlendResult AdditiveResult = Results.GetAdditiveResult(BlendID);
-		FBlendResult AdditiveFromBaseResult = Results.GetAdditiveFromBaseResult(BlendID);
+		static const FMovieSceneBlenderSystemID BlenderSystemID = TPiecewiseBlendableValueTraits<ValueType>::GetBlenderSystemID();
 
-#if DO_GUARD_SLOW
-		ensureMsgf(AbsoluteResult.Weight != 0.f, TEXT("Default blend combine being used for an entity that has no absolute weight. This should have an initial value and should be handled by each system, and excluded by default with UMovieSceneBlenderSystem::FinalCombineExclusionFilter ."));
-#endif
+		FEntityAllocation* Allocation = InItem;
+		const FComponentMask& AllocationType = InItem;
 
-		const float TotalWeight = AbsoluteResult.Weight;
-		if (TotalWeight != 0)
+		// Find all result types on this allocation
+		TArrayView<TComponentTypeID<ValueType>> ResultComponents = FPiecewiseBlendableValueTraits::GetResultComponents();
+		for (int32 ResultIndex = 0; ResultIndex < ResultComponents.Num(); ++ResultIndex)
 		{
-			const ValueType Value = AbsoluteResult.Total / AbsoluteResult.Weight + AdditiveResult.Total + AdditiveFromBaseResult.Total;
-			OutFinalBlendResult = Value;
+			TComponentTypeID<ValueType> ResultComponent = ResultComponents[ResultIndex];
+			if (!AllocationType.Contains(ResultComponent))
+			{
+				continue;
+			}
+
+			FAccumulationResult Results = AccumulationBuffers->FindResults(ResultComponent);
+			if (!Results.IsValid())
+			{
+				continue;
+			}
+
+			// Open the result channel for write
+			TComponentWriter<ValueType> ValueResults = Allocation->WriteComponents(ResultComponent, WriteContext);
+			for (int32 Index = 0; Index < Allocation->Num(); ++Index)
+			{
+				ensureMsgf(BlendIDs[Index].SystemID == BlenderSystemID, TEXT("Overriding the standard blender system of standard types isn't supported."));
+				BlendResults(Results, BlendIDs[Index].ChannelID, ValueResults[Index]);
+			}
 		}
 	}
 
 private:
 
-	TBitArray<> CachedRelevantProperties;
 	const FAccumulationBuffers* AccumulationBuffers;
-	const FPropertyRegistry* PropertyRegistry;
 	FEntityAllocationWriteContext WriteContext;
 };
 
@@ -401,6 +458,7 @@ void TPiecewiseBlenderSystemImpl<ValueType>::Run(FPiecewiseBlenderSystemImplRunP
 {
 	using FAccumulationTask = TAccumulationTask<ValueType>;
 	using FAdditiveFromBaseBlendTask = TAdditiveFromBaseBlendTask<ValueType>;
+	using FCombineBlendsForProperties = TCombineBlendsForProperties<ValueType>;
 	using FCombineBlends = TCombineBlends<ValueType>;
 
 	// We allocate space for every blend even if there are gaps so we can do a straight index into each array
@@ -493,12 +551,26 @@ void TPiecewiseBlenderSystemImpl<ValueType>::Run(FPiecewiseBlenderSystemImplRunP
 		}
 	}
 
-	// Master task that performs the actual blends
-	FEntityTaskBuilder()
-	.Read(BuiltInComponents->BlendChannelOutput)
-	.FilterAny(BlendedPropertyMask)
-	.SetStat(Params.CombineBlendsStatId)
-	.template Dispatch_PerAllocation<FCombineBlends>(&EntityManager, Prereqs, &Subsequents, CachedRelevantProperties, &AccumulationBuffers, FEntityAllocationWriteContext(EntityManager));
+	// Combine blends for all blend outputs on properties
+	if (BlendedPropertyMask.Find(true) != INDEX_NONE)
+	{
+		FEntityTaskBuilder()
+		.Read(BuiltInComponents->BlendChannelOutput)
+		.FilterAny(BlendedPropertyMask)
+		.SetStat(Params.CombineBlendsStatId)
+		.template Dispatch_PerAllocation<FCombineBlendsForProperties>(&EntityManager, Prereqs, &Subsequents, CachedRelevantProperties, &AccumulationBuffers, FEntityAllocationWriteContext(EntityManager));
+	}
+
+	if (bContainsNonPropertyBlends)
+	{
+		// Blend task that combines vanilla (non-property-based) components
+		FEntityTaskBuilder()
+		.Read(BuiltInComponents->BlendChannelOutput)
+		.FilterAny(BlendedResultMask)
+		.FilterNone(BlendedPropertyMask)
+		.SetStat(Params.CombineBlendsStatId)
+		.template Dispatch_PerAllocation<FCombineBlends>(&EntityManager, Prereqs, &Subsequents, &AccumulationBuffers, FEntityAllocationWriteContext(EntityManager));
+	}
 }
 
 template<typename ValueType>
@@ -565,6 +637,8 @@ void TPiecewiseBlenderSystemImpl<ValueType>::ReinitializeAccumulationBuffers(int
 		return;
 	}
 
+	FComponentMask AllPropertyTypes;
+
 	// This code works on the assumption that properties can never be removed (which is safe)
 	FEntityComponentFilter InclusionFilter;
 	TArrayView<const FPropertyDefinition> Properties = BuiltInComponents->PropertyRegistry.GetProperties();
@@ -573,6 +647,8 @@ void TPiecewiseBlenderSystemImpl<ValueType>::ReinitializeAccumulationBuffers(int
 		const FPropertyDefinition& PropertyDefinition = Properties[PropertyTypeIndex];
 		if (FPiecewiseBlendableValueTraits::HasAnyComposite(PropertyDefinition))
 		{
+			AllPropertyTypes.Set(PropertyDefinition.PropertyType);
+
 			InclusionFilter.Reset();
 			InclusionFilter.All({ BuiltInComponents->BlendChannelOutput, PropertyDefinition.PropertyType });
 			if (EntityManager.Contains(InclusionFilter))
@@ -584,6 +660,8 @@ void TPiecewiseBlenderSystemImpl<ValueType>::ReinitializeAccumulationBuffers(int
 			}
 		}
 	}
+
+	bContainsNonPropertyBlends = EntityManager.Contains(FEntityComponentFilter().All({ BuiltInComponents->BlendChannelOutput }).None(AllPropertyTypes));
 }
 
 template<typename ValueType>
diff --git a/Engine/Source/Runtime/MovieSceneTracks/Public/Systems/MovieScenePiecewiseBlenderSystemImpl.h b/Engine/Source/Runtime/MovieSceneTracks/Public/Systems/MovieScenePiecewiseBlenderSystemImpl.h
--- a/Engine/Source/Runtime/MovieSceneTracks/Public/Systems/MovieScenePiecewiseBlenderSystemImpl.h
+++ b/Engine/Source/Runtime/MovieSceneTracks/Public/Systems/MovieScenePiecewiseBlenderSystemImpl.h
@@ -136,6 +136,9 @@ private:
 
 	/** Bit array specifying FCompositePropertyTypeID's for properties contained within BlendedPropertyMask */
 	TBitArray<> CachedRelevantProperties;
+
+	/** Whether the current entity manager contains any non-property based blends */
+	bool bContainsNonPropertyBlends = false;
 };
 
 } // namespace MovieScene
