Sequencer: Material Parameters, Switchers and Collection Parameters are now driven by the ecs runtime

A new BoundMaterial component exists that represents an animatable material type (either a UMaterialInstanceDynamic or a UMaterialParameterCollectionInstance). Material parameter tracks are able to apply scalar and color/vector parameters directly to these bound material components, and saves the system from having to understand anything about _how_ that animatable material was created or bound. Separate upstream systems for widget materials, primitive component materials, and material parameter collections all initialize these bound materials using their own mechanisms.

Component types exist for all parameter name types and the generic UMovieSceneParameterSection is able to import parameters for all types, although some are currently ineffectual since there are no consumers of these parameter types.

diff --git a/Engine/Source/Runtime/MovieSceneTracks/Private/Evaluation/MovieSceneMaterialParameterCollectionTemplate.cpp b/Engine/Source/Runtime/MovieSceneTracks/Private/Evaluation/MovieSceneMaterialParameterCollectionTemplate.cpp
deleted file mode 100644
diff --git a/Engine/Source/Runtime/MovieSceneTracks/Private/Evaluation/MovieSceneMaterialParameterCollectionTemplate.h b/Engine/Source/Runtime/MovieSceneTracks/Private/Evaluation/MovieSceneMaterialParameterCollectionTemplate.h
deleted file mode 100644
diff --git a/Engine/Source/Runtime/MovieSceneTracks/Private/Evaluation/MovieSceneParameterTemplate.cpp b/Engine/Source/Runtime/MovieSceneTracks/Private/Evaluation/MovieSceneParameterTemplate.cpp
--- a/Engine/Source/Runtime/MovieSceneTracks/Private/Evaluation/MovieSceneParameterTemplate.cpp
+++ b/Engine/Source/Runtime/MovieSceneTracks/Private/Evaluation/MovieSceneParameterTemplate.cpp
@@ -111,93 +111,3 @@ void FMovieSceneParameterSectionTemplate::EvaluateCurves(const FMovieSceneContex
 		}
 	}
 }
-
-void FDefaultMaterialAccessor::Apply(UMaterialInstanceDynamic& Material, const FEvaluatedParameterSectionValues& Values)
-{
-	for (const FScalarParameterNameAndValue& ScalarValue : Values.ScalarValues)
-	{
-		Material.SetScalarParameterValue(ScalarValue.ParameterName, ScalarValue.Value);
-	}
-	for (const FVectorParameterNameAndValue& VectorValue : Values.VectorValues)
-	{
-		Material.SetVectorParameterValue(VectorValue.ParameterName, VectorValue.Value);
-	}
-	for (const FColorParameterNameAndValue& ColorValue : Values.ColorValues)
-	{
-		Material.SetVectorParameterValue(ColorValue.ParameterName, ColorValue.Value);
-	}
-}
-
-
-TMovieSceneAnimTypeIDContainer<int32> MaterialIndexAnimTypeIDs;
-
-struct FComponentMaterialAccessor : FDefaultMaterialAccessor
-{
-	FComponentMaterialAccessor(int32 InMaterialIndex)
-		: MaterialIndex(InMaterialIndex)
-	{}
-
-	FComponentMaterialAccessor(const FComponentMaterialAccessor&) = default;
-	FComponentMaterialAccessor& operator=(const FComponentMaterialAccessor&) = default;
-	
-	FMovieSceneAnimTypeID GetAnimTypeID() const
-	{
-		return MaterialIndexAnimTypeIDs.GetAnimTypeID(MaterialIndex);
-	}
-
-	UMaterialInterface* GetMaterialForObject(UObject& Object) const
-	{
-		if (UPrimitiveComponent* PrimitiveComponent = Cast<UPrimitiveComponent>(&Object))
-		{
-			return PrimitiveComponent->GetMaterial(MaterialIndex);
-		}
-		else if (UDecalComponent* DecalComponent = Cast<UDecalComponent>(&Object))
-		{
-			return DecalComponent->GetDecalMaterial();
-		}
-		return nullptr;
-	}
-
-	void SetMaterialForObject(UObject& Object, UMaterialInterface& Material) const
-	{
-		if (UPrimitiveComponent* PrimitiveComponent = Cast<UPrimitiveComponent>(&Object))
-		{
-			PrimitiveComponent->SetMaterial(MaterialIndex, &Material);
-		}
-		else if (UDecalComponent* DecalComponent = Cast<UDecalComponent>(&Object))
-		{
-			DecalComponent->SetDecalMaterial(&Material);
-		}
-	}
-
-	UMaterialInstanceDynamic* CreateMaterialInstanceDynamic(UObject& Object, UMaterialInterface& Material, FName UniqueDynamicName)
-	{
-		if (UPrimitiveComponent* PrimitiveComponent = Cast<UPrimitiveComponent>(&Object))
-		{
-			return UMaterialInstanceDynamic::Create(&Material, &Object, UniqueDynamicName );
-		}
-		else if (UDecalComponent* DecalComponent = Cast<UDecalComponent>(&Object))
-		{
-			return DecalComponent->CreateDynamicMaterialInstance();
-		}
-
-		return nullptr;
-	}
-
-	int32 MaterialIndex;
-};
-
-FMovieSceneComponentMaterialSectionTemplate::FMovieSceneComponentMaterialSectionTemplate(const UMovieSceneParameterSection& Section, const UMovieSceneComponentMaterialTrack& Track)
-	: FMovieSceneParameterSectionTemplate(Section)
-	, MaterialIndex(Track.GetMaterialIndex())
-{
-}
-
-void FMovieSceneComponentMaterialSectionTemplate::Evaluate(const FMovieSceneEvaluationOperand& Operand, const FMovieSceneContext& Context, const FPersistentEvaluationData& PersistentData, FMovieSceneExecutionTokens& ExecutionTokens) const
-{
-	TMaterialTrackExecutionToken<FComponentMaterialAccessor> ExecutionToken(MaterialIndex);
-
-	EvaluateCurves(Context, ExecutionToken.Values);
-
-	ExecutionTokens.Add(MoveTemp(ExecutionToken));
-}
diff --git a/Engine/Source/Runtime/MovieSceneTracks/Private/Evaluation/MovieScenePrimitiveMaterialTemplate.cpp b/Engine/Source/Runtime/MovieSceneTracks/Private/Evaluation/MovieScenePrimitiveMaterialTemplate.cpp
deleted file mode 100644
diff --git a/Engine/Source/Runtime/MovieSceneTracks/Private/Evaluation/MovieScenePrimitiveMaterialTemplate.h b/Engine/Source/Runtime/MovieSceneTracks/Private/Evaluation/MovieScenePrimitiveMaterialTemplate.h
deleted file mode 100644
diff --git a/Engine/Source/Runtime/MovieSceneTracks/Private/MovieSceneTracksComponentTypes.cpp b/Engine/Source/Runtime/MovieSceneTracks/Private/MovieSceneTracksComponentTypes.cpp
--- a/Engine/Source/Runtime/MovieSceneTracks/Private/MovieSceneTracksComponentTypes.cpp
+++ b/Engine/Source/Runtime/MovieSceneTracks/Private/MovieSceneTracksComponentTypes.cpp
@@ -23,6 +23,7 @@
 #include "Systems/MovieSceneVectorPropertySystem.h"
 #include "MovieSceneObjectBindingID.h"
 #include "GameFramework/Actor.h"
+#include "Materials/MaterialParameterCollection.h"
 #include "Misc/App.h"
 
 namespace UE
@@ -564,6 +565,17 @@
 	ComponentRegistry->NewComponentType(&LevelVisibility, TEXT("Level Visibility"));
 	ComponentRegistry->NewComponentType(&DataLayer, TEXT("Data Layer"));
 
+	ComponentRegistry->NewComponentType(&ComponentMaterialIndex, TEXT("Component Material Index"), EComponentTypeFlags::CopyToChildren | EComponentTypeFlags::CopyToOutput);
+	ComponentRegistry->NewComponentType(&BoundMaterial,          TEXT("Bound Material"), EComponentTypeFlags::CopyToChildren | EComponentTypeFlags::CopyToOutput);
+	ComponentRegistry->NewComponentType(&MPC,                    TEXT("Material Parameter Collection"), EComponentTypeFlags::CopyToChildren | EComponentTypeFlags::CopyToOutput);
+
+	ComponentRegistry->NewComponentType(&BoolParameterName,      TEXT("Bool Parameter Name"), EComponentTypeFlags::CopyToChildren | EComponentTypeFlags::CopyToOutput);
+	ComponentRegistry->NewComponentType(&ScalarParameterName,    TEXT("Scalar Parameter Name"), EComponentTypeFlags::CopyToChildren | EComponentTypeFlags::CopyToOutput);
+	ComponentRegistry->NewComponentType(&Vector2DParameterName,  TEXT("Vector2D Parameter Name"), EComponentTypeFlags::CopyToChildren | EComponentTypeFlags::CopyToOutput);
+	ComponentRegistry->NewComponentType(&VectorParameterName,    TEXT("Vector Parameter Name"), EComponentTypeFlags::CopyToChildren | EComponentTypeFlags::CopyToOutput);
+	ComponentRegistry->NewComponentType(&ColorParameterName,     TEXT("Color Parameter Name"), EComponentTypeFlags::CopyToChildren | EComponentTypeFlags::CopyToOutput);
+	ComponentRegistry->NewComponentType(&TransformParameterName, TEXT("Transform Parameter Name"), EComponentTypeFlags::CopyToChildren | EComponentTypeFlags::CopyToOutput);
+
 	FBuiltInComponentTypes* BuiltInComponents = FBuiltInComponentTypes::Get();
 
 	// --------------------------------------------------------------------------------------------
diff --git a/Engine/Source/Runtime/MovieSceneTracks/Private/Sections/MovieSceneParameterSection.cpp b/Engine/Source/Runtime/MovieSceneTracks/Private/Sections/MovieSceneParameterSection.cpp
--- a/Engine/Source/Runtime/MovieSceneTracks/Private/Sections/MovieSceneParameterSection.cpp
+++ b/Engine/Source/Runtime/MovieSceneTracks/Private/Sections/MovieSceneParameterSection.cpp
@@ -4,6 +4,35 @@
 #include "UObject/SequencerObjectVersion.h"
 #include "Channels/MovieSceneChannelProxy.h"
 
+#include "Evaluation/MovieSceneEvaluationField.h"
+#include "EntitySystem/BuiltInComponentTypes.h"
+#include "MovieSceneTracksComponentTypes.h"
+
+
+namespace UE::MovieScene
+{
+
+/* Entity IDs are an encoded type and index, with the upper 8 bits being the type, and the lower 24 bits as the index */
+uint32 EncodeEntityID(int32 InIndex, uint8 InType)
+{
+	check(InIndex >= 0 && InIndex < int32(0x00FFFFFF));
+	return static_cast<uint32>(InIndex) | (uint32(InType) << 24);
+}
+void DecodeEntityID(uint32 InEntityID, int32& OutIndex, uint8& OutType)
+{
+	// Mask out the type to get the index
+	OutIndex = static_cast<int32>(InEntityID & 0x00FFFFFF);
+	OutType = InEntityID >> 24;
+}
+
+
+}// namespace UE::MovieScene
+
+void IMovieSceneParameterSectionExtender::ExtendEntity(UMovieSceneEntitySystemLinker* EntityLinker, const UE::MovieScene::FEntityImportParams& Params, UE::MovieScene::FImportedEntity* OutImportedEntity)
+{
+	ExtendEntityImpl(EntityLinker, Params, OutImportedEntity);
+}
+
 FScalarParameterNameAndCurve::FScalarParameterNameAndCurve( FName InParameterName )
 {
 	ParameterName = InParameterName;
@@ -567,4 +596,162 @@ void UMovieSceneParameterSection::GetParameterNames( TSet<FName>& ParameterNames
 	{
 		ParameterNames.Add(TransformParameterNamesAndCurve.ParameterName);
 	}
+}
+
+void UMovieSceneParameterSection::ImportEntityImpl(UMovieSceneEntitySystemLinker* EntityLinker, const FEntityImportParams& Params, FImportedEntity* OutImportedEntity)
+{
+	using namespace UE::MovieScene;
+
+	IMovieSceneParameterSectionExtender* Extender = GetImplementingOuter<IMovieSceneParameterSectionExtender>();
+	if (!ensureMsgf(Extender, TEXT("It is not valid for a UMovieSceneParameterSection to be used for importing entities outside of an outer chain that implements IMovieSceneParameterSectionExtender")))
+	{
+		return;
+	}
+
+	uint8 ParameterType = 0;
+	int32 EntityIndex = 0;
+	DecodeEntityID(Params.EntityID, EntityIndex, ParameterType);
+
+	FBuiltInComponentTypes* BuiltInComponentTypes = FBuiltInComponentTypes::Get();
+	FMovieSceneTracksComponentTypes* TracksComponentTypes = FMovieSceneTracksComponentTypes::Get();
+
+	FGuid ObjectBindingID = Params.GetObjectBindingID();
+
+	TEntityBuilder<TAddConditional<FGuid>> BaseBuilder = FEntityBuilder()
+		.AddConditional(BuiltInComponentTypes->GenericObjectBinding, ObjectBindingID, ObjectBindingID.IsValid());
+
+	switch (ParameterType)
+	{
+		case 0:
+		{
+			const FScalarParameterNameAndCurve& Scalar = ScalarParameterNamesAndCurves[EntityIndex];
+
+			OutImportedEntity->AddBuilder(
+				BaseBuilder
+				.Add(TracksComponentTypes->ScalarParameterName, Scalar.ParameterName)
+				.Add(BuiltInComponentTypes->FloatChannel[0], &Scalar.ParameterCurve)
+			);
+			break;
+		}
+		case 1:
+		{
+			const FBoolParameterNameAndCurve& Bool = BoolParameterNamesAndCurves[EntityIndex];
+
+			OutImportedEntity->AddBuilder(
+				BaseBuilder
+				.Add(TracksComponentTypes->BoolParameterName, Bool.ParameterName)
+				.Add(BuiltInComponentTypes->BoolChannel, &Bool.ParameterCurve)
+			);
+			break;
+		}
+		case 2:
+		{
+			const FVector2DParameterNameAndCurves& Vector2D = Vector2DParameterNamesAndCurves[EntityIndex];
+			OutImportedEntity->AddBuilder(
+				BaseBuilder
+				.Add(TracksComponentTypes->Vector2DParameterName, Vector2D.ParameterName)
+				.Add(BuiltInComponentTypes->FloatChannel[0], &Vector2D.XCurve)
+				.Add(BuiltInComponentTypes->FloatChannel[1], &Vector2D.YCurve)
+			);
+			break;
+		}
+		case 3:
+		{
+			const FVectorParameterNameAndCurves& Vector = VectorParameterNamesAndCurves[EntityIndex];
+			OutImportedEntity->AddBuilder(
+				BaseBuilder
+				.Add(TracksComponentTypes->VectorParameterName, Vector.ParameterName)
+				.Add(BuiltInComponentTypes->FloatChannel[0], &Vector.XCurve)
+				.Add(BuiltInComponentTypes->FloatChannel[1], &Vector.YCurve)
+				.Add(BuiltInComponentTypes->FloatChannel[2], &Vector.ZCurve)
+			);
+			break;
+		}
+		case 4:
+		{
+			const FColorParameterNameAndCurves& Color = ColorParameterNamesAndCurves[EntityIndex];
+			OutImportedEntity->AddBuilder(
+				BaseBuilder
+				.Add(TracksComponentTypes->ColorParameterName, Color.ParameterName)
+				.Add(BuiltInComponentTypes->FloatChannel[0], &Color.RedCurve)
+				.Add(BuiltInComponentTypes->FloatChannel[1], &Color.GreenCurve)
+				.Add(BuiltInComponentTypes->FloatChannel[2], &Color.BlueCurve)
+				.Add(BuiltInComponentTypes->FloatChannel[3], &Color.AlphaCurve)
+			);
+			break;
+		}
+		case 5:
+		{
+			const FTransformParameterNameAndCurves& Transform = TransformParameterNamesAndCurves[EntityIndex];
+			OutImportedEntity->AddBuilder(
+				BaseBuilder
+				.Add(TracksComponentTypes->TransformParameterName, Transform.ParameterName)
+				.Add(BuiltInComponentTypes->FloatChannel[0], &Transform.Translation[0])
+				.Add(BuiltInComponentTypes->FloatChannel[1], &Transform.Translation[1])
+				.Add(BuiltInComponentTypes->FloatChannel[2], &Transform.Translation[2])
+				.Add(BuiltInComponentTypes->FloatChannel[3], &Transform.Rotation[0])
+				.Add(BuiltInComponentTypes->FloatChannel[4], &Transform.Rotation[1])
+				.Add(BuiltInComponentTypes->FloatChannel[5], &Transform.Rotation[2])
+				.Add(BuiltInComponentTypes->FloatChannel[6], &Transform.Scale[0])
+				.Add(BuiltInComponentTypes->FloatChannel[7], &Transform.Scale[1])
+				.Add(BuiltInComponentTypes->FloatChannel[8], &Transform.Scale[2])
+			);
+			break;
+		}
+	}
+
+	Extender->ExtendEntity(EntityLinker, Params, OutImportedEntity);
+}
+
+bool UMovieSceneParameterSection::PopulateEvaluationFieldImpl(const TRange<FFrameNumber>& EffectiveRange, const FMovieSceneEvaluationFieldEntityMetaData& InMetaData, FMovieSceneEntityComponentFieldBuilder* OutFieldBuilder)
+{
+	// By default, parameter sections do not populate any evaluation field entries
+	// that is the job of its outer UMovieSceneTrack through a call to ExternalPopulateEvaluationField
+	return true;
+}
+
+void UMovieSceneParameterSection::ExternalPopulateEvaluationField(const TRange<FFrameNumber>& EffectiveRange, const FMovieSceneEvaluationFieldEntityMetaData& InMetaData, FMovieSceneEntityComponentFieldBuilder* OutFieldBuilder)
+{
+	using namespace UE::MovieScene;
+
+	const int32 MetaDataIndex = OutFieldBuilder->AddMetaData(InMetaData);
+
+	// We use the top 8 bits of EntityID to encode the type of parameter
+	const int32 NumScalarID    = ScalarParameterNamesAndCurves.Num();
+	const int32 NumBoolID      = BoolParameterNamesAndCurves.Num();
+	const int32 NumVector2DID  = Vector2DParameterNamesAndCurves.Num();
+	const int32 NumVectorID    = VectorParameterNamesAndCurves.Num();
+	const int32 NumColorID     = ColorParameterNamesAndCurves.Num();
+	const int32 NumTransformID = TransformParameterNamesAndCurves.Num();
+
+	for (int32 Index = 0; Index < NumScalarID; ++Index)
+	{
+		const int32 EntityIndex = OutFieldBuilder->FindOrAddEntity(this, EncodeEntityID(Index, 0));
+		OutFieldBuilder->AddPersistentEntity(EffectiveRange, EntityIndex, MetaDataIndex);
+	}
+	for (int32 Index = 0; Index < NumBoolID; ++Index)
+	{
+		const int32 EntityIndex = OutFieldBuilder->FindOrAddEntity(this, EncodeEntityID(Index, 1));
+		OutFieldBuilder->AddPersistentEntity(EffectiveRange, EntityIndex, MetaDataIndex);
+	}
+	for (int32 Index = 0; Index < NumVector2DID; ++Index)
+	{
+		const int32 EntityIndex = OutFieldBuilder->FindOrAddEntity(this, EncodeEntityID(Index, 2));
+		OutFieldBuilder->AddPersistentEntity(EffectiveRange, EntityIndex, MetaDataIndex);
+	}
+	for (int32 Index = 0; Index < NumVectorID; ++Index)
+	{
+		const int32 EntityIndex = OutFieldBuilder->FindOrAddEntity(this, EncodeEntityID(Index, 3));
+		OutFieldBuilder->AddPersistentEntity(EffectiveRange, EntityIndex, MetaDataIndex);
+	}
+	for (int32 Index = 0; Index < NumColorID; ++Index)
+	{
+		const int32 EntityIndex = OutFieldBuilder->FindOrAddEntity(this, EncodeEntityID(Index, 4));
+		OutFieldBuilder->AddPersistentEntity(EffectiveRange, EntityIndex, MetaDataIndex);
+	}
+	for (int32 Index = 0; Index < NumTransformID; ++Index)
+	{
+		const int32 EntityIndex = OutFieldBuilder->FindOrAddEntity(this, EncodeEntityID(Index, 5));
+		OutFieldBuilder->AddPersistentEntity(EffectiveRange, EntityIndex, MetaDataIndex);
+	}
 }
\ No newline at end of file
diff --git a/Engine/Source/Runtime/MovieSceneTracks/Private/Sections/MovieScenePrimitiveMaterialSection.cpp b/Engine/Source/Runtime/MovieSceneTracks/Private/Sections/MovieScenePrimitiveMaterialSection.cpp
--- a/Engine/Source/Runtime/MovieSceneTracks/Private/Sections/MovieScenePrimitiveMaterialSection.cpp
+++ b/Engine/Source/Runtime/MovieSceneTracks/Private/Sections/MovieScenePrimitiveMaterialSection.cpp
@@ -1,8 +1,12 @@
 // Copyright Epic Games, Inc. All Rights Reserved.
 
 #include "Sections/MovieScenePrimitiveMaterialSection.h"
+#include "Tracks/MovieScenePrimitiveMaterialTrack.h"
 #include "Channels/MovieSceneChannelProxy.h"
 #include "Materials/MaterialInterface.h"
+#include "EntitySystem/MovieSceneEntityBuilder.h"
+#include "EntitySystem/BuiltInComponentTypes.h"
+#include "MovieSceneTracksComponentTypes.h"
 
 
 UMovieScenePrimitiveMaterialSection::UMovieScenePrimitiveMaterialSection(const FObjectInitializer& ObjInit)
@@ -19,4 +23,26 @@ UMovieScenePrimitiveMaterialSection::UMovieScenePrimitiveMaterialSection(const F
 #else
 	ChannelProxy = MakeShared<FMovieSceneChannelProxy>(MaterialChannel);
 #endif
+}
+
+void UMovieScenePrimitiveMaterialSection::ImportEntityImpl(UMovieSceneEntitySystemLinker* Linker, const FEntityImportParams& ImportParams, FImportedEntity* OutImportedEntity)
+{
+	using namespace UE::MovieScene;
+
+	UMovieScenePrimitiveMaterialTrack* Track = GetTypedOuter<UMovieScenePrimitiveMaterialTrack>();
+	check(Track);
+
+	const int32 MaterialIndex = Track->GetMaterialIndex();
+
+	FBuiltInComponentTypes* BuiltInComponentTypes = FBuiltInComponentTypes::Get();
+	FMovieSceneTracksComponentTypes* TracksComponentTypes = FMovieSceneTracksComponentTypes::Get();
+
+	FGuid ObjectBindingID = ImportParams.GetObjectBindingID();
+
+	OutImportedEntity->AddBuilder(
+		FEntityBuilder()
+		.AddConditional(BuiltInComponentTypes->GenericObjectBinding, ObjectBindingID, ObjectBindingID.IsValid())
+		.Add(BuiltInComponentTypes->ObjectPathChannel, &MaterialChannel)
+		.Add(TracksComponentTypes->ComponentMaterialIndex, MaterialIndex)
+	);
 }
\ No newline at end of file
diff --git a/Engine/Source/Runtime/MovieSceneTracks/Private/Systems/MovieSceneComponentMaterialSystem.cpp b/Engine/Source/Runtime/MovieSceneTracks/Private/Systems/MovieSceneComponentMaterialSystem.cpp
new file mode 100644
--- /dev/null
+++ b/Engine/Source/Runtime/MovieSceneTracks/Private/Systems/MovieSceneComponentMaterialSystem.cpp
@@ -0,0 +1,126 @@
+// Copyright Epic Games, Inc. All Rights Reserved.
+
+#include "Systems/MovieSceneComponentMaterialSystem.h"
+
+#include "EntitySystem/MovieSceneEntityMutations.h"
+
+#include "Systems/FloatChannelEvaluatorSystem.h"
+#include "Systems/MovieScenePiecewiseFloatBlenderSystem.h"
+
+#include "Materials/MaterialInstanceDynamic.h"
+
+#include "Components/PrimitiveComponent.h"
+#include "Components/DecalComponent.h"
+
+namespace UE::MovieScene
+{
+
+FComponentMaterialAccessor::FComponentMaterialAccessor(const FComponentMaterialKey& InKey)
+	: Object(InKey.Object.ResolveObjectPtr())
+	, MaterialIndex(InKey.MaterialIndex)
+{}
+
+FComponentMaterialAccessor::FComponentMaterialAccessor(UObject* InObject, int32 InMaterialIndex)
+	: Object(InObject)
+	, MaterialIndex(InMaterialIndex)
+{}
+
+FString FComponentMaterialAccessor::ToString() const
+{
+	return FString::Printf(TEXT("object %s (element %d)"), *Object->GetPathName(), MaterialIndex);
+}
+
+UMaterialInterface* FComponentMaterialAccessor::GetMaterial() const
+{
+	if (UPrimitiveComponent* PrimitiveComponent = Cast<UPrimitiveComponent>(Object))
+	{
+		return PrimitiveComponent->GetMaterial(MaterialIndex);
+	}
+	else if (UDecalComponent* DecalComponent = Cast<UDecalComponent>(Object))
+	{
+		return DecalComponent->GetDecalMaterial();
+	}
+	return nullptr;
+}
+
+void FComponentMaterialAccessor::SetMaterial(UMaterialInterface* InMaterial) const
+{
+	if (UPrimitiveComponent* PrimitiveComponent = Cast<UPrimitiveComponent>(Object))
+	{
+		PrimitiveComponent->SetMaterial(MaterialIndex, InMaterial);
+	}
+	else if (UDecalComponent* DecalComponent = Cast<UDecalComponent>(Object))
+	{
+		DecalComponent->SetDecalMaterial(InMaterial);
+	}
+}
+
+UMaterialInstanceDynamic* FComponentMaterialAccessor::CreateDynamicMaterial(UMaterialInterface* InMaterial)
+{
+	// Need to create a new MID, either because the parent has changed, or because one doesn't already exist
+	if (UPrimitiveComponent* PrimitiveComponent = Cast<UPrimitiveComponent>(Object))
+	{
+		TStringBuilder<128> DynamicName;
+		InMaterial->GetFName().ToString(DynamicName);
+		DynamicName.Append(TEXT("_Animated"));
+		FName UniqueDynamicName = MakeUniqueObjectName(Object, UMaterialInstanceDynamic::StaticClass() , DynamicName.ToString());
+
+		UMaterialInstanceDynamic* Result = UMaterialInstanceDynamic::Create(InMaterial, Object, UniqueDynamicName );
+		PrimitiveComponent->SetMaterial(MaterialIndex, Result);
+		return Result;
+	}
+	else if (UDecalComponent* DecalComponent = Cast<UDecalComponent>(Object))
+	{
+		return DecalComponent->CreateDynamicMaterialInstance();
+	}
+	return nullptr;
+}
+
+TAutoRegisterPreAnimatedStorageID<FPreAnimatedComponentMaterialSwitcherStorage> FPreAnimatedComponentMaterialSwitcherStorage::StorageID;
+TAutoRegisterPreAnimatedStorageID<FPreAnimatedComponentMaterialParameterStorage> FPreAnimatedComponentMaterialParameterStorage::StorageID;
+
+} // namespace UE::MovieScene
+
+UMovieSceneComponentMaterialSystem::UMovieSceneComponentMaterialSystem(const FObjectInitializer& ObjInit)
+	: Super(ObjInit)
+{
+	using namespace UE::MovieScene;
+
+	FBuiltInComponentTypes* BuiltInComponents = FBuiltInComponentTypes::Get();
+	FMovieSceneTracksComponentTypes* TracksComponents = FMovieSceneTracksComponentTypes::Get();
+
+	RelevantComponent = TracksComponents->ComponentMaterialIndex;
+	Phase = ESystemPhase::Instantiation | ESystemPhase::Evaluation;
+
+	if (HasAnyFlags(RF_ClassDefaultObject))
+	{
+		DefineComponentConsumer(GetClass(), BuiltInComponents->ObjectResult);
+		DefineComponentConsumer(GetClass(), BuiltInComponents->BoundObject);
+		DefineComponentProducer(GetClass(), TracksComponents->BoundMaterial);
+		DefineImplicitPrerequisite(UMovieSceneCachePreAnimatedStateSystem::StaticClass(), GetClass());
+	}
+}
+
+void UMovieSceneComponentMaterialSystem::OnLink()
+{
+	using namespace UE::MovieScene;
+
+	SystemImpl.MaterialSwitcherStorage = Linker->PreAnimatedState.GetOrCreateStorage<FPreAnimatedComponentMaterialSwitcherStorage>();
+	SystemImpl.MaterialParameterStorage = Linker->PreAnimatedState.GetOrCreateStorage<FPreAnimatedComponentMaterialParameterStorage>();
+
+	SystemImpl.OnLink(Linker, FBuiltInComponentTypes::Get()->BoundObject, FMovieSceneTracksComponentTypes::Get()->ComponentMaterialIndex);
+}
+
+void UMovieSceneComponentMaterialSystem::OnRun(FSystemTaskPrerequisites& InPrerequisites, FSystemSubsequentTasks& Subsequents)
+{
+	using namespace UE::MovieScene;
+
+	SystemImpl.OnRun(Linker, FBuiltInComponentTypes::Get()->BoundObject, FMovieSceneTracksComponentTypes::Get()->ComponentMaterialIndex, InPrerequisites, Subsequents);
+}
+
+void UMovieSceneComponentMaterialSystem::SavePreAnimatedState(const FPreAnimationParameters& InParameters)
+{
+	using namespace UE::MovieScene;
+
+	SystemImpl.SavePreAnimatedState(Linker, FBuiltInComponentTypes::Get()->BoundObject, FMovieSceneTracksComponentTypes::Get()->ComponentMaterialIndex, InParameters);
+}
\ No newline at end of file
diff --git a/Engine/Source/Runtime/MovieSceneTracks/Private/Systems/MovieSceneMaterialParameterCollectionSystem.cpp b/Engine/Source/Runtime/MovieSceneTracks/Private/Systems/MovieSceneMaterialParameterCollectionSystem.cpp
new file mode 100644
--- /dev/null
+++ b/Engine/Source/Runtime/MovieSceneTracks/Private/Systems/MovieSceneMaterialParameterCollectionSystem.cpp
@@ -0,0 +1,303 @@
+// Copyright Epic Games, Inc. All Rights Reserved.
+
+#include "Systems/MovieSceneMaterialParameterCollectionSystem.h"
+
+#include "IMovieScenePlayer.h"
+#include "MovieSceneSequence.h"
+#include "Evaluation/MovieSceneEvaluationTemplateInstance.h"
+#include "EntitySystem/BuiltInComponentTypes.h"
+#include "EntitySystem/MovieSceneEntityMutations.h"
+
+#include "MovieSceneTracksComponentTypes.h"
+#include "Systems/FloatChannelEvaluatorSystem.h"
+#include "Systems/MovieScenePiecewiseFloatBlenderSystem.h"
+
+#include "Logging/MessageLog.h"
+#include "Misc/UObjectToken.h"
+
+#include "Materials/MaterialParameterCollection.h"
+#include "Materials/MaterialParameterCollectionInstance.h"
+
+#include "Components/PrimitiveComponent.h"
+#include "Components/DecalComponent.h"
+
+namespace UE::MovieScene
+{
+
+
+struct FMaterialParameterCollectionKey
+{
+	TObjectKey<UMaterialParameterCollectionInstance> MPCI;
+	FName ParameterName;
+
+	FMaterialParameterCollectionKey(UObject* InBoundMaterial, const FName& InParameterName)
+		: MPCI(CastChecked<UMaterialParameterCollectionInstance>(InBoundMaterial))
+		, ParameterName(InParameterName)
+	{}
+
+	friend uint32 GetTypeHash(const FMaterialParameterCollectionKey& InKey)
+	{
+		return GetTypeHash(InKey.MPCI) ^ GetTypeHash(InKey.ParameterName);
+	}
+
+	friend bool operator==(const FMaterialParameterCollectionKey& A, const FMaterialParameterCollectionKey& B)
+	{
+		return A.MPCI == B.MPCI && A.ParameterName == B.ParameterName;
+	}
+};
+
+struct FMaterialParameterCollectionScalarTraits : FBoundObjectPreAnimatedStateTraits
+{
+	using KeyType = FMaterialParameterCollectionKey;
+	using StorageType = float;
+
+	static void ReplaceObject(FMaterialParameterCollectionKey& InOutKey, const FObjectKey& NewObject)
+	{
+		InOutKey.MPCI = Cast<UMaterialParameterCollectionInstance>(NewObject.ResolveObjectPtr());
+	}
+
+	static float CachePreAnimatedValue(UObject* InBoundMaterial, const FName& ParameterName)
+	{
+		UMaterialParameterCollectionInstance* MPCI = CastChecked<UMaterialParameterCollectionInstance>(InBoundMaterial);
+
+		float ParameterValue = 0.f;
+		MPCI->GetScalarParameterValue(ParameterName, ParameterValue);
+		return ParameterValue;
+	}
+
+	static void RestorePreAnimatedValue(const FMaterialParameterCollectionKey& InKey, float OldValue, const FRestoreStateParams& Params)
+	{
+		UMaterialParameterCollectionInstance* MPCI = InKey.MPCI.ResolveObjectPtr();
+		if (MPCI)
+		{
+			MPCI->SetScalarParameterValue(InKey.ParameterName, OldValue);
+		}
+	}
+};
+
+struct FMaterialParameterCollectionVectorTraits : FBoundObjectPreAnimatedStateTraits
+{
+	using KeyType = FMaterialParameterCollectionKey;
+	using StorageType = FLinearColor;
+
+	static void ReplaceObject(FMaterialParameterCollectionKey& InOutKey, const FObjectKey& NewObject)
+	{
+		InOutKey.MPCI = Cast<UMaterialParameterCollectionInstance>(NewObject.ResolveObjectPtr());
+	}
+
+	static FLinearColor CachePreAnimatedValue(UObject* InBoundMaterial, const FName& ParameterName)
+	{
+		UMaterialParameterCollectionInstance* MPCI = CastChecked<UMaterialParameterCollectionInstance>(InBoundMaterial);
+		FLinearColor ParameterValue = FLinearColor::White;
+		MPCI->GetVectorParameterValue(ParameterName, ParameterValue);
+		return ParameterValue;
+	}
+
+	static void RestorePreAnimatedValue(const FMaterialParameterCollectionKey& InKey, const FLinearColor& OldValue, const FRestoreStateParams& Params)
+	{
+		UMaterialParameterCollectionInstance* MPCI = InKey.MPCI.ResolveObjectPtr();
+		if (MPCI)
+		{
+			MPCI->SetVectorParameterValue(InKey.ParameterName, OldValue);
+		}
+	}
+};
+
+struct FPreAnimatedMPCScalarStorage
+	: public TPreAnimatedStateStorage<FMaterialParameterCollectionScalarTraits>
+{
+	static TAutoRegisterPreAnimatedStorageID<FPreAnimatedMPCScalarStorage> StorageID;
+};
+
+struct FPreAnimatedMPCVectorStorage
+	: public TPreAnimatedStateStorage<FMaterialParameterCollectionVectorTraits>
+{
+	static TAutoRegisterPreAnimatedStorageID<FPreAnimatedMPCVectorStorage> StorageID;
+};
+
+
+TAutoRegisterPreAnimatedStorageID<FPreAnimatedMPCScalarStorage> FPreAnimatedMPCScalarStorage::StorageID;
+TAutoRegisterPreAnimatedStorageID<FPreAnimatedMPCVectorStorage> FPreAnimatedMPCVectorStorage::StorageID;
+
+} // namespace UE::MovieScene
+
+UMovieSceneMaterialParameterCollectionSystem::UMovieSceneMaterialParameterCollectionSystem(const FObjectInitializer& ObjInit)
+	: Super(ObjInit)
+{
+	using namespace UE::MovieScene;
+
+	FBuiltInComponentTypes* BuiltInComponents = FBuiltInComponentTypes::Get();
+	FMovieSceneTracksComponentTypes* TracksComponents = FMovieSceneTracksComponentTypes::Get();
+
+	RelevantComponent = TracksComponents->MPC;
+	Phase = ESystemPhase::Instantiation;
+
+	if (HasAnyFlags(RF_ClassDefaultObject))
+	{
+		DefineComponentProducer(GetClass(), TracksComponents->BoundMaterial);
+		DefineImplicitPrerequisite(UMovieSceneCachePreAnimatedStateSystem::StaticClass(), GetClass());
+	}
+}
+
+void UMovieSceneMaterialParameterCollectionSystem::OnLink()
+{
+	using namespace UE::MovieScene;
+
+	ScalarParameterStorage = Linker->PreAnimatedState.GetOrCreateStorage<FPreAnimatedMPCScalarStorage>();
+	VectorParameterStorage = Linker->PreAnimatedState.GetOrCreateStorage<FPreAnimatedMPCVectorStorage>();
+}
+
+void UMovieSceneMaterialParameterCollectionSystem::OnRun(FSystemTaskPrerequisites& InPrerequisites, FSystemSubsequentTasks& Subsequents)
+{
+	using namespace UE::MovieScene;
+
+	FBuiltInComponentTypes*          BuiltInComponents = FBuiltInComponentTypes::Get();
+	FMovieSceneTracksComponentTypes* TracksComponents  = FMovieSceneTracksComponentTypes::Get();
+
+	struct FAddMPCMutation : IMovieSceneEntityMutation
+	{
+		FAddMPCMutation(UMovieSceneEntitySystemLinker* InLinker)
+			: InstanceRegistry(InLinker->GetInstanceRegistry())
+		{
+			BuiltInComponents = FBuiltInComponentTypes::Get();
+			TracksComponents  = FMovieSceneTracksComponentTypes::Get();
+		}
+		virtual void CreateMutation(FEntityManager* EntityManager, FComponentMask* InOutEntityComponentTypes) const override
+		{
+			InOutEntityComponentTypes->Set(TracksComponents->BoundMaterial);
+		}
+		virtual void InitializeAllocation(FEntityAllocation* Allocation, const FComponentMask& AllocationType) const
+		{
+			TComponentWriter<UObject*> OutBoundMaterials = Allocation->WriteComponents(TracksComponents->BoundMaterial, FEntityAllocationWriteContext::NewAllocation());
+			TComponentReader<UMaterialParameterCollection*> MPCs = Allocation->ReadComponents(TracksComponents->MPC);
+			TComponentReader<FInstanceHandle> InstanceHandles = Allocation->ReadComponents(BuiltInComponents->InstanceHandle);
+
+			TOptionalComponentReader<FName> ScalarParameterNames = Allocation->TryReadComponents(TracksComponents->ScalarParameterName);
+			TOptionalComponentReader<FName> VectorParameterNames = Allocation->TryReadComponents(TracksComponents->VectorParameterName);
+			TOptionalComponentReader<FName> ColorParameterNames = Allocation->TryReadComponents(TracksComponents->ColorParameterName);
+
+			const int32 Num = Allocation->Num();
+			for (int32 Index = 0; Index < Num; ++Index)
+			{
+				OutBoundMaterials[Index] = nullptr;
+
+				UMaterialParameterCollection* Collection = MPCs[Index];
+				IMovieScenePlayer* Player = InstanceRegistry->GetInstance(InstanceHandles[Index]).GetPlayer();
+				UObject* WorldContextObject = Player->GetPlaybackContext();
+				UWorld* World = GEngine->GetWorldFromContextObject(WorldContextObject, EGetWorldErrorMode::LogAndReturnNull);
+				if (World && Collection)
+				{
+					UMaterialParameterCollectionInstance* Instance = World->GetParameterCollectionInstance(Collection);
+					OutBoundMaterials[Index] = Instance;
+
+					if (!ensureAlwaysMsgf(Instance != nullptr,
+						TEXT("Unable to create MPC instance for %s with World %s. Material parameter collection tracks will not function."),
+						*Collection->GetName(), *World->GetName()))
+					{
+						// Remove the entity if it is null to avoid having to check for null downstream
+						InvalidEntities.Add(Allocation->GetRawEntityIDs()[Index]);
+					}
+					else 
+					{
+						if (ScalarParameterNames)
+						{
+							FName Name = ScalarParameterNames[Index];
+							if (Collection->GetScalarParameterByName(Name) == nullptr)
+							{
+								InvalidEntities.Add(Allocation->GetRawEntityIDs()[Index]);
+
+#if !UE_BUILD_SHIPPING && !UE_BUILD_TEST
+								if (!Instance->bLoggedMissingParameterWarning)
+								{
+									MissingParameters.FindOrAdd(MakeTuple(Instance, Player)).Add(Name.ToString());
+								}
+#endif
+							}
+						}
+						else 
+						{
+							FName Name = VectorParameterNames ? VectorParameterNames[Index] : ColorParameterNames[Index];
+							if (Collection->GetVectorParameterByName(Name) == nullptr)
+							{
+								InvalidEntities.Add(Allocation->GetRawEntityIDs()[Index]);
+
+#if !UE_BUILD_SHIPPING && !UE_BUILD_TEST
+								if (!Instance->bLoggedMissingParameterWarning)
+								{
+									MissingParameters.FindOrAdd(MakeTuple(Instance, Player)).Add(Name.ToString());
+								}
+#endif
+							}
+						}
+					}
+				}
+			}
+		}
+
+		void Cleanup(UMovieSceneEntitySystemLinker* InLinker)
+		{
+			for (FMovieSceneEntityID InvalidEntity : InvalidEntities)
+			{
+				InLinker->EntityManager.RemoveComponent(InvalidEntity, TracksComponents->BoundMaterial);
+			}
+
+#if !UE_BUILD_SHIPPING && !UE_BUILD_TEST
+
+			for (TPair<TTuple<UMaterialParameterCollectionInstance*, IMovieScenePlayer*>, TArray<FString>>& Pair : MissingParameters)
+			{
+				FFormatNamedArguments Arguments;
+				Arguments.Add(TEXT("ParamNames"), FText::FromString(FString::Join(Pair.Value, TEXT(", "))));
+				FMessageLog("PIE").Warning()
+					->AddToken(FTextToken::Create(NSLOCTEXT("MaterialParameterCollectionTrack", "InvalidParameterText", "Invalid parameter name or type applied in sequence")))
+					->AddToken(FUObjectToken::Create(Pair.Key.Get<1>()->GetEvaluationTemplate().GetSequence(MovieSceneSequenceID::Root)))
+					->AddToken(FTextToken::Create(NSLOCTEXT("MaterialParameterCollectionTrack", "OnText", "on")))
+					->AddToken(FUObjectToken::Create(Pair.Key.Get<0>()))
+					->AddToken(FTextToken::Create(FText::Format(NSLOCTEXT("MaterialParameterCollectionTrack", "InvalidParameterFormatText", "with the following invalid parameters: {ParamNames}."), Arguments)));
+				Pair.Key.Get<0>()->bLoggedMissingParameterWarning = true;
+			}
+
+#endif
+		}
+	private:
+
+		FInstanceRegistry* InstanceRegistry;
+		FBuiltInComponentTypes* BuiltInComponents;
+		FMovieSceneTracksComponentTypes* TracksComponents;
+
+		mutable TMap<TTuple<UMaterialParameterCollectionInstance*, IMovieScenePlayer*>, TArray<FString>> MissingParameters;
+
+		/** Track entities that we added a BoundMaterial component to but were unable to resolve so we can remove the component. This should happen rarely */
+		mutable TArray<FMovieSceneEntityID> InvalidEntities;
+	};
+
+
+	// Only mutate things that are tagged as requiring linking
+	FEntityComponentFilter Filter;
+	Filter.All({TracksComponents->MPC, BuiltInComponents->InstanceHandle, BuiltInComponents->Tags.NeedsLink});
+
+	// Initialize bound dynamic materials (for material collection parameters)
+	FAddMPCMutation BindMaterialsMutation(Linker);
+	Linker->EntityManager.MutateAll(Filter, BindMaterialsMutation);
+	BindMaterialsMutation.Cleanup(Linker);
+}
+
+
+void UMovieSceneMaterialParameterCollectionSystem::SavePreAnimatedState(const FPreAnimationParameters& InParameters)
+{
+	using namespace UE::MovieScene;
+
+	FBuiltInComponentTypes*          BuiltInComponents = FBuiltInComponentTypes::Get();
+	FMovieSceneTracksComponentTypes* TracksComponents  = FMovieSceneTracksComponentTypes::Get();
+
+	TSavePreAnimatedStateParams<UObject*, FName> Params;
+
+	Params.AdditionalFilter.All({ TracksComponents->MPC });
+	ScalarParameterStorage->BeginTrackingEntitiesTask(Linker, Params, TracksComponents->BoundMaterial, TracksComponents->ScalarParameterName);
+	ScalarParameterStorage->CachePreAnimatedValuesTask(Linker, Params, TracksComponents->BoundMaterial, TracksComponents->ScalarParameterName);
+
+	VectorParameterStorage->BeginTrackingEntitiesTask(Linker, Params, TracksComponents->BoundMaterial, TracksComponents->VectorParameterName);
+	VectorParameterStorage->BeginTrackingEntitiesTask(Linker, Params, TracksComponents->BoundMaterial, TracksComponents->ColorParameterName);
+	VectorParameterStorage->CachePreAnimatedValuesTask(Linker, Params, TracksComponents->BoundMaterial, TracksComponents->VectorParameterName);
+	VectorParameterStorage->CachePreAnimatedValuesTask(Linker, Params, TracksComponents->BoundMaterial, TracksComponents->ColorParameterName);
+}
+
diff --git a/Engine/Source/Runtime/MovieSceneTracks/Private/Systems/MovieSceneMaterialParameterSystem.cpp b/Engine/Source/Runtime/MovieSceneTracks/Private/Systems/MovieSceneMaterialParameterSystem.cpp
new file mode 100644
--- /dev/null
+++ b/Engine/Source/Runtime/MovieSceneTracks/Private/Systems/MovieSceneMaterialParameterSystem.cpp
@@ -0,0 +1,295 @@
+// Copyright Epic Games, Inc. All Rights Reserved.
+
+#include "Systems/MovieSceneMaterialParameterSystem.h"
+
+#include "Systems/BoolChannelEvaluatorSystem.h"
+#include "Systems/FloatChannelEvaluatorSystem.h"
+#include "Systems/MovieScenePiecewiseFloatBlenderSystem.h"
+#include "Systems/MovieSceneHierarchicalBiasSystem.h"
+#include "Systems/MovieSceneMaterialSystem.h"
+
+#include "Materials/MaterialParameterCollectionInstance.h"
+#include "Materials/MaterialInstanceDynamic.h"
+
+namespace UE::MovieScene
+{
+
+/** Apply scalar material parameters */
+struct FApplyScalarParameters
+{
+	static void ForEachEntity(UObject* BoundMaterial, FName ParameterName, float InScalarValue)
+	{
+		// WARNING: BoundMaterial may be nullptr here
+		if (UMaterialInstanceDynamic* MID = Cast<UMaterialInstanceDynamic>(BoundMaterial))
+		{
+			MID->SetScalarParameterValue(ParameterName, InScalarValue);
+		}
+		else if (UMaterialParameterCollectionInstance* MPCI = Cast<UMaterialParameterCollectionInstance>(BoundMaterial))
+		{
+			MPCI->SetScalarParameterValue(ParameterName, InScalarValue);
+		}
+	}
+};
+
+/** Apply vector material parameters */
+struct FApplyVectorParameters
+{
+	static void ForEachAllocation(FEntityAllocationIteratorItem Item,
+		TRead<UObject*> BoundMaterials,
+		TReadOneOrMoreOf<FName, FName> VectorOrColorParameterNames,
+		TReadOneOrMoreOf<float, float, float, float> VectorChannels)
+	{
+		const int32 Num = Item.GetAllocation()->Num();
+		// Use either the vector parameter name, or the color parameter name
+		const FName* ParameterNames = VectorOrColorParameterNames.Get<0>() ? VectorOrColorParameterNames.Get<0>() : VectorOrColorParameterNames.Get<1>();
+		const float* RESTRICT R = VectorChannels.Get<0>();
+		const float* RESTRICT G = VectorChannels.Get<1>();
+		const float* RESTRICT B = VectorChannels.Get<2>();
+		const float* RESTRICT A = VectorChannels.Get<3>();
+		
+		for (int32 Index = 0; Index < Num; ++Index)
+		{
+			UObject* BoundMaterial = BoundMaterials[Index];
+			if (!BoundMaterial)
+			{
+				continue;
+			}
+
+			FLinearColor Color(
+				R ? R[Index] : 0.f,
+				G ? G[Index] : 0.f,
+				B ? B[Index] : 0.f,
+				A ? A[Index] : 1.f
+			);
+			if (UMaterialInstanceDynamic* MID = Cast<UMaterialInstanceDynamic>(BoundMaterial))
+			{
+				MID->SetVectorParameterValue(ParameterNames[Index], Color);
+			}
+			else if (UMaterialParameterCollectionInstance* MPCI = Cast<UMaterialParameterCollectionInstance>(BoundMaterial))
+			{
+				MPCI->SetVectorParameterValue(ParameterNames[Index], Color);
+			}
+		}
+	}
+};
+
+/** Handler that manages creation of blend outputs where there are multiple contributors for the same material parameter */
+struct FOverlappingMaterialParameterHandler
+{
+	UMovieSceneEntitySystemLinker* Linker;
+	UMovieSceneMaterialParameterSystem* System;
+	FComponentMask DefaultComponentMask;
+
+	FOverlappingMaterialParameterHandler(UMovieSceneMaterialParameterSystem* InSystem)
+		: Linker(InSystem->GetLinker())
+		, System(InSystem)
+	{}
+
+	void InitializeOutput(UObject* BoundMaterial, FName ParameterName, TArrayView<const FMovieSceneEntityID> Inputs, FAnimatedMaterialParameterInfo* Output, FEntityOutputAggregate Aggregate)
+	{
+		UpdateOutput(BoundMaterial, ParameterName, Inputs, Output, Aggregate);
+	}
+
+	void UpdateOutput(UObject* BoundMaterial, FName ParameterName, TArrayView<const FMovieSceneEntityID> Inputs, FAnimatedMaterialParameterInfo* Output, FEntityOutputAggregate Aggregate)
+	{
+		if (BoundMaterial == nullptr)
+		{
+			return;
+		}
+
+		using namespace UE::MovieScene;
+
+		FBuiltInComponentTypes* BuiltInComponents = FBuiltInComponentTypes::Get();
+		FMovieSceneTracksComponentTypes* TracksComponents = FMovieSceneTracksComponentTypes::Get();
+
+		const int32 NumContributors = Inputs.Num();
+		if (NumContributors > 1)
+		{
+			if (!Output->OutputEntityID)
+			{
+				if (!System->FloatBlenderSystem)
+				{
+					System->FloatBlenderSystem = Linker->LinkSystem<UMovieScenePiecewiseFloatBlenderSystem>();
+					Linker->SystemGraph.AddReference(System, System->FloatBlenderSystem);
+				}
+
+				// Initialize the blend channel ID
+				Output->BlendChannelID = System->FloatBlenderSystem->AllocateBlendChannel();
+
+				// Needs blending
+				Output->OutputEntityID = FEntityBuilder()
+				.Add(TracksComponents->BoundMaterial, BoundMaterial)
+				.Add(BuiltInComponents->BlendChannelOutput, Output->BlendChannelID)
+				.AddTag(BuiltInComponents->Tags.NeedsLink)
+				.AddMutualComponents()
+				.CreateEntity(&Linker->EntityManager, DefaultComponentMask);
+
+				Linker->EntityManager.CopyComponents(Inputs[0], Output->OutputEntityID, Linker->EntityManager.GetComponents()->GetCopyAndMigrationMask());
+			}
+
+			for (FMovieSceneEntityID Input : Inputs)
+			{
+				if (!Linker->EntityManager.HasComponent(Input, BuiltInComponents->BlendChannelInput))
+				{
+					Linker->EntityManager.AddComponent(Input, BuiltInComponents->BlendChannelInput, Output->BlendChannelID);
+				}
+			}
+		}
+		else
+		{
+			// Previously blended, but is no more - remove the tag from the remaining input if necessary and delete the output entity
+			if (Output->OutputEntityID)
+			{
+				Linker->EntityManager.AddComponent(Output->OutputEntityID, BuiltInComponents->Tags.NeedsUnlink);
+				Output->OutputEntityID = FMovieSceneEntityID();
+
+				if (ensure(System->FloatBlenderSystem))
+				{
+					System->FloatBlenderSystem->ReleaseBlendChannel(Output->BlendChannelID);
+				}
+			}
+
+			for (FMovieSceneEntityID Input : Inputs)
+			{
+				Linker->EntityManager.RemoveComponent(Input, BuiltInComponents->BlendChannelInput);
+			}
+		}
+
+		Output->NumContributors = NumContributors;
+	}
+
+	void DestroyOutput(UObject* BoundMaterial, FName ParameterName, FAnimatedMaterialParameterInfo* Output, FEntityOutputAggregate Aggregate)
+	{
+		if (Output->OutputEntityID)
+		{
+			FBuiltInComponentTypes* BuiltInComponents = FBuiltInComponentTypes::Get();
+
+			Linker->EntityManager.AddComponent(Output->OutputEntityID, BuiltInComponents->Tags.NeedsUnlink);
+			Output->OutputEntityID = FMovieSceneEntityID();
+		}
+	}
+
+};
+
+} // namespace UE::MovieScene
+
+UMovieSceneMaterialParameterSystem::UMovieSceneMaterialParameterSystem(const FObjectInitializer& ObjInit)
+	: Super(ObjInit)
+{
+	using namespace UE::MovieScene;
+
+	FBuiltInComponentTypes* BuiltInComponents = FBuiltInComponentTypes::Get();
+	FMovieSceneTracksComponentTypes* TracksComponents = FMovieSceneTracksComponentTypes::Get();
+
+	RelevantComponent = TracksComponents->BoundMaterial;
+	Phase = ESystemPhase::Instantiation | ESystemPhase::Evaluation;
+
+	if (HasAnyFlags(RF_ClassDefaultObject))
+	{
+		DefineComponentConsumer(GetClass(), TracksComponents->BoundMaterial);
+
+		DefineImplicitPrerequisite(UFloatChannelEvaluatorSystem::StaticClass(), GetClass());
+
+		DefineImplicitPrerequisite(UMovieScenePiecewiseFloatBlenderSystem::StaticClass(), GetClass());
+		DefineImplicitPrerequisite(GetClass(), UMovieSceneHierarchicalBiasSystem::StaticClass());
+	}
+}
+
+void UMovieSceneMaterialParameterSystem::OnLink()
+{
+	using namespace UE::MovieScene;
+
+	// Always reset the float blender system on link to ensure that recycled systems are correctly initialized.
+	FloatBlenderSystem = nullptr;
+
+	FOverlappingMaterialParameterHandler Handler(this);
+
+	ScalarParameterTracker.Destroy(Handler);
+	VectorParameterTracker.Destroy(Handler);
+
+	ScalarParameterTracker.Initialize(this);
+	VectorParameterTracker.Initialize(this);
+}
+
+void UMovieSceneMaterialParameterSystem::OnRun(FSystemTaskPrerequisites& InPrerequisites, FSystemSubsequentTasks& Subsequents)
+{
+	using namespace UE::MovieScene;
+
+	ESystemPhase CurrentPhase = Linker->GetActiveRunner()->GetCurrentPhase();
+	if (CurrentPhase == ESystemPhase::Instantiation)
+	{
+		OnInstantiation();
+	}
+	else if (CurrentPhase == ESystemPhase::Evaluation)
+	{
+		OnEvaluation(InPrerequisites, Subsequents);
+	}
+}
+
+
+void UMovieSceneMaterialParameterSystem::OnInstantiation()
+{
+	using namespace UE::MovieScene;
+
+	FBuiltInComponentTypes*          BuiltInComponents = FBuiltInComponentTypes::Get();
+	FMovieSceneTracksComponentTypes* TracksComponents  = FMovieSceneTracksComponentTypes::Get();
+
+	const bool bHasScalars = Linker->EntityManager.ContainsComponent(TracksComponents->ScalarParameterName);
+	const bool bHasColors = Linker->EntityManager.ContainsComponent(TracksComponents->ColorParameterName);
+	const bool bHasVectors = Linker->EntityManager.ContainsComponent(TracksComponents->VectorParameterName);
+
+	if (bHasScalars)
+	{
+		ScalarParameterTracker.Update(Linker, TracksComponents->BoundMaterial, TracksComponents->ScalarParameterName, FEntityComponentFilter());
+
+		FOverlappingMaterialParameterHandler Handler(this);
+		Handler.DefaultComponentMask.Set(BuiltInComponents->FloatResult[0]);
+		ScalarParameterTracker.ProcessInvalidatedOutputs(Linker, Handler);
+	}
+
+	if (bHasColors)
+	{
+		VectorParameterTracker.Update(Linker, TracksComponents->BoundMaterial, TracksComponents->ColorParameterName, FEntityComponentFilter());
+	}
+	if (bHasVectors)
+	{
+		VectorParameterTracker.Update(Linker, TracksComponents->BoundMaterial, TracksComponents->VectorParameterName, FEntityComponentFilter());
+	}
+	if (bHasColors || bHasVectors)
+	{
+		FOverlappingMaterialParameterHandler Handler(this);
+		Handler.DefaultComponentMask.SetAll({ BuiltInComponents->FloatResult[0], BuiltInComponents->FloatResult[1], BuiltInComponents->FloatResult[2], BuiltInComponents->FloatResult[3] });
+		VectorParameterTracker.ProcessInvalidatedOutputs(Linker, Handler);
+	}
+}
+
+void UMovieSceneMaterialParameterSystem::OnEvaluation(FSystemTaskPrerequisites& InPrerequisites, FSystemSubsequentTasks& Subsequents)
+{
+	using namespace UE::MovieScene;
+
+	FBuiltInComponentTypes* BuiltInComponents = FBuiltInComponentTypes::Get();
+	FMovieSceneTracksComponentTypes* TracksComponents = FMovieSceneTracksComponentTypes::Get();
+
+	if (Linker->EntityManager.ContainsComponent(TracksComponents->ScalarParameterName))
+	{
+		FEntityTaskBuilder()
+		.Read(TracksComponents->BoundMaterial)
+		.Read(TracksComponents->ScalarParameterName)
+		.Read(BuiltInComponents->FloatResult[0])
+		.FilterNone({ BuiltInComponents->BlendChannelInput })
+		.SetDesiredThread(Linker->EntityManager.GetDispatchThread())
+		.Dispatch_PerEntity<FApplyScalarParameters>(&Linker->EntityManager, InPrerequisites, &Subsequents);
+	}
+
+	// Vectors and colors use the same API
+	if (Linker->EntityManager.ContainsComponent(TracksComponents->VectorParameterName) || Linker->EntityManager.ContainsComponent(TracksComponents->ColorParameterName))
+	{
+		FEntityTaskBuilder()
+		.Read(TracksComponents->BoundMaterial)
+		.ReadOneOrMoreOf(TracksComponents->VectorParameterName, TracksComponents->ColorParameterName)
+		.ReadOneOrMoreOf(BuiltInComponents->FloatResult[0], BuiltInComponents->FloatResult[1], BuiltInComponents->FloatResult[2], BuiltInComponents->FloatResult[3])
+		.FilterNone({ BuiltInComponents->BlendChannelInput })
+		.SetDesiredThread(Linker->EntityManager.GetDispatchThread())
+		.Dispatch_PerAllocation<FApplyVectorParameters>(&Linker->EntityManager, InPrerequisites, &Subsequents);
+	}
+}
\ No newline at end of file
diff --git a/Engine/Source/Runtime/MovieSceneTracks/Private/Tracks/MovieSceneMaterialParameterCollectionTrack.cpp b/Engine/Source/Runtime/MovieSceneTracks/Private/Tracks/MovieSceneMaterialParameterCollectionTrack.cpp
--- a/Engine/Source/Runtime/MovieSceneTracks/Private/Tracks/MovieSceneMaterialParameterCollectionTrack.cpp
+++ b/Engine/Source/Runtime/MovieSceneTracks/Private/Tracks/MovieSceneMaterialParameterCollectionTrack.cpp
@@ -2,33 +2,82 @@
 
 #include "Tracks/MovieSceneMaterialParameterCollectionTrack.h"
 #include "Materials/MaterialParameterCollection.h"
-#include "Evaluation/MovieSceneMaterialParameterCollectionTemplate.h"
+
+#include "Evaluation/MovieSceneEvaluationField.h"
+#include "EntitySystem/BuiltInComponentTypes.h"
+#include "MovieSceneTracksComponentTypes.h"
 
 #define LOCTEXT_NAMESPACE "MovieSceneMaterialParameterCollectionTrack"
 
 UMovieSceneMaterialParameterCollectionTrack::UMovieSceneMaterialParameterCollectionTrack(const FObjectInitializer& ObjectInitializer)
 	: Super(ObjectInitializer)
 {
+	BuiltInTreePopulationMode = ETreePopulationMode::Blended;
+
 #if WITH_EDITORONLY_DATA
 	TrackTint = FColor(64,192,64,65);
 #endif
 }
 
-bool UMovieSceneMaterialParameterCollectionTrack::SupportsType(TSubclassOf<UMovieSceneSection> SectionClass) const
-{
-	return SectionClass == UMovieSceneParameterSection::StaticClass();
-}
-
 UMovieSceneSection* UMovieSceneMaterialParameterCollectionTrack::CreateNewSection()
 {
 	UMovieSceneSection* NewSection = NewObject<UMovieSceneParameterSection>(this, NAME_None, RF_Transactional);
 	NewSection->SetRange(TRange<FFrameNumber>::All());
+	NewSection->SetBlendType(EMovieSceneBlendType::Absolute);
 	return NewSection;
 }
 
-FMovieSceneEvalTemplatePtr UMovieSceneMaterialParameterCollectionTrack::CreateTemplateForSection(const UMovieSceneSection& InSection) const
+void UMovieSceneMaterialParameterCollectionTrack::ImportEntityImpl(UMovieSceneEntitySystemLinker* EntityLinker, const FEntityImportParams& Params, FImportedEntity* OutImportedEntity)
+{
+	// These tracks don't define any entities for themselves
+	checkf(false, TEXT("This track should never have created entities for itself - this assertion indicates an error in the entity-component field"));
+}
+
+void UMovieSceneMaterialParameterCollectionTrack::ExtendEntityImpl(UMovieSceneEntitySystemLinker* EntityLinker, const UE::MovieScene::FEntityImportParams& Params, UE::MovieScene::FImportedEntity* OutImportedEntity)
+{
+	using namespace UE::MovieScene;
+
+	FBuiltInComponentTypes* BuiltInComponents = FBuiltInComponentTypes::Get();
+	FMovieSceneTracksComponentTypes* TracksComponents = FMovieSceneTracksComponentTypes::Get();
+
+	// Material parameters are always absolute blends for the time being
+	OutImportedEntity->AddBuilder(
+		FEntityBuilder()
+		.Add(TracksComponents->MPC, MPC)
+		.AddTag(BuiltInComponents->Tags.AbsoluteBlend)
+		.AddTag(BuiltInComponents->Tags.Master)
+	);
+}
+
+bool UMovieSceneMaterialParameterCollectionTrack::PopulateEvaluationFieldImpl(const TRange<FFrameNumber>& EffectiveRange, const FMovieSceneEvaluationFieldEntityMetaData& InMetaData, FMovieSceneEntityComponentFieldBuilder* OutFieldBuilder)
+{
+	const FMovieSceneTrackEvaluationField& LocalEvaluationField = GetEvaluationField();
+
+	// Define entities for every entry in our evaluation field
+	for (const FMovieSceneTrackEvaluationFieldEntry& Entry : LocalEvaluationField.Entries)
+	{
+		UMovieSceneParameterSection* ParameterSection = Cast<UMovieSceneParameterSection>(Entry.Section);
+		if (!ParameterSection || IsRowEvalDisabled(ParameterSection->GetRowIndex()))
+		{
+			continue;
+		}
+
+		TRange<FFrameNumber> SectionEffectiveRange = TRange<FFrameNumber>::Intersection(EffectiveRange, Entry.Range);
+		if (!SectionEffectiveRange.IsEmpty())
+		{
+			FMovieSceneEvaluationFieldEntityMetaData SectionMetaData = InMetaData;
+			SectionMetaData.Flags = Entry.Flags;
+
+			ParameterSection->ExternalPopulateEvaluationField(SectionEffectiveRange, SectionMetaData, OutFieldBuilder);
+		}
+	}
+
+	return true;
+}
+
+bool UMovieSceneMaterialParameterCollectionTrack::SupportsType(TSubclassOf<UMovieSceneSection> SectionClass) const
 {
-	return FMovieSceneMaterialParameterCollectionTemplate(*CastChecked<UMovieSceneParameterSection>(&InSection), *this);
+	return SectionClass == UMovieSceneParameterSection::StaticClass();
 }
 
 #if WITH_EDITORONLY_DATA
diff --git a/Engine/Source/Runtime/MovieSceneTracks/Private/Tracks/MovieSceneMaterialTrack.cpp b/Engine/Source/Runtime/MovieSceneTracks/Private/Tracks/MovieSceneMaterialTrack.cpp
--- a/Engine/Source/Runtime/MovieSceneTracks/Private/Tracks/MovieSceneMaterialTrack.cpp
+++ b/Engine/Source/Runtime/MovieSceneTracks/Private/Tracks/MovieSceneMaterialTrack.cpp
@@ -3,8 +3,10 @@
 #include "Tracks/MovieSceneMaterialTrack.h"
 #include "MovieScene.h"
 #include "MovieSceneCommonHelpers.h"
-#include "Evaluation/MovieSceneParameterTemplate.h"
-#include "Evaluation/MovieSceneEvaluationTrack.h"
+#include "Evaluation/MovieSceneEvaluationField.h"
+#include "EntitySystem/MovieSceneEntityBuilder.h"
+#include "EntitySystem/BuiltInComponentTypes.h"
+#include "MovieSceneTracksComponentTypes.h"
 
 UMovieSceneMaterialTrack::UMovieSceneMaterialTrack(const FObjectInitializer& ObjectInitializer)
 	: Super(ObjectInitializer)
@@ -12,6 +14,8 @@ UMovieSceneMaterialTrack::UMovieSceneMaterialTrack(const FObjectInitializer& Obj
 #if WITH_EDITORONLY_DATA
 	TrackTint = FColor(64,192,64,65);
 #endif
+
+	SupportedBlendTypes.Add(EMovieSceneBlendType::Absolute);
 }
 
 
@@ -60,6 +64,10 @@ bool UMovieSceneMaterialTrack::IsEmpty() const
 	return Sections.Num() == 0;
 }
 
+bool UMovieSceneMaterialTrack::SupportsMultipleRows() const
+{
+	return true;
+}
 
 const TArray<UMovieSceneSection*>& UMovieSceneMaterialTrack::GetAllSections() const
 {
@@ -111,20 +119,62 @@ void UMovieSceneMaterialTrack::AddColorParameterKey(FName ParameterName, FFrameN
 UMovieSceneComponentMaterialTrack::UMovieSceneComponentMaterialTrack(const FObjectInitializer& ObjectInitializer)
 	: Super(ObjectInitializer)
 {
-	
+	BuiltInTreePopulationMode = ETreePopulationMode::Blended;
 }
 
+void UMovieSceneComponentMaterialTrack::AddSection(UMovieSceneSection& Section)
+{
+	// Materials are always blendable now
+	Section.SetBlendType(EMovieSceneBlendType::Absolute);
+	Super::AddSection(Section);
+}
 
-FMovieSceneEvalTemplatePtr UMovieSceneComponentMaterialTrack::CreateTemplateForSection(const UMovieSceneSection& InSection) const
+void UMovieSceneComponentMaterialTrack::ImportEntityImpl(UMovieSceneEntitySystemLinker* EntityLinker, const FEntityImportParams& Params, FImportedEntity* OutImportedEntity)
 {
-	return FMovieSceneComponentMaterialSectionTemplate(*CastChecked<UMovieSceneParameterSection>(&InSection), *this);
+	// These tracks don't define any entities for themselves
+	checkf(false, TEXT("This track should never have created entities for itself - this assertion indicates an error in the entity-component field"));
 }
 
-void UMovieSceneComponentMaterialTrack::PostCompile(FMovieSceneEvaluationTrack& OutTrack, const FMovieSceneTrackCompilerArgs& Args) const
+void UMovieSceneComponentMaterialTrack::ExtendEntityImpl(UMovieSceneEntitySystemLinker* EntityLinker, const UE::MovieScene::FEntityImportParams& Params, UE::MovieScene::FImportedEntity* OutImportedEntity)
 {
-	OutTrack.SetEvaluationPriority(EvaluationPriority);
+	using namespace UE::MovieScene;
+
+	FBuiltInComponentTypes* BuiltInComponents = FBuiltInComponentTypes::Get();
+	FMovieSceneTracksComponentTypes* TracksComponents = FMovieSceneTracksComponentTypes::Get();
+
+	// Material parameters are always absolute blends for the time being
+	OutImportedEntity->AddBuilder(
+		FEntityBuilder()
+		.Add(TracksComponents->ComponentMaterialIndex, MaterialIndex)
+		.AddTag(BuiltInComponents->Tags.AbsoluteBlend)
+	);
 }
 
+bool UMovieSceneComponentMaterialTrack::PopulateEvaluationFieldImpl(const TRange<FFrameNumber>& EffectiveRange, const FMovieSceneEvaluationFieldEntityMetaData& InMetaData, FMovieSceneEntityComponentFieldBuilder* OutFieldBuilder)
+{
+	const FMovieSceneTrackEvaluationField& LocalEvaluationField = GetEvaluationField();
+
+	// Define entities for every entry in our evaluation field
+	for (const FMovieSceneTrackEvaluationFieldEntry& Entry : LocalEvaluationField.Entries)
+	{
+		UMovieSceneParameterSection* ParameterSection = Cast<UMovieSceneParameterSection>(Entry.Section);
+		if (!ParameterSection || IsRowEvalDisabled(ParameterSection->GetRowIndex()))
+		{
+			continue;
+		}
+
+		TRange<FFrameNumber> SectionEffectiveRange = TRange<FFrameNumber>::Intersection(EffectiveRange, Entry.Range);
+		if (!SectionEffectiveRange.IsEmpty())
+		{
+			FMovieSceneEvaluationFieldEntityMetaData SectionMetaData = InMetaData;
+			SectionMetaData.Flags = Entry.Flags;
+
+			ParameterSection->ExternalPopulateEvaluationField(SectionEffectiveRange, SectionMetaData, OutFieldBuilder);
+		}
+	}
+
+	return true;
+}
 
 #if WITH_EDITORONLY_DATA
 FText UMovieSceneComponentMaterialTrack::GetDefaultDisplayName() const
diff --git a/Engine/Source/Runtime/MovieSceneTracks/Private/Tracks/MovieScenePrimitiveMaterialTrack.cpp b/Engine/Source/Runtime/MovieSceneTracks/Private/Tracks/MovieScenePrimitiveMaterialTrack.cpp
--- a/Engine/Source/Runtime/MovieSceneTracks/Private/Tracks/MovieScenePrimitiveMaterialTrack.cpp
+++ b/Engine/Source/Runtime/MovieSceneTracks/Private/Tracks/MovieScenePrimitiveMaterialTrack.cpp
@@ -3,7 +3,6 @@
 #include "Tracks/MovieScenePrimitiveMaterialTrack.h"
 #include "Tracks/MovieSceneMaterialTrack.h"
 #include "Sections/MovieScenePrimitiveMaterialSection.h"
-#include "Evaluation/MovieScenePrimitiveMaterialTemplate.h"
 #include "Evaluation/MovieSceneEvaluationTrack.h"
 
 
@@ -26,17 +25,6 @@ bool UMovieScenePrimitiveMaterialTrack::SupportsType(TSubclassOf<UMovieSceneSect
 	return SectionClass == UMovieScenePrimitiveMaterialSection::StaticClass();
 }
 
-FMovieSceneEvalTemplatePtr UMovieScenePrimitiveMaterialTrack::CreateTemplateForSection(const UMovieSceneSection& InSection) const
-{
-	return FMovieScenePrimitiveMaterialTemplate(*CastChecked<UMovieScenePrimitiveMaterialSection>(&InSection), *this);
-}
-
-void UMovieScenePrimitiveMaterialTrack::PostCompile(FMovieSceneEvaluationTrack& OutTrack, const FMovieSceneTrackCompilerArgs& Args) const
-{
-	// Must evaluate before material parameter tracks
-	OutTrack.SetEvaluationPriority(UMovieSceneComponentMaterialTrack::EvaluationPriority + 1);
-}
-
 int32 UMovieScenePrimitiveMaterialTrack::GetMaterialIndex() const
 {
 	return MaterialIndex;
diff --git a/Engine/Source/Runtime/MovieSceneTracks/Public/Evaluation/MovieSceneParameterTemplate.h b/Engine/Source/Runtime/MovieSceneTracks/Public/Evaluation/MovieSceneParameterTemplate.h
--- a/Engine/Source/Runtime/MovieSceneTracks/Public/Evaluation/MovieSceneParameterTemplate.h
+++ b/Engine/Source/Runtime/MovieSceneTracks/Public/Evaluation/MovieSceneParameterTemplate.h
@@ -6,16 +6,12 @@
 #include "Stats/Stats.h"
 #include "UObject/ObjectMacros.h"
 #include "MovieSceneFwd.h"
-#include "MovieSceneExecutionToken.h"
 #include "Evaluation/MovieSceneEvalTemplate.h"
 #include "Sections/MovieSceneParameterSection.h"
-#include "Materials/MaterialInstanceDynamic.h"
 #include "IMovieScenePlayer.h"
 
 #include "MovieSceneParameterTemplate.generated.h"
 
-class UMovieSceneComponentMaterialTrack;
-
 DECLARE_CYCLE_STAT(TEXT("Parameter Track Token Execute"), MovieSceneEval_ParameterTrack_TokenExecute, STATGROUP_MovieSceneEval);
 
 /** Evaluation structure that holds evaluated values */
@@ -85,146 +81,3 @@ protected:
 	UPROPERTY()
 	TArray<FTransformParameterNameAndCurves> Transforms;
 };
-
-/** Default accessor type for use with TMaterialTrackExecutionToken */
-struct FDefaultMaterialAccessor
-{
-	// Implement in derived classes:
-
-	/** Get the anim type ID for the evaluation token */
-	// FMovieSceneAnimTypeID 	GetAnimTypeID();
-
-	/** Get the material from the specified object */
-	// UMaterialInterface* 		GetMaterialForObject(UObject& Object);
-
-	/** Set the material for the specified object */
-	// void 					SetMaterialForObject(UObject& Object, UMaterialInterface& Material);
-
-	/** Apply the specified values onto the specified material */
-	MOVIESCENETRACKS_API void Apply(UMaterialInstanceDynamic& Material, const FEvaluatedParameterSectionValues& Values);
-};
-
-/**
- * Material track execution token
- * Templated on accessor type to allow for copyable accessors into pre animated state
- */
-template<typename AccessorType>
-struct TMaterialTrackExecutionToken : IMovieSceneExecutionToken
-{
-	TMaterialTrackExecutionToken(AccessorType InAccessor)
-		: Accessor(MoveTemp(InAccessor))
-	{
-	}
-
-	TMaterialTrackExecutionToken(TMaterialTrackExecutionToken&&) = default;
-	TMaterialTrackExecutionToken& operator=(TMaterialTrackExecutionToken&&) = default;
-
-	// Non-copyable
-	TMaterialTrackExecutionToken(const TMaterialTrackExecutionToken&) = delete;
-	TMaterialTrackExecutionToken& operator=(const TMaterialTrackExecutionToken&) = delete;
-
-	virtual void Execute(const FMovieSceneContext& Context, const FMovieSceneEvaluationOperand& Operand, FPersistentEvaluationData& PersistentData, IMovieScenePlayer& Player)
-	{
-		MOVIESCENE_DETAILED_SCOPE_CYCLE_COUNTER(MovieSceneEval_ParameterTrack_TokenExecute)
-		
-		for (TWeakObjectPtr<>& WeakObject : Player.FindBoundObjects(Operand))
-		{
-			UObject* Object = WeakObject.Get();
-			UMaterialInterface* Material = Object ? Accessor.GetMaterialForObject(*Object) : nullptr;
-
-			if (!Material)
-			{
-				continue;
-			}
-
-			// Save the old instance
-			Player.SavePreAnimatedState(*Object, Accessor.GetAnimTypeID(), FPreAnimatedTokenProducer(Accessor));
-
-			UMaterialInstanceDynamic* DynamicMaterialInstance = Cast<UMaterialInstanceDynamic>(Material);
-			if (!DynamicMaterialInstance)
-			{
-				FString DynamicName = Material->GetName() + "_Animated";
-				FName UniqueDynamicName = MakeUniqueObjectName( Object, UMaterialInstanceDynamic::StaticClass() , *DynamicName );
-				DynamicMaterialInstance = Accessor.CreateMaterialInstanceDynamic(*Object, *Material, UniqueDynamicName);
-
-				Accessor.SetMaterialForObject(*Object, *DynamicMaterialInstance);
-			}
-
-			Accessor.Apply(*DynamicMaterialInstance, Values);
-		}
-	}
-
-	AccessorType Accessor;
-	FEvaluatedParameterSectionValues Values;
-
-private:
-
-	struct FPreAnimatedTokenProducer : IMovieScenePreAnimatedTokenProducer
-	{
-		struct FPreAnimatedToken : IMovieScenePreAnimatedToken
-		{
-			FPreAnimatedToken(UObject& Object, const AccessorType& InAccessor)
-				: Accessor(InAccessor)
-				, Material(Accessor.GetMaterialForObject(Object))
-			{
-				// If the current material we're overriding is already a material instance dynamic, copy it since we will be modifying the data. 
-				// The copied material will be used to restore the values in RestoreState.
-				UMaterialInstanceDynamic* DynamicMaterialInstance = Cast<UMaterialInstanceDynamic>(Material);
-				if (DynamicMaterialInstance)
-				{
-					UObject* DuplicatedObject = StaticDuplicateObject(DynamicMaterialInstance, &Object);
-					if (DuplicatedObject)
-					{
-						CopiedDynamicMaterialInstance = Cast<UMaterialInstanceDynamic>(DuplicatedObject);
-					}
-				}
-			}
-
-			virtual void RestoreState(UObject& Object, const UE::MovieScene::FRestoreStateParams& Params) override
-			{
-				if (UMaterialInterface* PinnedMaterial = Material.Get())
-				{
-					// If the original material was a material instance dynamic, copy the material parameters back since we modified it in place.
-					UMaterialInstanceDynamic* DynamicMaterialInstance = Cast<UMaterialInstanceDynamic>(PinnedMaterial);
-					if (DynamicMaterialInstance && CopiedDynamicMaterialInstance.IsValid())
-					{
-						DynamicMaterialInstance->CopyMaterialUniformParameters(CopiedDynamicMaterialInstance.Get());
-					}
-
-					Accessor.SetMaterialForObject(Object, *PinnedMaterial);
-				}
-			}
-
-			AccessorType Accessor;
-			TWeakObjectPtr<UMaterialInterface> Material;
-			TWeakObjectPtr<UMaterialInstanceDynamic> CopiedDynamicMaterialInstance;
-		};
-
-		FPreAnimatedTokenProducer(const AccessorType& InAccessor) : Accessor(InAccessor) {}
-
-		virtual IMovieScenePreAnimatedTokenPtr CacheExistingState(UObject& Object) const override
-		{
-			return FPreAnimatedToken(Object, Accessor);
-		}
-
-		const AccessorType& Accessor;
-	};
-};
-
-/** Evaluation template for primitive component materials */
-USTRUCT()
-struct FMovieSceneComponentMaterialSectionTemplate : public FMovieSceneParameterSectionTemplate
-{
-	GENERATED_BODY()
-
-	FMovieSceneComponentMaterialSectionTemplate() : MaterialIndex(0) {}
-	FMovieSceneComponentMaterialSectionTemplate(const UMovieSceneParameterSection& Section, const UMovieSceneComponentMaterialTrack& Track);
-
-private:
-
-	virtual UScriptStruct& GetScriptStructImpl() const override { return *StaticStruct(); }
-	virtual void Evaluate(const FMovieSceneEvaluationOperand& Operand, const FMovieSceneContext& Context, const FPersistentEvaluationData& PersistentData, FMovieSceneExecutionTokens& ExecutionTokens) const override;
-
-	UPROPERTY()
-	int32 MaterialIndex;
-};
diff --git a/Engine/Source/Runtime/MovieSceneTracks/Public/MovieSceneTracksComponentTypes.h b/Engine/Source/Runtime/MovieSceneTracks/Public/MovieSceneTracksComponentTypes.h
--- a/Engine/Source/Runtime/MovieSceneTracks/Public/MovieSceneTracksComponentTypes.h
+++ b/Engine/Source/Runtime/MovieSceneTracks/Public/MovieSceneTracksComponentTypes.h
@@ -15,6 +15,7 @@
 #include "Styling/SlateColor.h"
 #include "MovieSceneTracksComponentTypes.generated.h"
 
+class UMaterialParameterCollection;
 class UMovieSceneDataLayerSection;
 class UMovieSceneLevelVisibilitySection;
 struct FMovieSceneObjectBindingID;
@@ -488,6 +489,18 @@ struct MOVIESCENETRACKS_API FMovieSceneTracksComponentTypes
 	TComponentTypeID<FFloatPerlinNoiseParams> FloatPerlinNoiseChannel;
 	TComponentTypeID<FDoublePerlinNoiseParams> DoublePerlinNoiseChannel;
 
+	TComponentTypeID<int32> ComponentMaterialIndex;
+
+	TComponentTypeID<FName> BoolParameterName;
+	TComponentTypeID<FName> ScalarParameterName;
+	TComponentTypeID<FName> Vector2DParameterName;
+	TComponentTypeID<FName> VectorParameterName;
+	TComponentTypeID<FName> ColorParameterName;
+	TComponentTypeID<FName> TransformParameterName;
+
+	TComponentTypeID<UObject*> BoundMaterial;
+	TComponentTypeID<UMaterialParameterCollection*> MPC;
+
 	struct
 	{
 		TCustomPropertyRegistration<FBoolPropertyTraits> Bool;
diff --git a/Engine/Source/Runtime/MovieSceneTracks/Public/Sections/MovieSceneParameterSection.h b/Engine/Source/Runtime/MovieSceneTracks/Public/Sections/MovieSceneParameterSection.h
--- a/Engine/Source/Runtime/MovieSceneTracks/Public/Sections/MovieSceneParameterSection.h
+++ b/Engine/Source/Runtime/MovieSceneTracks/Public/Sections/MovieSceneParameterSection.h
@@ -10,11 +10,35 @@
 #include "Channels/MovieSceneFloatChannel.h"
 #include "Channels/MovieSceneBoolChannel.h"
 #include "Sections/MovieScene3DTransformSection.h"
+#include "EntitySystem/IMovieSceneEntityProvider.h"
 #include "MovieSceneParameterSection.generated.h"
 
 struct FMovieSceneParameterPropertyInterface;
 struct FMovieSceneConstParameterPropertyInterface;
 
+
+UINTERFACE()
+class MOVIESCENETRACKS_API UMovieSceneParameterSectionExtender : public UInterface
+{
+public:
+	GENERATED_BODY()
+};
+
+class MOVIESCENETRACKS_API IMovieSceneParameterSectionExtender
+{
+public:
+
+	GENERATED_BODY()
+
+	void ExtendEntity(UMovieSceneEntitySystemLinker* EntityLinker, const UE::MovieScene::FEntityImportParams& Params, UE::MovieScene::FImportedEntity* OutImportedEntity);
+
+private:
+
+	virtual void ExtendEntityImpl(UMovieSceneEntitySystemLinker* EntityLinker, const UE::MovieScene::FEntityImportParams& Params, UE::MovieScene::FImportedEntity* OutImportedEntity) = 0;
+};
+
+
+
 /**
  * Structure representing the animated value of a scalar parameter.
  */
@@ -308,6 +332,7 @@ struct MOVIESCENETRACKS_API FTransformParameterNameAndCurves
 UCLASS()
 class MOVIESCENETRACKS_API UMovieSceneParameterSection
 	: public UMovieSceneSection
+	, public IMovieSceneEntityProvider
 {
 	GENERATED_UCLASS_BODY()
 
@@ -418,6 +443,11 @@ public:
 	UFUNCTION(BlueprintPure, Category = "Sequencer|Section")
 	void GetParameterNames(TSet<FName>& ParameterNames) const;
 
+	virtual void ImportEntityImpl(UMovieSceneEntitySystemLinker* EntityLinker, const FEntityImportParams& Params, FImportedEntity* OutImportedEntity) override;
+	virtual bool PopulateEvaluationFieldImpl(const TRange<FFrameNumber>& EffectiveRange, const FMovieSceneEvaluationFieldEntityMetaData& InMetaData, FMovieSceneEntityComponentFieldBuilder* OutFieldBuilder) override;
+
+	void ExternalPopulateEvaluationField(const TRange<FFrameNumber>& EffectiveRange, const FMovieSceneEvaluationFieldEntityMetaData& InMetaData, FMovieSceneEntityComponentFieldBuilder* OutFieldBuilder);
+
 protected:
 
 	//~ UMovieSceneSection interface
diff --git a/Engine/Source/Runtime/MovieSceneTracks/Public/Sections/MovieScenePrimitiveMaterialSection.h b/Engine/Source/Runtime/MovieSceneTracks/Public/Sections/MovieScenePrimitiveMaterialSection.h
--- a/Engine/Source/Runtime/MovieSceneTracks/Public/Sections/MovieScenePrimitiveMaterialSection.h
+++ b/Engine/Source/Runtime/MovieSceneTracks/Public/Sections/MovieScenePrimitiveMaterialSection.h
@@ -4,10 +4,13 @@
 
 #include "MovieSceneSection.h"
 #include "Channels/MovieSceneObjectPathChannel.h"
+#include "EntitySystem/IMovieSceneEntityProvider.h"
 #include "MovieScenePrimitiveMaterialSection.generated.h"
 
 UCLASS(MinimalAPI)
-class UMovieScenePrimitiveMaterialSection : public UMovieSceneSection
+class UMovieScenePrimitiveMaterialSection
+	: public UMovieSceneSection
+	, public IMovieSceneEntityProvider
 {
 public:
 
@@ -15,6 +18,8 @@ public:
 
 	UMovieScenePrimitiveMaterialSection(const FObjectInitializer& ObjInit);
 
+	virtual void ImportEntityImpl(UMovieSceneEntitySystemLinker* Linker, const FEntityImportParams& ImportParams, FImportedEntity* OutImportedEntity) override;
+
 	UPROPERTY()
 	FMovieSceneObjectPathChannel MaterialChannel;
 };
diff --git a/Engine/Source/Runtime/MovieSceneTracks/Public/Systems/MovieSceneComponentMaterialSystem.h b/Engine/Source/Runtime/MovieSceneTracks/Public/Systems/MovieSceneComponentMaterialSystem.h
new file mode 100644
--- /dev/null
+++ b/Engine/Source/Runtime/MovieSceneTracks/Public/Systems/MovieSceneComponentMaterialSystem.h
@@ -0,0 +1,87 @@
+// Copyright Epic Games, Inc. All Rights Reserved.
+
+#pragma once
+
+#include "EntitySystem/MovieSceneEntitySystem.h"
+#include "Evaluation/PreAnimatedState/MovieScenePreAnimatedStateStorage.h"
+
+#include "Systems/MovieSceneMaterialSystem.h"
+
+#include "MovieSceneComponentMaterialSystem.generated.h"
+
+class UMovieScenePiecewiseFloatBlenderSystem;
+
+namespace UE::MovieScene
+{
+
+struct FComponentMaterialKey
+{
+	FObjectKey Object;
+	int32 MaterialIndex;
+
+	friend uint32 GetTypeHash(const FComponentMaterialKey& In)
+	{
+		return GetTypeHash(In.Object) ^ GetTypeHash(In.MaterialIndex);
+	}
+	friend bool operator==(const FComponentMaterialKey& A, const FComponentMaterialKey& B)
+	{
+		return A.Object == B.Object && A.MaterialIndex == B.MaterialIndex;
+	}
+};
+
+struct FComponentMaterialAccessor
+{
+	using KeyType = FComponentMaterialKey;
+
+	UObject* Object;
+	int32 MaterialIndex;
+
+	FComponentMaterialAccessor(const FComponentMaterialKey& InKey);
+	FComponentMaterialAccessor(UObject* InObject, int32 InMaterialIndex);
+
+	UMaterialInterface* GetMaterial() const;
+	void SetMaterial(UMaterialInterface* InMaterial) const;
+	UMaterialInstanceDynamic* CreateDynamicMaterial(UMaterialInterface* InMaterial);
+	FString ToString() const;
+};
+
+using FPreAnimatedComponentMaterialTraits          = TPreAnimatedMaterialTraits<FComponentMaterialAccessor, UObject*, int32>;
+using FPreAnimatedComponentMaterialParameterTraits = TPreAnimatedMaterialParameterTraits<FComponentMaterialAccessor, UObject*, int32>;
+
+struct FPreAnimatedComponentMaterialSwitcherStorage
+	: public TPreAnimatedStateStorage<TPreAnimatedMaterialTraits<FComponentMaterialAccessor, UObject*, int32>>
+{
+	static TAutoRegisterPreAnimatedStorageID<FPreAnimatedComponentMaterialSwitcherStorage> StorageID;
+};
+
+struct FPreAnimatedComponentMaterialParameterStorage
+	: public TPreAnimatedStateStorage<TPreAnimatedMaterialParameterTraits<FComponentMaterialAccessor, UObject*, int32>>
+{
+	static TAutoRegisterPreAnimatedStorageID<FPreAnimatedComponentMaterialParameterStorage> StorageID;
+};
+
+} // namespace UE::MovieScene
+
+
+UCLASS(MinimalAPI)
+class UMovieSceneComponentMaterialSystem
+	: public UMovieSceneEntitySystem
+	, public IMovieScenePreAnimatedStateSystemInterface
+{
+public:
+
+	GENERATED_BODY()
+
+	UMovieSceneComponentMaterialSystem(const FObjectInitializer& ObjInit);
+
+private:
+
+	virtual void OnLink() override;
+	virtual void OnRun(FSystemTaskPrerequisites& InPrerequisites, FSystemSubsequentTasks& Subsequents) override;
+
+	virtual void SavePreAnimatedState(const FPreAnimationParameters& InParameters) override;
+
+private:
+
+	UE::MovieScene::TMovieSceneMaterialSystem<UE::MovieScene::FComponentMaterialAccessor, UObject*, int32> SystemImpl;
+};
diff --git a/Engine/Source/Runtime/MovieSceneTracks/Public/Systems/MovieSceneMaterialParameterCollectionSystem.h b/Engine/Source/Runtime/MovieSceneTracks/Public/Systems/MovieSceneMaterialParameterCollectionSystem.h
new file mode 100644
--- /dev/null
+++ b/Engine/Source/Runtime/MovieSceneTracks/Public/Systems/MovieSceneMaterialParameterCollectionSystem.h
@@ -0,0 +1,49 @@
+// Copyright Epic Games, Inc. All Rights Reserved.
+
+#pragma once
+
+#include "EntitySystem/MovieSceneEntitySystem.h"
+#include "Evaluation/PreAnimatedState/MovieScenePreAnimatedStateStorage.h"
+
+#include "Systems/MovieSceneMaterialSystem.h"
+
+#include "MovieSceneMaterialParameterCollectionSystem.generated.h"
+
+namespace UE::MovieScene
+{
+
+struct FPreAnimatedMPCScalarStorage;
+struct FPreAnimatedMPCVectorStorage;
+
+} // namespace UE::MovieScene
+
+
+/**
+ * System that resolves MPC components into BoundMaterial components that can be used by downstream
+ * parameter systems
+ */
+UCLASS(MinimalAPI)
+class UMovieSceneMaterialParameterCollectionSystem
+	: public UMovieSceneEntitySystem
+	, public IMovieScenePreAnimatedStateSystemInterface
+{
+public:
+
+	GENERATED_BODY()
+
+	UMovieSceneMaterialParameterCollectionSystem(const FObjectInitializer& ObjInit);
+
+private:
+
+	virtual void OnLink() override;
+	virtual void OnRun(FSystemTaskPrerequisites& InPrerequisites, FSystemSubsequentTasks& Subsequents) override;
+
+	virtual void SavePreAnimatedState(const FPreAnimationParameters& InParameters) override;
+
+private:
+
+	/** Holds pre-animated values for scalar MPC values */
+	TSharedPtr<UE::MovieScene::FPreAnimatedMPCScalarStorage> ScalarParameterStorage;
+	/** Holds pre-animated values for vector or color MPC values */
+	TSharedPtr<UE::MovieScene::FPreAnimatedMPCVectorStorage> VectorParameterStorage;
+};
diff --git a/Engine/Source/Runtime/MovieSceneTracks/Public/Systems/MovieSceneMaterialParameterSystem.h b/Engine/Source/Runtime/MovieSceneTracks/Public/Systems/MovieSceneMaterialParameterSystem.h
new file mode 100644
--- /dev/null
+++ b/Engine/Source/Runtime/MovieSceneTracks/Public/Systems/MovieSceneMaterialParameterSystem.h
@@ -0,0 +1,66 @@
+// Copyright Epic Games, Inc. All Rights Reserved.
+
+#pragma once
+
+#include "EntitySystem/MovieSceneEntitySystem.h"
+#include "EntitySystem/MovieSceneOverlappingEntityTracker.h"
+#include "Evaluation/PreAnimatedState/MovieScenePreAnimatedStateStorage.h"
+
+#include "MovieSceneMaterialParameterSystem.generated.h"
+
+class UMovieScenePiecewiseFloatBlenderSystem;
+
+namespace UE::MovieScene
+{
+
+struct FAnimatedMaterialParameterInfo
+{
+	int32 NumContributors = 0;
+	FMovieSceneEntityID OutputEntityID;
+	FMovieSceneBlendChannelID BlendChannelID;
+};
+
+} // namespace UE::MovieScene
+
+
+/**
+ * System responsible for tracking and animating material parameter entities.
+ * Operates on the following component types from FMovieSceneTracksComponentTypes:
+ *
+ * Instantiation: Tracks any BoundMaterial with a ScalarParameterName, ColorParameterName or VectorParameterName.
+ *                Manages adding BlendChannelInputs and Outputs where multiple entities animate the same parameter
+ *                on the same bound material.
+ *                BoundMaterials may be a UMaterialInstanceDynamic, or a UMaterialParameterCollectionInstance.
+ *
+ * Evaluation:    Visits any BoundMaterial with the supported parameter names and either a BlendChannelOutput component
+ *                or no BlendChannelInput, and applies the resulting parameter to the bound material instance.
+ */
+UCLASS(MinimalAPI)
+class UMovieSceneMaterialParameterSystem
+	: public UMovieSceneEntitySystem
+{
+public:
+
+	GENERATED_BODY()
+
+	UMovieSceneMaterialParameterSystem(const FObjectInitializer& ObjInit);
+
+private:
+
+	virtual void OnLink() override;
+	virtual void OnRun(FSystemTaskPrerequisites& InPrerequisites, FSystemSubsequentTasks& Subsequents) override;
+
+	void OnInstantiation();
+	void OnEvaluation(FSystemTaskPrerequisites& InPrerequisites, FSystemSubsequentTasks& Subsequents);
+
+private:
+
+	/** Overlapping trackers that track multiple entities animating the same bound object and name */
+	UE::MovieScene::TOverlappingEntityTracker<UE::MovieScene::FAnimatedMaterialParameterInfo, UObject*, FName> ScalarParameterTracker;
+	UE::MovieScene::TOverlappingEntityTracker<UE::MovieScene::FAnimatedMaterialParameterInfo, UObject*, FName> VectorParameterTracker;
+
+public:
+
+	UPROPERTY()
+	TObjectPtr<UMovieScenePiecewiseFloatBlenderSystem> FloatBlenderSystem;
+};
diff --git a/Engine/Source/Runtime/MovieSceneTracks/Public/Systems/MovieSceneMaterialSystem.h b/Engine/Source/Runtime/MovieSceneTracks/Public/Systems/MovieSceneMaterialSystem.h
new file mode 100644
--- /dev/null
+++ b/Engine/Source/Runtime/MovieSceneTracks/Public/Systems/MovieSceneMaterialSystem.h
@@ -0,0 +1,318 @@
+// Copyright Epic Games, Inc. All Rights Reserved.
+
+#pragma once
+
+#include "EntitySystem/MovieSceneEntitySystem.h"
+#include "EntitySystem/MovieScenePreAnimatedStateSystem.h"
+#include "EntitySystem/MovieSceneEntityMutations.h"
+
+#include "Materials/MaterialInstance.h"
+#include "Materials/MaterialInstanceDynamic.h"
+
+#include "EntitySystem/BuiltInComponentTypes.h"
+#include "EntitySystem/MovieSceneEntitySystemRunner.h"
+#include "MovieSceneTracksComponentTypes.h"
+
+namespace UE::MovieScene
+{
+
+struct FPreAnimatedMaterialParameters
+{
+	UMaterialInterface* PreviousMaterial = nullptr;
+	UMaterialInterface* PreviousParameterContainer = nullptr;
+};
+
+template<typename AccessorType, typename... RequiredComponents>
+struct TPreAnimatedMaterialTraits : FBoundObjectPreAnimatedStateTraits
+{
+	using KeyType     = typename AccessorType::KeyType;
+	using StorageType = UMaterialInterface*;
+
+	static UMaterialInterface* CachePreAnimatedValue(typename TCallTraits<RequiredComponents>::ParamType... InRequiredComponents)
+	{
+		return AccessorType{ InRequiredComponents... }.GetMaterial();
+	}
+
+	static void RestorePreAnimatedValue(const KeyType& InKeyType, UMaterialInterface* OldMaterial, const FRestoreStateParams& Params)
+	{
+		AccessorType{ InKeyType }.SetMaterial(OldMaterial);
+	}
+};
+
+template<typename AccessorType, typename... RequiredComponents>
+struct TPreAnimatedMaterialParameterTraits : FBoundObjectPreAnimatedStateTraits
+{
+	using KeyType     = typename AccessorType::KeyType;
+	using StorageType = FPreAnimatedMaterialParameters;
+
+	static FPreAnimatedMaterialParameters CachePreAnimatedValue(typename TCallTraits<RequiredComponents>::ParamType... InRequiredComponents)
+	{
+		AccessorType Accessor{ InRequiredComponents... };
+
+		FPreAnimatedMaterialParameters Parameters;
+		Parameters.PreviousMaterial = Accessor.GetMaterial();
+
+		// If the current material we're overriding is already a material instance dynamic, copy it since we will be modifying the data. 
+		// The copied material will be used to restore the values in RestoreState.
+		UMaterialInstanceDynamic* MID = Cast<UMaterialInstanceDynamic>(Parameters.PreviousMaterial);
+		if (MID)
+		{
+			Parameters.PreviousParameterContainer = DuplicateObject<UMaterialInterface>(MID, MID->GetOuter());
+		}
+
+		return Parameters;
+	}
+
+	static void RestorePreAnimatedValue(const KeyType& InKey, const FPreAnimatedMaterialParameters& PreAnimatedValue, const FRestoreStateParams& Params)
+	{
+		AccessorType Accessor{ InKey };
+		if (PreAnimatedValue.PreviousParameterContainer != nullptr)
+		{
+			// If we cached parameter values in CachePreAnimatedValue that means the previous material was already a MID
+			// and we probably did not replace it with a new one when resolving bound materials. Therefore we
+			// just copy the parameters back over without changing the material
+			UMaterialInstanceDynamic* CurrentMID = Cast<UMaterialInstanceDynamic>(Accessor.GetMaterial());
+			if (CurrentMID)
+			{
+				CurrentMID->CopyMaterialUniformParameters(PreAnimatedValue.PreviousParameterContainer);
+				return;
+			}
+		}
+
+		Accessor.SetMaterial(PreAnimatedValue.PreviousMaterial);
+	}
+};
+
+template<typename AccessorType, typename... RequiredComponents>
+class TMovieSceneMaterialSystem
+{
+public:
+
+	using MaterialSwitcherStorageType = TPreAnimatedStateStorage<TPreAnimatedMaterialTraits<AccessorType, RequiredComponents...>>;
+	using MaterialParameterStorageType = TPreAnimatedStateStorage<TPreAnimatedMaterialParameterTraits<AccessorType, RequiredComponents...>>;
+
+	TSharedPtr<MaterialSwitcherStorageType> MaterialSwitcherStorage;
+	TSharedPtr<MaterialParameterStorageType> MaterialParameterStorage;
+
+	void OnLink(UMovieSceneEntitySystemLinker* Linker, TComponentTypeID<RequiredComponents>... InRequiredComponents);
+	void OnRun(UMovieSceneEntitySystemLinker* Linker, TComponentTypeID<RequiredComponents>... InRequiredComponents, FSystemTaskPrerequisites& InPrerequisites, FSystemSubsequentTasks& Subsequents);
+
+	void SavePreAnimatedState(UMovieSceneEntitySystemLinker* Linker, TComponentTypeID<RequiredComponents>... InRequiredComponents, const IMovieScenePreAnimatedStateSystemInterface::FPreAnimationParameters& InParameters);
+
+protected:
+
+	UE::MovieScene::FEntityComponentFilter MaterialSwitcherFilter;
+	UE::MovieScene::FEntityComponentFilter MaterialParameterFilter;
+};
+
+template<typename AccessorType, typename... RequiredComponents>
+struct TApplyMaterialSwitchers
+{
+	static void ForEachEntity(typename TCallTraits<RequiredComponents>::ParamType... Inputs, UObject* ObjectResult)
+	{
+		// ObjectResult must be a material
+		UMaterialInterface* NewMaterial = Cast<UMaterialInterface>(ObjectResult);
+
+		AccessorType Accessor(Inputs...);
+
+		UMaterialInterface*        ExistingMaterial = Accessor.GetMaterial();
+		UMaterialInstanceDynamic*  ExistingMID      = Cast<UMaterialInstanceDynamic>(ExistingMaterial);
+
+		if (ExistingMID && ExistingMID->Parent && ExistingMID->Parent == NewMaterial)
+		{
+			// Do not re-assign materials when a dynamic instance is already assigned with the same parent (since that's basically the same material, just with animated parameters)
+			// This is required for supporting material switchers alongside parameter tracks
+			return;
+		}
+
+		Accessor.SetMaterial(NewMaterial);
+	}
+};
+
+template<typename AccessorType, typename... RequiredComponents>
+struct TInitializeBoundMaterials
+{
+	static void ForEachEntity(typename TCallTraits<RequiredComponents>::ParamType... Inputs, UObject*& OutDynamicMaterial)
+	{
+		AccessorType Accessor(Inputs...);
+
+		UMaterialInterface* ExistingMaterial = Accessor.GetMaterial();
+
+		if (!ExistingMaterial)
+		{
+			// The object no longer has a valid material assigned.
+			// Rather than null the bound material and cause all downstream systems to have to check for null
+			// We just leave it assigned to the previous MID, even if that won't have an effect any more
+			return;
+		}
+		else if (UMaterialInstanceDynamic* MID = Cast<UMaterialInstanceDynamic>(ExistingMaterial))
+		{
+			OutDynamicMaterial = MID;
+			return;
+		}
+
+		if (OutDynamicMaterial && OutDynamicMaterial->IsA<UMaterialInstanceDynamic>())
+		{
+			return;
+		}
+		
+		OutDynamicMaterial = Accessor.CreateDynamicMaterial(ExistingMaterial);
+	}
+};
+
+template<typename...>
+struct TAddBoundMaterialMutationImpl;
+
+template<typename AccessorType, typename... RequiredComponents, int... Indices>
+struct TAddBoundMaterialMutationImpl<AccessorType, TIntegerSequence<int, Indices...>, RequiredComponents...> : IMovieSceneEntityMutation
+{
+	TAddBoundMaterialMutationImpl(TComponentTypeID<RequiredComponents>... InRequiredComponents)
+		: ComponentTypes(InRequiredComponents...)
+	{
+		BuiltInComponents = FBuiltInComponentTypes::Get();
+		TracksComponents  = FMovieSceneTracksComponentTypes::Get();
+	}
+	virtual void CreateMutation(FEntityManager* EntityManager, FComponentMask* InOutEntityComponentTypes) const override
+	{
+		InOutEntityComponentTypes->Set(TracksComponents->BoundMaterial);
+	}
+	virtual void InitializeAllocation(FEntityAllocation* Allocation, const FComponentMask& AllocationType) const
+	{
+		TComponentWriter<UObject*> BoundMaterials = Allocation->WriteComponents(TracksComponents->BoundMaterial, FEntityAllocationWriteContext::NewAllocation());
+		InitializeAllocation(Allocation, BoundMaterials, Allocation->ReadComponents(ComponentTypes.template Get<Indices>())...);
+	}
+
+	void InitializeAllocation(FEntityAllocation* Allocation, UObject** OutBoundMaterials, const RequiredComponents*... InRequiredComponents) const
+	{
+		const int32 Num = Allocation->Num();
+		for (int32 Index = 0; Index < Num; ++Index)
+		{
+			OutBoundMaterials[Index] = nullptr;
+			TInitializeBoundMaterials<AccessorType, RequiredComponents...>::ForEachEntity(InRequiredComponents[Index]..., OutBoundMaterials[Index]);
+
+			// @todo: We could remove the entity if it is null to avoid having to check for null downstream
+			//        but that is a problem for some widgets that play animations in PreConstruct _before_
+			//        the material is able to resolve correctly.
+			// if (!ensureAlwaysMsgf(OutBoundMaterials[Index] != nullptr,
+			//	TEXT("Unable to resolve material for %s. Material parameter tracks on this object will not function until it is resolved."),
+			//	*AccessorType(InRequiredComponents[Index]...).ToString()))
+			//{
+			//	
+			//	InvalidEntities.Add(Allocation->GetRawEntityIDs()[Index]);
+			//}
+		}
+	}
+
+	void RemoveInvalidBoundMaterials(UMovieSceneEntitySystemLinker* InLinker)
+	{
+#if 0
+		// Currently disabled due to the comment in InitializeAllocation. Another approach would be to
+		// Tag these differently, but that would require us mutating the entity manager during evaluation
+		// if the material _does_ resolve in the future, and mutation during eval is not supported
+		for (FMovieSceneEntityID InvalidEntity : InvalidEntities)
+		{
+			InLinker->EntityManager.RemoveComponent(InvalidEntity, TracksComponents->BoundMaterial);
+		}
+#endif
+	}
+private:
+
+	FBuiltInComponentTypes* BuiltInComponents;
+	FMovieSceneTracksComponentTypes* TracksComponents;
+
+	/** Track entities that we added a BoundMaterial component to but were unable to resolve so we can remove the component. This should happen rarely */
+	mutable TArray<FMovieSceneEntityID> InvalidEntities;
+
+	TTuple<TComponentTypeID<RequiredComponents>...> ComponentTypes;
+};
+
+
+template<typename AccessorType, typename... RequiredComponents>
+void TMovieSceneMaterialSystem<AccessorType, RequiredComponents...>::OnLink(UMovieSceneEntitySystemLinker* Linker, TComponentTypeID<RequiredComponents>... InRequiredComponents)
+{
+	using namespace UE::MovieScene;
+
+	FBuiltInComponentTypes*          BuiltInComponents = FBuiltInComponentTypes::Get();
+	FMovieSceneTracksComponentTypes* TracksComponents  = FMovieSceneTracksComponentTypes::Get();
+
+	MaterialSwitcherFilter.Reset();
+	MaterialSwitcherFilter.All({ InRequiredComponents..., BuiltInComponents->ObjectResult });
+
+	// Currently the only supported entities that we initialize are ones that contain Scalar, Vector or Color parameters
+	// Imported entities are implicitly excluded by way of filtering by BoundObject, which do not exist on imported entities
+	MaterialParameterFilter.Reset();
+	MaterialParameterFilter.All({ InRequiredComponents... });
+	MaterialParameterFilter.Any({ TracksComponents->ScalarParameterName, TracksComponents->ColorParameterName, TracksComponents->VectorParameterName });
+}
+
+template<typename AccessorType, typename... RequiredComponents>
+void TMovieSceneMaterialSystem<AccessorType, RequiredComponents...>::OnRun(UMovieSceneEntitySystemLinker* Linker, TComponentTypeID<RequiredComponents>... InRequiredComponents, FSystemTaskPrerequisites& InPrerequisites, FSystemSubsequentTasks& Subsequents)
+{
+	using namespace UE::MovieScene;
+
+	FBuiltInComponentTypes*          BuiltInComponents = FBuiltInComponentTypes::Get();
+	FMovieSceneTracksComponentTypes* TracksComponents  = FMovieSceneTracksComponentTypes::Get();
+
+	ESystemPhase CurrentPhase = Linker->GetActiveRunner()->GetCurrentPhase();
+	if (CurrentPhase == ESystemPhase::Instantiation)
+	{
+		// Only mutate things that are tagged as requiring linking
+		FEntityComponentFilter Filter(MaterialParameterFilter);
+		Filter.All({ BuiltInComponents->Tags.NeedsLink});
+
+		// Initialize bound dynamic materials (for material parameters)
+		using MutationType = TAddBoundMaterialMutationImpl<AccessorType, TMakeIntegerSequence<int, sizeof...(RequiredComponents)>, RequiredComponents...>;
+		MutationType BindMaterialsMutation(InRequiredComponents...);
+
+		Linker->EntityManager.MutateAll(Filter, BindMaterialsMutation);
+		BindMaterialsMutation.RemoveInvalidBoundMaterials(Linker);
+	}
+	else if (CurrentPhase == ESystemPhase::Evaluation)
+	{
+		using ApplyMaterialSwitchers = TApplyMaterialSwitchers<AccessorType, RequiredComponents...>;
+		using InitializeBoundMaterials = TInitializeBoundMaterials<AccessorType, RequiredComponents...>;
+
+		// Apply material switchers
+		FEntityTaskBuilder()
+		.ReadAllOf(InRequiredComponents...)
+		.Read(BuiltInComponents->ObjectResult)
+		.SetDesiredThread(Linker->EntityManager.GetDispatchThread())
+		.template Dispatch_PerEntity<ApplyMaterialSwitchers>(&Linker->EntityManager, InPrerequisites, &Subsequents);
+
+		// Initialize bound dynamic materials
+		FEntityTaskBuilder()
+		.ReadAllOf(InRequiredComponents...)
+		.Write(TracksComponents->BoundMaterial)
+		.SetDesiredThread(Linker->EntityManager.GetDispatchThread())
+		.template Dispatch_PerEntity<InitializeBoundMaterials>(&Linker->EntityManager, InPrerequisites, &Subsequents);
+	}
+}
+
+template<typename AccessorType, typename... RequiredComponents>
+void TMovieSceneMaterialSystem<AccessorType, RequiredComponents...>::SavePreAnimatedState(UMovieSceneEntitySystemLinker* Linker, TComponentTypeID<RequiredComponents>... InRequiredComponents, const IMovieScenePreAnimatedStateSystemInterface::FPreAnimationParameters& InParameters)
+{
+	using namespace UE::MovieScene;
+
+	FBuiltInComponentTypes*          BuiltInComponents = FBuiltInComponentTypes::Get();
+	FMovieSceneTracksComponentTypes* TracksComponents  = FMovieSceneTracksComponentTypes::Get();
+
+	// If we have material results to apply save those as well
+	if (Linker->EntityManager.Contains(MaterialSwitcherFilter))
+	{
+		TSavePreAnimatedStateParams<RequiredComponents...> Params;
+		Params.AdditionalFilter = MaterialSwitcherFilter;
+
+		MaterialSwitcherStorage->BeginTrackingAndCachePreAnimatedValuesTask(Linker, Params, InRequiredComponents...);
+	}
+
+	// If we have bound materials to resolve save the current material
+	if (Linker->EntityManager.Contains(MaterialParameterFilter))
+	{
+		TSavePreAnimatedStateParams<RequiredComponents...> Params;
+		Params.AdditionalFilter = MaterialParameterFilter;
+
+		MaterialParameterStorage->BeginTrackingAndCachePreAnimatedValuesTask(Linker, Params, InRequiredComponents...);
+	}
+}
+
+} // namespace UE::MovieScene
\ No newline at end of file
diff --git a/Engine/Source/Runtime/MovieSceneTracks/Public/Tracks/MovieSceneMaterialParameterCollectionTrack.h b/Engine/Source/Runtime/MovieSceneTracks/Public/Tracks/MovieSceneMaterialParameterCollectionTrack.h
--- a/Engine/Source/Runtime/MovieSceneTracks/Public/Tracks/MovieSceneMaterialParameterCollectionTrack.h
+++ b/Engine/Source/Runtime/MovieSceneTracks/Public/Tracks/MovieSceneMaterialParameterCollectionTrack.h
@@ -5,6 +5,7 @@
 #include "CoreMinimal.h"
 #include "UObject/ObjectMacros.h"
 #include "Tracks/MovieSceneMaterialTrack.h"
+#include "EntitySystem/IMovieSceneEntityProvider.h"
 #include "MovieSceneMaterialParameterCollectionTrack.generated.h"
 
 class UMaterialParameterCollection;
@@ -15,7 +16,8 @@ class UMaterialParameterCollection;
 UCLASS()
 class MOVIESCENETRACKS_API UMovieSceneMaterialParameterCollectionTrack
 	: public UMovieSceneMaterialTrack
-	, public IMovieSceneTrackTemplateProducer
+	, public IMovieSceneEntityProvider
+	, public IMovieSceneParameterSectionExtender
 {
 public:
 
@@ -29,7 +31,14 @@ public:
 
 	virtual bool SupportsType(TSubclassOf<UMovieSceneSection> SectionClass) const override;
 	virtual UMovieSceneSection* CreateNewSection() override;
-	virtual FMovieSceneEvalTemplatePtr CreateTemplateForSection(const UMovieSceneSection& InSection) const override;
+
+	/*~ IMovieSceneEntityProvider */
+	virtual void ImportEntityImpl(UMovieSceneEntitySystemLinker* EntityLinker, const FEntityImportParams& Params, FImportedEntity* OutImportedEntity) override;
+	virtual bool PopulateEvaluationFieldImpl(const TRange<FFrameNumber>& EffectiveRange, const FMovieSceneEvaluationFieldEntityMetaData& InMetaData, FMovieSceneEntityComponentFieldBuilder* OutFieldBuilder) override;
+
+	/*~ IMovieSceneParameterSectionExtender */
+	virtual void ExtendEntityImpl(UMovieSceneEntitySystemLinker* EntityLinker, const UE::MovieScene::FEntityImportParams& Params, UE::MovieScene::FImportedEntity* OutImportedEntity) override;
+
 #if WITH_EDITORONLY_DATA
 	virtual FText GetDefaultDisplayName() const override;
 #endif
diff --git a/Engine/Source/Runtime/MovieSceneTracks/Public/Tracks/MovieSceneMaterialTrack.h b/Engine/Source/Runtime/MovieSceneTracks/Public/Tracks/MovieSceneMaterialTrack.h
--- a/Engine/Source/Runtime/MovieSceneTracks/Public/Tracks/MovieSceneMaterialTrack.h
+++ b/Engine/Source/Runtime/MovieSceneTracks/Public/Tracks/MovieSceneMaterialTrack.h
@@ -6,7 +6,7 @@
 #include "UObject/ObjectMacros.h"
 #include "Sections/MovieSceneParameterSection.h"
 #include "MovieSceneNameableTrack.h"
-#include "Compilation/IMovieSceneTrackTemplateProducer.h"
+#include "EntitySystem/IMovieSceneEntityProvider.h"
 #include "MovieSceneMaterialTrack.generated.h"
 
 /**
@@ -32,7 +32,8 @@ public:
 	virtual void RemoveSection(UMovieSceneSection& Section) override;
 	virtual void RemoveSectionAt(int32 SectionIndex) override;
 	virtual bool IsEmpty() const override;
-	virtual const TArray<UMovieSceneSection*>& GetAllSections() const override;
+	virtual const TArray<UMovieSceneSection*>& GetAllSections() const override;	
+	virtual bool SupportsMultipleRows() const override;
 
 public:
 
@@ -66,18 +67,23 @@ private:
 UCLASS(MinimalAPI)
 class UMovieSceneComponentMaterialTrack
 	: public UMovieSceneMaterialTrack
-	, public IMovieSceneTrackTemplateProducer
+	, public IMovieSceneEntityProvider
+	, public IMovieSceneParameterSectionExtender
 {
 	GENERATED_UCLASS_BODY()
 
 public:
 
 	// UMovieSceneTrack interface
+	virtual void AddSection(UMovieSceneSection& Section) override;
+
+	/*~ IMovieSceneEntityProvider */
+	virtual void ImportEntityImpl(UMovieSceneEntitySystemLinker* EntityLinker, const FEntityImportParams& Params, FImportedEntity* OutImportedEntity) override;
+	virtual bool PopulateEvaluationFieldImpl(const TRange<FFrameNumber>& EffectiveRange, const FMovieSceneEvaluationFieldEntityMetaData& InMetaData, FMovieSceneEntityComponentFieldBuilder* OutFieldBuilder) override;
 
-	static const uint16 EvaluationPriority = 1000;
+	/*~ IMovieSceneParameterSectionExtender */
+	virtual void ExtendEntityImpl(UMovieSceneEntitySystemLinker* EntityLinker, const UE::MovieScene::FEntityImportParams& Params, UE::MovieScene::FImportedEntity* OutImportedEntity) override;
 
-	virtual FMovieSceneEvalTemplatePtr CreateTemplateForSection(const UMovieSceneSection& InSection) const override;
-	virtual void PostCompile(FMovieSceneEvaluationTrack& OutTrack, const FMovieSceneTrackCompilerArgs& Args) const override;
 	virtual FName GetTrackName() const { return FName( *FString::FromInt(MaterialIndex) ); }
 
 #if WITH_EDITORONLY_DATA
diff --git a/Engine/Source/Runtime/MovieSceneTracks/Public/Tracks/MovieScenePrimitiveMaterialTrack.h b/Engine/Source/Runtime/MovieSceneTracks/Public/Tracks/MovieScenePrimitiveMaterialTrack.h
--- a/Engine/Source/Runtime/MovieSceneTracks/Public/Tracks/MovieScenePrimitiveMaterialTrack.h
+++ b/Engine/Source/Runtime/MovieSceneTracks/Public/Tracks/MovieScenePrimitiveMaterialTrack.h
@@ -4,11 +4,10 @@
 
 #include "MovieSceneTrack.h"
 #include "Tracks/MovieScenePropertyTrack.h"
-#include "Compilation/IMovieSceneTrackTemplateProducer.h"
 #include "MovieScenePrimitiveMaterialTrack.generated.h"
 
 UCLASS(MinimalAPI)
-class UMovieScenePrimitiveMaterialTrack : public UMovieScenePropertyTrack, public IMovieSceneTrackTemplateProducer
+class UMovieScenePrimitiveMaterialTrack : public UMovieScenePropertyTrack
 {
 public:
 
@@ -25,8 +24,6 @@ public:
 	/*~ UMovieSceneTrack interface */
 	virtual UMovieSceneSection* CreateNewSection() override;
 	virtual bool SupportsType(TSubclassOf<UMovieSceneSection> SectionClass) const override;
-	virtual FMovieSceneEvalTemplatePtr CreateTemplateForSection(const UMovieSceneSection& InSection) const override;
-	virtual void PostCompile(FMovieSceneEvaluationTrack& OutTrack, const FMovieSceneTrackCompilerArgs& Args) const override;
 
 private:
 	UPROPERTY()
diff --git a/Engine/Source/Runtime/UMG/Private/Animation/MovieSceneUMGComponentTypes.cpp b/Engine/Source/Runtime/UMG/Private/Animation/MovieSceneUMGComponentTypes.cpp
--- a/Engine/Source/Runtime/UMG/Private/Animation/MovieSceneUMGComponentTypes.cpp
+++ b/Engine/Source/Runtime/UMG/Private/Animation/MovieSceneUMGComponentTypes.cpp
@@ -69,6 +69,7 @@
 	CastChecked<UWidget>(Object)->SetRenderTransform(Transform);
 }
 
+
 FMovieSceneUMGComponentTypes::FMovieSceneUMGComponentTypes()
 {
 	FComponentRegistry* ComponentRegistry = UMovieSceneEntitySystemLinker::GetComponents();
@@ -77,6 +78,8 @@
 
 	ComponentRegistry->NewPropertyType(WidgetTransform, TEXT("FWidgetTransform Property"));
 
+	ComponentRegistry->NewComponentType(&WidgetMaterialPath, TEXT("Widget Material Path"), EComponentTypeFlags::CopyToChildren | EComponentTypeFlags::CopyToOutput);
+
 	FBuiltInComponentTypes* BuiltInComponents = FBuiltInComponentTypes::Get();
 
 	FMovieSceneTracksComponentTypes::Get()->Accessors.Float.Add(UWidget::StaticClass(), "RenderOpacity", &GetRenderOpacity, &SetRenderOpacity);
diff --git a/Engine/Source/Runtime/UMG/Private/Animation/MovieSceneWidgetMaterialSystem.cpp b/Engine/Source/Runtime/UMG/Private/Animation/MovieSceneWidgetMaterialSystem.cpp
new file mode 100644
--- /dev/null
+++ b/Engine/Source/Runtime/UMG/Private/Animation/MovieSceneWidgetMaterialSystem.cpp
@@ -0,0 +1,129 @@
+// Copyright Epic Games, Inc. All Rights Reserved.
+
+#include "Animation/MovieSceneWidgetMaterialSystem.h"
+#include "Animation/MovieSceneUMGComponentTypes.h"
+#include "Animation/WidgetMaterialTrackUtilities.h"
+
+#include "EntitySystem/MovieSceneEntityMutations.h"
+
+#include "Systems/FloatChannelEvaluatorSystem.h"
+#include "Systems/MovieScenePiecewiseFloatBlenderSystem.h"
+
+#include "Materials/MaterialInstanceDynamic.h"
+
+#include "Components/Widget.h"
+
+#include "String/Join.h"
+
+namespace UE::MovieScene
+{
+
+FWidgetMaterialAccessor::FWidgetMaterialAccessor(const FWidgetMaterialKey& InKey)
+	: Widget(CastChecked<UWidget>(InKey.Object.ResolveObjectPtr()))
+	, WidgetMaterialPath(InKey.WidgetMaterialPath)
+{}
+
+FWidgetMaterialAccessor::FWidgetMaterialAccessor(UObject* InObject, FWidgetMaterialPath InWidgetMaterialPath)
+	: Widget(Cast<UWidget>(InObject))
+	, WidgetMaterialPath(MoveTemp(InWidgetMaterialPath))
+{
+	// Object must be a widget
+	check(!InObject || Widget);
+}
+
+FString FWidgetMaterialAccessor::ToString() const
+{
+	return FString::Printf(TEXT("Brush %s.%s"), *Widget->GetPathName(), *FString::JoinBy(WidgetMaterialPath.Path, TEXT("."), UE_PROJECTION_MEMBER(FName, ToString)));
+}
+
+UMaterialInterface* FWidgetMaterialAccessor::GetMaterial() const
+{
+	FWidgetMaterialHandle Handle = WidgetMaterialTrackUtilities::GetMaterialHandle(Widget, WidgetMaterialPath.Path);
+	if (Handle.IsValid())
+	{
+		return Handle.GetMaterial();
+	}
+
+	return nullptr;
+}
+
+void FWidgetMaterialAccessor::SetMaterial(UMaterialInterface* InMaterial) const
+{
+	FWidgetMaterialHandle Handle = WidgetMaterialTrackUtilities::GetMaterialHandle(Widget, WidgetMaterialPath.Path);
+	if (Handle.IsValid())
+	{
+		Handle.SetMaterial(InMaterial, Widget);
+	}
+}
+
+UMaterialInstanceDynamic* FWidgetMaterialAccessor::CreateDynamicMaterial(UMaterialInterface* InMaterial)
+{
+	// Need to create a new MID, either because the parent has changed, or because one doesn't already exist
+	TStringBuilder<128> DynamicName;
+	InMaterial->GetFName().ToString(DynamicName);
+	DynamicName.Append(TEXT("_Animated"));
+	FName UniqueDynamicName = MakeUniqueObjectName(Widget, UMaterialInstanceDynamic::StaticClass() , DynamicName.ToString());
+
+	UMaterialInstanceDynamic* MID = UMaterialInstanceDynamic::Create(InMaterial, Widget, UniqueDynamicName);
+	SetMaterial(MID);
+	return MID;
+}
+
+TAutoRegisterPreAnimatedStorageID<FPreAnimatedWidgetMaterialSwitcherStorage> FPreAnimatedWidgetMaterialSwitcherStorage::StorageID;
+TAutoRegisterPreAnimatedStorageID<FPreAnimatedWidgetMaterialParameterStorage> FPreAnimatedWidgetMaterialParameterStorage::StorageID;
+
+} // namespace UE::MovieScene
+
+UMovieSceneWidgetMaterialSystem::UMovieSceneWidgetMaterialSystem(const FObjectInitializer& ObjInit)
+	: Super(ObjInit)
+{
+	using namespace UE::MovieScene;
+
+	FBuiltInComponentTypes*          BuiltInComponents = FBuiltInComponentTypes::Get();
+	FMovieSceneUMGComponentTypes*    WidgetComponents  = FMovieSceneUMGComponentTypes::Get();
+	FMovieSceneTracksComponentTypes* TracksComponents  = FMovieSceneTracksComponentTypes::Get();
+
+	RelevantComponent = WidgetComponents->WidgetMaterialPath;
+	Phase = ESystemPhase::Instantiation | ESystemPhase::Evaluation;
+
+	if (HasAnyFlags(RF_ClassDefaultObject))
+	{
+		DefineComponentConsumer(GetClass(), BuiltInComponents->ObjectResult);
+		DefineComponentConsumer(GetClass(), BuiltInComponents->BoundObject);
+		DefineComponentProducer(GetClass(), TracksComponents->BoundMaterial);
+		DefineImplicitPrerequisite(UMovieSceneCachePreAnimatedStateSystem::StaticClass(), GetClass());
+	}
+}
+
+void UMovieSceneWidgetMaterialSystem::OnLink()
+{
+	using namespace UE::MovieScene;
+
+	FBuiltInComponentTypes*       BuiltInComponents = FBuiltInComponentTypes::Get();
+	FMovieSceneUMGComponentTypes* WidgetComponents  = FMovieSceneUMGComponentTypes::Get();
+
+	SystemImpl.MaterialSwitcherStorage = Linker->PreAnimatedState.GetOrCreateStorage<FPreAnimatedWidgetMaterialSwitcherStorage>();
+	SystemImpl.MaterialParameterStorage = Linker->PreAnimatedState.GetOrCreateStorage<FPreAnimatedWidgetMaterialParameterStorage>();
+
+	SystemImpl.OnLink(Linker, BuiltInComponents->BoundObject, WidgetComponents->WidgetMaterialPath);
+}
+
+void UMovieSceneWidgetMaterialSystem::OnRun(FSystemTaskPrerequisites& InPrerequisites, FSystemSubsequentTasks& Subsequents)
+{
+	using namespace UE::MovieScene;
+
+	FBuiltInComponentTypes*       BuiltInComponents = FBuiltInComponentTypes::Get();
+	FMovieSceneUMGComponentTypes* WidgetComponents  = FMovieSceneUMGComponentTypes::Get();
+
+	SystemImpl.OnRun(Linker, BuiltInComponents->BoundObject, WidgetComponents->WidgetMaterialPath, InPrerequisites, Subsequents);
+}
+
+void UMovieSceneWidgetMaterialSystem::SavePreAnimatedState(const FPreAnimationParameters& InParameters)
+{
+	using namespace UE::MovieScene;
+
+	FBuiltInComponentTypes*       BuiltInComponents = FBuiltInComponentTypes::Get();
+	FMovieSceneUMGComponentTypes* WidgetComponents  = FMovieSceneUMGComponentTypes::Get();
+
+	SystemImpl.SavePreAnimatedState(Linker, BuiltInComponents->BoundObject, WidgetComponents->WidgetMaterialPath, InParameters);
+}
\ No newline at end of file
diff --git a/Engine/Source/Runtime/UMG/Private/Animation/MovieSceneWidgetMaterialTemplate.cpp b/Engine/Source/Runtime/UMG/Private/Animation/MovieSceneWidgetMaterialTemplate.cpp
deleted file mode 100644
diff --git a/Engine/Source/Runtime/UMG/Private/Animation/MovieSceneWidgetMaterialTemplate.h b/Engine/Source/Runtime/UMG/Private/Animation/MovieSceneWidgetMaterialTemplate.h
deleted file mode 100644
diff --git a/Engine/Source/Runtime/UMG/Private/Animation/MovieSceneWidgetMaterialTrack.cpp b/Engine/Source/Runtime/UMG/Private/Animation/MovieSceneWidgetMaterialTrack.cpp
--- a/Engine/Source/Runtime/UMG/Private/Animation/MovieSceneWidgetMaterialTrack.cpp
+++ b/Engine/Source/Runtime/UMG/Private/Animation/MovieSceneWidgetMaterialTrack.cpp
@@ -2,17 +2,67 @@
 
 #include "Animation/MovieSceneWidgetMaterialTrack.h"
 #include "Animation/WidgetMaterialTrackUtilities.h"
-#include "Animation/MovieSceneWidgetMaterialTemplate.h"
+#include "Animation/MovieSceneUMGComponentTypes.h"
+#include "EntitySystem/BuiltInComponentTypes.h"
 
 UMovieSceneWidgetMaterialTrack::UMovieSceneWidgetMaterialTrack( const FObjectInitializer& ObjectInitializer )
 	: Super(ObjectInitializer)
 {
+	BuiltInTreePopulationMode = ETreePopulationMode::Blended;
 }
 
+void UMovieSceneWidgetMaterialTrack::AddSection(UMovieSceneSection& Section)
+{
+	// Materials are always blendable now
+	Section.SetBlendType(EMovieSceneBlendType::Absolute);
+	Super::AddSection(Section);
+}
+
+void UMovieSceneWidgetMaterialTrack::ImportEntityImpl(UMovieSceneEntitySystemLinker* EntityLinker, const FEntityImportParams& Params, FImportedEntity* OutImportedEntity)
+{
+	// These tracks don't define any entities for themselves
+	checkf(false, TEXT("This track should never have created entities for itself - this assertion indicates an error in the entity-component field"));
+}
+
+void UMovieSceneWidgetMaterialTrack::ExtendEntityImpl(UMovieSceneEntitySystemLinker* EntityLinker, const UE::MovieScene::FEntityImportParams& Params, UE::MovieScene::FImportedEntity* OutImportedEntity)
+{
+	using namespace UE::MovieScene;
 
-FMovieSceneEvalTemplatePtr UMovieSceneWidgetMaterialTrack::CreateTemplateForSection(const UMovieSceneSection& InSection) const
+	FBuiltInComponentTypes*       BuiltInComponents = FBuiltInComponentTypes::Get();
+	FMovieSceneUMGComponentTypes* WidgetComponents  = FMovieSceneUMGComponentTypes::Get();
+
+	// Material parameters are always absolute blends for the time being
+	OutImportedEntity->AddBuilder(
+		FEntityBuilder()
+		.Add(WidgetComponents->WidgetMaterialPath, FWidgetMaterialPath(BrushPropertyNamePath))
+		.AddTag(BuiltInComponents->Tags.AbsoluteBlend)
+	);
+}
+
+bool UMovieSceneWidgetMaterialTrack::PopulateEvaluationFieldImpl(const TRange<FFrameNumber>& EffectiveRange, const FMovieSceneEvaluationFieldEntityMetaData& InMetaData, FMovieSceneEntityComponentFieldBuilder* OutFieldBuilder)
 {
-	return FMovieSceneWidgetMaterialSectionTemplate(*CastChecked<UMovieSceneParameterSection>(&InSection), *this);
+	const FMovieSceneTrackEvaluationField& LocalEvaluationField = GetEvaluationField();
+
+	// Define entities for every entry in our evaluation field
+	for (const FMovieSceneTrackEvaluationFieldEntry& Entry : LocalEvaluationField.Entries)
+	{
+		UMovieSceneParameterSection* ParameterSection = Cast<UMovieSceneParameterSection>(Entry.Section);
+		if (!ParameterSection || IsRowEvalDisabled(ParameterSection->GetRowIndex()))
+		{
+			continue;
+		}
+
+		TRange<FFrameNumber> SectionEffectiveRange = TRange<FFrameNumber>::Intersection(EffectiveRange, Entry.Range);
+		if (!SectionEffectiveRange.IsEmpty())
+		{
+			FMovieSceneEvaluationFieldEntityMetaData SectionMetaData = InMetaData;
+			SectionMetaData.Flags = Entry.Flags;
+
+			ParameterSection->ExternalPopulateEvaluationField(SectionEffectiveRange, SectionMetaData, OutFieldBuilder);
+		}
+	}
+
+	return true;
 }
 
 
diff --git a/Engine/Source/Runtime/UMG/Private/Animation/WidgetMaterialTrackUtilities.cpp b/Engine/Source/Runtime/UMG/Private/Animation/WidgetMaterialTrackUtilities.cpp
--- a/Engine/Source/Runtime/UMG/Private/Animation/WidgetMaterialTrackUtilities.cpp
+++ b/Engine/Source/Runtime/UMG/Private/Animation/WidgetMaterialTrackUtilities.cpp
@@ -143,7 +143,7 @@ void FWidgetMaterialHandle::SetMaterial(UMaterialInterface* InMaterial, UWidget*
 }
 
 
-FWidgetMaterialHandle GetPropertyValueByPath(void* DataObject, UStruct* PropertySource, const TArray<FName>& PropertyPath, int32 PathIndex )
+FWidgetMaterialHandle GetPropertyValueByPath(void* DataObject, UStruct* PropertySource, TArrayView<const FName> PropertyPath, int32 PathIndex )
 {
 	if ( DataObject != nullptr && PathIndex < PropertyPath.Num() )
 	{
@@ -184,7 +184,7 @@ FWidgetMaterialHandle GetPropertyValueByPath(void* DataObject, UStruct* Property
 	return FWidgetMaterialHandle();
 }
 
-FWidgetMaterialHandle WidgetMaterialTrackUtilities::GetMaterialHandle(UWidget* Widget, const TArray<FName>& BrushPropertyNamePath)
+FWidgetMaterialHandle WidgetMaterialTrackUtilities::GetMaterialHandle(UWidget* Widget, TArrayView<const FName> BrushPropertyNamePath)
 {
 	return GetPropertyValueByPath(Widget, Widget->GetClass(), BrushPropertyNamePath, 0);
 }
@@ -247,7 +247,7 @@ void WidgetMaterialTrackUtilities::GetMaterialBrushPropertyPaths( UWidget* Widge
 	GetMaterialBrushPropertyPathsRecursive( Widget, Widget->GetClass(), PropertyPath, MaterialBrushPropertyPaths );
 }
 
-FName WidgetMaterialTrackUtilities::GetTrackNameFromPropertyNamePath( const TArray<FName>& PropertyNamePath )
+FName WidgetMaterialTrackUtilities::GetTrackNameFromPropertyNamePath( TArrayView<const FName> PropertyNamePath )
 {
 	if ( PropertyNamePath.Num() == 0 )
 	{
diff --git a/Engine/Source/Runtime/UMG/Public/Animation/MovieSceneUMGComponentTypes.h b/Engine/Source/Runtime/UMG/Public/Animation/MovieSceneUMGComponentTypes.h
--- a/Engine/Source/Runtime/UMG/Public/Animation/MovieSceneUMGComponentTypes.h
+++ b/Engine/Source/Runtime/UMG/Public/Animation/MovieSceneUMGComponentTypes.h
@@ -9,12 +9,24 @@
 #include "EntitySystem/MovieScenePropertyTraits.h"
 #include "EntitySystem/MovieScenePropertyMetaDataTraits.h"
 
+#include "Containers/ArrayView.h"
+
 
 namespace UE
 {
 namespace MovieScene
 {
 
+struct FWidgetMaterialPath
+{
+	FWidgetMaterialPath() = default;
+	FWidgetMaterialPath(TArrayView<const FName> Names)
+		: Path(Names.GetData(), Names.Num())
+	{}
+
+	TArray<FName, TInlineAllocator<2>> Path;
+};
+
 struct FIntermediateWidgetTransform
 {
 	float TranslationX;
@@ -38,6 +50,8 @@ struct UMG_API FMovieSceneUMGComponentTypes
 	TPropertyComponents<FMarginTraits> Margin;
 	TPropertyComponents<FWidgetTransformPropertyTraits> WidgetTransform;
 
+	TComponentTypeID<FWidgetMaterialPath> WidgetMaterialPath;
+
 	TCustomPropertyRegistration<FWidgetTransformPropertyTraits, 1> CustomWidgetTransformAccessors;
 
 	static void Destroy();
diff --git a/Engine/Source/Runtime/UMG/Public/Animation/MovieSceneWidgetMaterialSystem.h b/Engine/Source/Runtime/UMG/Public/Animation/MovieSceneWidgetMaterialSystem.h
new file mode 100644
--- /dev/null
+++ b/Engine/Source/Runtime/UMG/Public/Animation/MovieSceneWidgetMaterialSystem.h
@@ -0,0 +1,94 @@
+// Copyright Epic Games, Inc. All Rights Reserved.
+
+#pragma once
+
+#include "EntitySystem/MovieSceneEntitySystem.h"
+#include "Evaluation/PreAnimatedState/MovieScenePreAnimatedStateStorage.h"
+
+#include "Systems/MovieSceneMaterialSystem.h"
+#include "Animation/MovieSceneUMGComponentTypes.h"
+
+#include "MovieSceneWidgetMaterialSystem.generated.h"
+
+class UWidget;
+class UMovieScenePiecewiseFloatBlenderSystem;
+
+namespace UE::MovieScene
+{
+
+struct FWidgetMaterialKey
+{
+	FObjectKey Object;
+	FWidgetMaterialPath WidgetMaterialPath;
+
+	friend uint32 GetTypeHash(const FWidgetMaterialKey& In)
+	{
+		uint32 Accumulator = GetTypeHash(In.Object);
+		for (const FName& Name : In.WidgetMaterialPath.Path)
+		{
+			Accumulator ^= GetTypeHash(Name);
+		}
+		return Accumulator;
+	}
+	friend bool operator==(const FWidgetMaterialKey& A, const FWidgetMaterialKey& B)
+	{
+		return A.Object == B.Object && A.WidgetMaterialPath.Path == B.WidgetMaterialPath.Path;
+	}
+};
+
+struct FWidgetMaterialAccessor
+{
+	using KeyType = FWidgetMaterialKey;
+
+	UWidget* Widget;
+	FWidgetMaterialPath WidgetMaterialPath;
+
+	FWidgetMaterialAccessor(const FWidgetMaterialKey& InKey);
+	FWidgetMaterialAccessor(UObject* InObject, FWidgetMaterialPath InWidgetMaterialPath);
+
+	UMaterialInterface* GetMaterial() const;
+	void SetMaterial(UMaterialInterface* InMaterial) const;
+	UMaterialInstanceDynamic* CreateDynamicMaterial(UMaterialInterface* InMaterial);
+	FString ToString() const;
+};
+
+using FPreAnimatedWidgetMaterialTraits          = TPreAnimatedMaterialTraits<FWidgetMaterialAccessor, UObject*, FWidgetMaterialPath>;
+using FPreAnimatedWidgetMaterialParameterTraits = TPreAnimatedMaterialParameterTraits<FWidgetMaterialAccessor, UObject*, FWidgetMaterialPath>;
+
+struct FPreAnimatedWidgetMaterialSwitcherStorage
+	: public TPreAnimatedStateStorage<TPreAnimatedMaterialTraits<FWidgetMaterialAccessor, UObject*, FWidgetMaterialPath>>
+{
+	static TAutoRegisterPreAnimatedStorageID<FPreAnimatedWidgetMaterialSwitcherStorage> StorageID;
+};
+
+struct FPreAnimatedWidgetMaterialParameterStorage
+	: public TPreAnimatedStateStorage<TPreAnimatedMaterialParameterTraits<FWidgetMaterialAccessor, UObject*, FWidgetMaterialPath>>
+{
+	static TAutoRegisterPreAnimatedStorageID<FPreAnimatedWidgetMaterialParameterStorage> StorageID;
+};
+
+} // namespace UE::MovieScene
+
+
+UCLASS(MinimalAPI)
+class UMovieSceneWidgetMaterialSystem
+	: public UMovieSceneEntitySystem
+	, public IMovieScenePreAnimatedStateSystemInterface
+{
+public:
+
+	GENERATED_BODY()
+
+	UMovieSceneWidgetMaterialSystem(const FObjectInitializer& ObjInit);
+
+private:
+
+	virtual void OnLink() override;
+	virtual void OnRun(FSystemTaskPrerequisites& InPrerequisites, FSystemSubsequentTasks& Subsequents) override;
+
+	virtual void SavePreAnimatedState(const FPreAnimationParameters& InParameters) override;
+
+private:
+
+	UE::MovieScene::TMovieSceneMaterialSystem<UE::MovieScene::FWidgetMaterialAccessor, UObject*, UE::MovieScene::FWidgetMaterialPath> SystemImpl;
+};
diff --git a/Engine/Source/Runtime/UMG/Public/Animation/MovieSceneWidgetMaterialTrack.h b/Engine/Source/Runtime/UMG/Public/Animation/MovieSceneWidgetMaterialTrack.h
--- a/Engine/Source/Runtime/UMG/Public/Animation/MovieSceneWidgetMaterialTrack.h
+++ b/Engine/Source/Runtime/UMG/Public/Animation/MovieSceneWidgetMaterialTrack.h
@@ -13,17 +13,24 @@
 UCLASS(MinimalAPI)
 class UMovieSceneWidgetMaterialTrack
 	: public UMovieSceneMaterialTrack
-	, public IMovieSceneTrackTemplateProducer
+	, public IMovieSceneEntityProvider
+	, public IMovieSceneParameterSectionExtender
 {
 	GENERATED_UCLASS_BODY()
 
 public:
 
 	// UMovieSceneTrack interface
-
-	virtual FMovieSceneEvalTemplatePtr CreateTemplateForSection(const UMovieSceneSection& InSection) const override;
+	virtual void AddSection(UMovieSceneSection& Section) override;
 	virtual FName GetTrackName() const override;
 
+	/*~ IMovieSceneEntityProvider */
+	virtual void ImportEntityImpl(UMovieSceneEntitySystemLinker* EntityLinker, const FEntityImportParams& Params, FImportedEntity* OutImportedEntity) override;
+	virtual bool PopulateEvaluationFieldImpl(const TRange<FFrameNumber>& EffectiveRange, const FMovieSceneEvaluationFieldEntityMetaData& InMetaData, FMovieSceneEntityComponentFieldBuilder* OutFieldBuilder) override;
+
+	/*~ IMovieSceneParameterSectionExtender */
+	virtual void ExtendEntityImpl(UMovieSceneEntitySystemLinker* EntityLinker, const UE::MovieScene::FEntityImportParams& Params, UE::MovieScene::FImportedEntity* OutImportedEntity) override;
+
 #if WITH_EDITORONLY_DATA
 	virtual FText GetDefaultDisplayName() const override;
 #endif
diff --git a/Engine/Source/Runtime/UMG/Public/Animation/WidgetMaterialTrackUtilities.h b/Engine/Source/Runtime/UMG/Public/Animation/WidgetMaterialTrackUtilities.h
--- a/Engine/Source/Runtime/UMG/Public/Animation/WidgetMaterialTrackUtilities.h
+++ b/Engine/Source/Runtime/UMG/Public/Animation/WidgetMaterialTrackUtilities.h
@@ -52,11 +52,11 @@ struct FWidgetMaterialPropertyPath
 namespace WidgetMaterialTrackUtilities
 {
 	/** Gets a material handle from a property on a widget by the properties FName path. */
-	UMG_API FWidgetMaterialHandle GetMaterialHandle(UWidget* Widget, const TArray<FName>& BrushPropertyNamePath);
+	UMG_API FWidgetMaterialHandle GetMaterialHandle(UWidget* Widget, TArrayView<const FName> BrushPropertyNamePath);
 
 	/** Gets the property paths on a widget which are slate brush properties, and who's slate brush has a valid material. */
 	UMG_API void GetMaterialBrushPropertyPaths( UWidget* Widget, TArray<FWidgetMaterialPropertyPath>& MaterialBrushPropertyPaths );
 
 	/** Converts a property name path into a single name which is appropriate for a track name. */
-	UMG_API FName GetTrackNameFromPropertyNamePath( const  TArray<FName>& PropertyNamePath );
+	UMG_API FName GetTrackNameFromPropertyNamePath( TArrayView<const FName> PropertyNamePath );
 }
