Fix compile errors with VS2019.

diff --git a/Engine/Source/Runtime/MovieSceneTracks/Public/Systems/MovieSceneComponentMaterialSystem.h b/Engine/Source/Runtime/MovieSceneTracks/Public/Systems/MovieSceneComponentMaterialSystem.h
--- a/Engine/Source/Runtime/MovieSceneTracks/Public/Systems/MovieSceneComponentMaterialSystem.h
+++ b/Engine/Source/Runtime/MovieSceneTracks/Public/Systems/MovieSceneComponentMaterialSystem.h
@@ -22,7 +22,7 @@ struct FComponentMaterialKey
 
 	friend uint32 GetTypeHash(const FComponentMaterialKey& In)
 	{
-		return GetTypeHash(In.Object) ^ GetTypeHash(In.MaterialIndex);
+		return GetTypeHash(In.Object) ^ ::GetTypeHash(In.MaterialIndex);
 	}
 	friend bool operator==(const FComponentMaterialKey& A, const FComponentMaterialKey& B)
 	{
