Remove float result components and the float blender system, do everything in double precision. Handle LWC float->double conversion during property instantiation.

diff --git a/Engine/Plugins/MovieScene/MatineeToLevelSequence/Source/MatineeToLevelSequence/Private/CameraAnimToTemplateSequenceConverter.cpp b/Engine/Plugins/MovieScene/MatineeToLevelSequence/Source/MatineeToLevelSequence/Private/CameraAnimToTemplateSequenceConverter.cpp
--- a/Engine/Plugins/MovieScene/MatineeToLevelSequence/Source/MatineeToLevelSequence/Private/CameraAnimToTemplateSequenceConverter.cpp
+++ b/Engine/Plugins/MovieScene/MatineeToLevelSequence/Source/MatineeToLevelSequence/Private/CameraAnimToTemplateSequenceConverter.cpp
@@ -526,7 +526,7 @@ UObject* FCameraAnimToTemplateSequenceConverter::ConvertSingleCameraAnimToTempla
 						NewMovieScene->GetPlaybackRange().GetLowerBoundValue() : FFrameNumber(0));
 				Interrogator.Update();
 
-				TArray<float> AnimatedFOVs;
+				TArray<double> AnimatedFOVs;
 				FMovieSceneTracksComponentTypes* TracksComponents = FMovieSceneTracksComponentTypes::Get();
 				Interrogator.QueryPropertyValues(TracksComponents->Float, AnimatedFOVs);
 
diff --git a/Engine/Plugins/MovieScene/TemplateSequence/Source/TemplateSequence/Private/Systems/TemplateSequenceSystem.cpp b/Engine/Plugins/MovieScene/TemplateSequence/Source/TemplateSequence/Private/Systems/TemplateSequenceSystem.cpp
--- a/Engine/Plugins/MovieScene/TemplateSequence/Source/TemplateSequence/Private/Systems/TemplateSequenceSystem.cpp
+++ b/Engine/Plugins/MovieScene/TemplateSequence/Source/TemplateSequence/Private/Systems/TemplateSequenceSystem.cpp
@@ -20,7 +20,6 @@
 #include "Systems/MovieSceneComponentTransformSystem.h"
 #include "Systems/MovieSceneFloatPropertySystem.h"
 #include "Systems/MovieScenePiecewiseDoubleBlenderSystem.h"
-#include "Systems/MovieScenePiecewiseFloatBlenderSystem.h"
 #include "TemplateSequence.h"
 #include "TemplateSequenceComponentTypes.h"
 
@@ -271,7 +270,6 @@ UTemplateSequencePropertyScalingEvaluatorSystem::UTemplateSequencePropertyScalin
 		DefineImplicitPrerequisite(UDoubleChannelEvaluatorSystem::StaticClass(), GetClass());
 		DefineImplicitPrerequisite(UMovieSceneBaseValueEvaluatorSystem::StaticClass(), GetClass());
 		// We need to multiply values before they are blended together with other values we shouldn't touch.
-		DefineImplicitPrerequisite(GetClass(), UMovieScenePiecewiseFloatBlenderSystem::StaticClass());
 		DefineImplicitPrerequisite(GetClass(), UMovieScenePiecewiseDoubleBlenderSystem::StaticClass());
 		// We need to multiply values before they are set on the animated objects. For blended values, the dependencies above are
 		// enough, but for non-blended values that use the "fast-path", we need this dependency.
@@ -448,8 +446,8 @@ struct FScaleFloatProperties
 			TRead<FInstanceHandle> InstanceHandles,
 			TRead<FMovieScenePropertyBinding> PropertyBindings,
 			TRead<FGuid> ReverseBindingLookups,
-			TReadOptional<float> BasePropertyValues,
-			TWrite<float> PropertyValues,
+			TReadOptional<double> BasePropertyValues,
+			TWrite<double> PropertyValues,
 			TWrite<FSourceFloatChannelFlags> PropertyFlags)
 	{
 		for (int32 Index = 0; Index < Allocation->Num(); ++Index)
@@ -507,7 +505,7 @@ void UTemplateSequencePropertyScalingEvaluatorSystem::OnRun(FSystemTaskPrerequis
 	FGraphEventRef GatherTask = FEntityTaskBuilder()
 		.Read(BuiltInComponents->InstanceHandle)
 		.Read(TemplateSequenceComponents->PropertyScale)
-		.Read(BuiltInComponents->FloatResult[0])
+		.Read(BuiltInComponents->DoubleResult[0])
 		.Dispatch_PerEntity<FGatherPropertyScales>(&Linker->EntityManager, InPrerequisites, &Subsequents, this);
 	if (GatherTask)
 	{
@@ -551,8 +549,8 @@ void UTemplateSequencePropertyScalingEvaluatorSystem::OnRun(FSystemTaskPrerequis
 			.Read(BuiltInComponents->InstanceHandle)
 			.Read(BuiltInComponents->PropertyBinding)
 			.Read(TemplateSequenceComponents->PropertyScaleReverseBindingLookup)
-			.ReadOptional(BuiltInComponents->BaseFloat[0])
-			.Write(BuiltInComponents->FloatResult[0])
+			.ReadOptional(BuiltInComponents->BaseDouble[0])
+			.Write(BuiltInComponents->DoubleResult[0])
 			.Write(BuiltInComponents->FloatChannelFlags[0])
 			.FilterAll({ TrackComponents->Float.PropertyTag })
 			.Dispatch_PerAllocation<FScaleFloatProperties>(&Linker->EntityManager, InPrerequisites, &Subsequents, this);
diff --git a/Engine/Source/Runtime/MovieScene/Private/EntitySystem/BuiltInComponentTypes.cpp b/Engine/Source/Runtime/MovieScene/Private/EntitySystem/BuiltInComponentTypes.cpp
--- a/Engine/Source/Runtime/MovieScene/Private/EntitySystem/BuiltInComponentTypes.cpp
+++ b/Engine/Source/Runtime/MovieScene/Private/EntitySystem/BuiltInComponentTypes.cpp
@@ -96,15 +96,6 @@ FBuiltInComponentTypes::FBuiltInComponentTypes()
 	ComponentRegistry->NewComponentType(&BoolResult,            TEXT("Bool Result"));
 	ComponentRegistry->NewComponentType(&ByteResult,            TEXT("Byte Result"));
 	ComponentRegistry->NewComponentType(&IntegerResult,         TEXT("Integer Result"));
-	ComponentRegistry->NewComponentType(&FloatResult[0],        TEXT("Float Result 0"));
-	ComponentRegistry->NewComponentType(&FloatResult[1],        TEXT("Float Result 1"));
-	ComponentRegistry->NewComponentType(&FloatResult[2],        TEXT("Float Result 2"));
-	ComponentRegistry->NewComponentType(&FloatResult[3],        TEXT("Float Result 3"));
-	ComponentRegistry->NewComponentType(&FloatResult[4],        TEXT("Float Result 4"));
-	ComponentRegistry->NewComponentType(&FloatResult[5],        TEXT("Float Result 5"));
-	ComponentRegistry->NewComponentType(&FloatResult[6],        TEXT("Float Result 6"));
-	ComponentRegistry->NewComponentType(&FloatResult[7],        TEXT("Float Result 7"));
-	ComponentRegistry->NewComponentType(&FloatResult[8],        TEXT("Float Result 8"));
 	ComponentRegistry->NewComponentType(&DoubleResult[0],       TEXT("Double Result 0"));
 	ComponentRegistry->NewComponentType(&DoubleResult[1],       TEXT("Double Result 1"));
 	ComponentRegistry->NewComponentType(&DoubleResult[2],       TEXT("Double Result 2"));
@@ -117,15 +108,6 @@ FBuiltInComponentTypes::FBuiltInComponentTypes()
 	ComponentRegistry->NewComponentType(&ObjectResult,          TEXT("Object Result"));
 
 	ComponentRegistry->NewComponentType(&BaseInteger,			TEXT("Base Integer"));
-	ComponentRegistry->NewComponentType(&BaseFloat[0],          TEXT("Base Float 0"));
-	ComponentRegistry->NewComponentType(&BaseFloat[1],          TEXT("Base Float 1"));
-	ComponentRegistry->NewComponentType(&BaseFloat[2],          TEXT("Base Float 2"));
-	ComponentRegistry->NewComponentType(&BaseFloat[3],          TEXT("Base Float 3"));
-	ComponentRegistry->NewComponentType(&BaseFloat[4],          TEXT("Base Float 4"));
-	ComponentRegistry->NewComponentType(&BaseFloat[5],          TEXT("Base Float 5"));
-	ComponentRegistry->NewComponentType(&BaseFloat[6],          TEXT("Base Float 6"));
-	ComponentRegistry->NewComponentType(&BaseFloat[7],          TEXT("Base Float 7"));
-	ComponentRegistry->NewComponentType(&BaseFloat[8],          TEXT("Base Float 8"));
 	ComponentRegistry->NewComponentType(&BaseDouble[0],         TEXT("Base Double 0"));
 	ComponentRegistry->NewComponentType(&BaseDouble[1],         TEXT("Base Double 1"));
 	ComponentRegistry->NewComponentType(&BaseDouble[2],         TEXT("Base Double 2"));
@@ -229,7 +211,7 @@ FBuiltInComponentTypes::FBuiltInComponentTypes()
 	// Float channel relationships
 	{
 		static_assert(
-				UE_ARRAY_COUNT(FloatChannel) == UE_ARRAY_COUNT(BaseFloat), 
+				UE_ARRAY_COUNT(FloatChannel) == UE_ARRAY_COUNT(BaseDouble), 
 				"Base floats and float results should have the same size.");
 		static_assert(
 				UE_ARRAY_COUNT(FloatChannel) == UE_ARRAY_COUNT(FloatChannelFlags),
@@ -239,17 +221,17 @@ FBuiltInComponentTypes::FBuiltInComponentTypes()
 		for (int32 Index = 0; Index < UE_ARRAY_COUNT(FloatChannel); ++Index)
 		{
 			ComponentRegistry->Factories.DuplicateChildComponent(FloatChannel[Index]);
-			ComponentRegistry->Factories.DefineMutuallyInclusiveComponent(FloatChannel[Index], FloatResult[Index]);
+			ComponentRegistry->Factories.DefineMutuallyInclusiveComponent(FloatChannel[Index], DoubleResult[Index]);
 			ComponentRegistry->Factories.DefineMutuallyInclusiveComponent(FloatChannel[Index], EvalTime);
 			ComponentRegistry->Factories.DefineMutuallyInclusiveComponent(FloatChannel[Index], FloatChannelFlags[Index]);
 		}
 
 		// Create base float components for float channels that are meant to be additive from base.
-		for (int32 Index = 0; Index < UE_ARRAY_COUNT(BaseFloat); ++Index)
+		for (int32 Index = 0; Index < UE_ARRAY_COUNT(BaseDouble); ++Index)
 		{
 			ComponentRegistry->Factories.DefineComplexInclusiveComponents(
 					FComplexInclusivityFilter::All({ FloatChannel[Index], BaseValueEvalTime, Tags.AdditiveFromBaseBlend }),
-					BaseFloat[Index]);
+					BaseDouble[Index]);
 		}
 	}
 
diff --git a/Engine/Source/Runtime/MovieScene/Public/EntitySystem/BuiltInComponentTypes.h b/Engine/Source/Runtime/MovieScene/Public/EntitySystem/BuiltInComponentTypes.h
--- a/Engine/Source/Runtime/MovieScene/Public/EntitySystem/BuiltInComponentTypes.h
+++ b/Engine/Source/Runtime/MovieScene/Public/EntitySystem/BuiltInComponentTypes.h
@@ -286,15 +286,9 @@ public:
 	// FMovieSceneObjectPathChannel that represents a changing object path over time
 	TComponentTypeID<FSourceObjectPathChannel> ObjectPathChannel;
 
-	// A float representing the output of the channel considered to be at index N within the source structure (ie 0 = Location.X, Vector.X, Color.R; 1 = Location.Y, Vector.Y, Color.G)
-	TComponentTypeID<float> FloatResult[9];
-
 	// A double considered to be at index N within the source structure (ie 0 = Location.X, Vector.X; 1 = Location.Y, Vector.Y)
 	TComponentTypeID<double> DoubleResult[9];
 
-	// A float representing the base value for the float channel at index N, for the purposes of "additive from base" blending.
-	TComponentTypeID<float> BaseFloat[9];
-
 	// A double representing the base value for the double channel at index N, for the purposes of "additive from base" blending.
 	TComponentTypeID<double> BaseDouble[9];
 
@@ -302,7 +296,7 @@ public:
 	TComponentTypeID<FFrameTime> BaseValueEvalTime;
 
 	// A float representing the evaluated output of a weight channel
-	TComponentTypeID<float> WeightResult;
+	TComponentTypeID<double> WeightResult;
 
 	// The result of an evaluated FMovieSceneObjectPathChannel
 	TComponentTypeID<UObject*> ObjectResult;
@@ -319,7 +313,7 @@ public:
 	TComponentTypeID<FMovieSceneSequenceID> HierarchicalEasingProvider;
 
 	// A float representing the evaluated easing weight
-	TComponentTypeID<float> WeightAndEasingResult;
+	TComponentTypeID<double> WeightAndEasingResult;
 
 	/** A blender type that should be used for blending this entity */
 	TComponentTypeID<TSubclassOf<UMovieSceneBlenderSystem>> BlenderType;
diff --git a/Engine/Source/Runtime/MovieScene/Public/EntitySystem/MovieScenePropertyComponentHandler.h b/Engine/Source/Runtime/MovieScene/Public/EntitySystem/MovieScenePropertyComponentHandler.h
--- a/Engine/Source/Runtime/MovieScene/Public/EntitySystem/MovieScenePropertyComponentHandler.h
+++ b/Engine/Source/Runtime/MovieScene/Public/EntitySystem/MovieScenePropertyComponentHandler.h
@@ -403,7 +403,7 @@
 		FGraphEventArray Tasks;
 		for (int32 Index = 0; Index < NumComposites; ++Index)
 		{
-			if ((PropertyDefinition.FloatCompositeMask & (1 << Index)) == 0 && (PropertyDefinition.DoubleCompositeMask & (1 << Index)) == 0)
+			if ((PropertyDefinition.DoubleCompositeMask & (1 << Index)) == 0)
 			{
 				continue;
 			}
@@ -437,23 +437,8 @@
 
 			for (int32 CompositeIndex = 0; CompositeIndex < NumComposites; ++CompositeIndex)
 			{
-				if ((PropertyDefinition.FloatCompositeMask & (1 << CompositeIndex)) != 0)
+				if ((PropertyDefinition.DoubleCompositeMask & ( 1 << Index)) != 0)
 				{
-					const float* InitialValueComposite = nullptr;
-					FAlignedDecomposedValue& AlignedOutput = AlignedOutputs[CompositeIndex];
-					if (InitialValueComponent)
-					{
-						const StorageType& InitialValue = (*InitialValueComponent);
-						InitialValueComposite = reinterpret_cast<const float*>(reinterpret_cast<const uint8*>(&InitialValue) + Composites[CompositeIndex].CompositeOffset);
-					}
-
-					const float NewComposite = *reinterpret_cast<const float*>(reinterpret_cast<const uint8*>(&InCurrentValue) + Composites[CompositeIndex].CompositeOffset);
-
-					float* RecomposedComposite = reinterpret_cast<float*>(Result + Composites[CompositeIndex].CompositeOffset);
-					*RecomposedComposite = AlignedOutput.Value.Recompose(EntityID, NewComposite, InitialValueComposite);
-				}
-				else if ((PropertyDefinition.DoubleCompositeMask & ( 1 << Index)) != 0)
-				{
 					const double* InitialValueComposite = nullptr;
 					FAlignedDecomposedValue& AlignedOutput = AlignedOutputs[CompositeIndex];
 					if (InitialValueComponent)
@@ -513,20 +498,8 @@
 		{
 			FMovieSceneEntityID EntityID = LocalParams.Query.Entities[Index];
 
-			if ((PropertyDefinition.FloatCompositeMask & (1 << CompositeIndex)) != 0)
+			if ((PropertyDefinition.DoubleCompositeMask & (1 << CompositeIndex)) != 0)
 			{
-				const float* InitialValueComposite = nullptr;
-				if (InitialValueComponent)
-				{
-					const StorageType& InitialValue = (*InitialValueComponent);
-					InitialValueComposite = reinterpret_cast<const float*>(reinterpret_cast<const uint8*>(&InitialValue) + Composite.CompositeOffset);
-				}
-
-				const float RecomposedComposite = AlignedOutput.Value.Recompose(EntityID, InCurrentValue, InitialValueComposite);
-				OutResults[Index] = (double)RecomposedComposite;
-			}
-			else if ((PropertyDefinition.DoubleCompositeMask & (1 << CompositeIndex)) != 0)
-			{
 				const double* InitialValueComposite = nullptr;
 				if (InitialValueComponent)
 				{
@@ -585,12 +558,11 @@
 		Registry->CompositeDefinitions.Add(NewChannel);
 
 		Definition->CompositeSize = 1;
-		if (TIsSame<typename PropertyTraits::StorageType, float>::Value)
+
+		static_assert(!TIsSame<typename PropertyTraits::StorageType, float>::Value, "Please use double-precision composites");
+
+		if (TIsSame<typename PropertyTraits::StorageType, double>::Value)
 		{
-			Definition->FloatCompositeMask = 1;
-		}
-		else if (TIsSame<typename PropertyTraits::StorageType, double>::Value)
-		{
 			Definition->DoubleCompositeMask = 1;
 		}
 
@@ -667,12 +639,10 @@
 		FPropertyCompositeDefinition NewChannel = { InComponent, static_cast<uint16>(CompositeOffset) };
 		Registry->CompositeDefinitions.Add(NewChannel);
 
-		if (TIsSame<T, float>::Value)
+		static_assert(!TIsSame<T, float>::Value, "Please use double-precision composites");
+
+		if (TIsSame<T, double>::Value)
 		{
-			Definition->FloatCompositeMask |= 1 << Definition->CompositeSize;
-		}
-		else if (TIsSame<T, double>::Value)
-		{
 			Definition->DoubleCompositeMask |= 1 << Definition->CompositeSize;
 		}
 
@@ -680,21 +650,6 @@
 		return TCompositePropertyDefinitionBuilder<PropertyTraits, Composites..., T>(Definition, Registry);
 	}
 
-	TCompositePropertyDefinitionBuilder<PropertyTraits, Composites..., float> AddComposite(TComponentTypeID<float> InComponent, float StorageType::*DataPtr)
-	{
-		checkf(Definition == &Registry->GetProperties().Last(), TEXT("Cannot re-define a property type after another has been added."));
-
-		const PTRINT CompositeOffset = (PTRINT)&(((StorageType*)0)->*DataPtr);
-
-		FPropertyCompositeDefinition NewChannel = { InComponent, static_cast<uint16>(CompositeOffset) };
-		Registry->CompositeDefinitions.Add(NewChannel);
-
-		Definition->FloatCompositeMask |= 1 << Definition->CompositeSize;
-
-		++Definition->CompositeSize;
-		return TCompositePropertyDefinitionBuilder<PropertyTraits, Composites..., float>(Definition, Registry);
-	}
-
 	TCompositePropertyDefinitionBuilder<PropertyTraits, Composites..., double> AddComposite(TComponentTypeID<double> InComponent, double StorageType::*DataPtr)
 	{
 		checkf(Definition == &Registry->GetProperties().Last(), TEXT("Cannot re-define a property type after another has been added."));
diff --git a/Engine/Source/Runtime/MovieScene/Public/EntitySystem/MovieScenePropertyRegistry.h b/Engine/Source/Runtime/MovieScene/Public/EntitySystem/MovieScenePropertyRegistry.h
--- a/Engine/Source/Runtime/MovieScene/Public/EntitySystem/MovieScenePropertyRegistry.h
+++ b/Engine/Source/Runtime/MovieScene/Public/EntitySystem/MovieScenePropertyRegistry.h
@@ -57,7 +57,6 @@
 			uint16 InVariableSizeCompositeOffset, uint16 InSizeofStorageType, uint16 InAlignofStorageType,
 			FComponentTypeID InPropertyType, FComponentTypeID InInitialValueType)
 		: CustomPropertyRegistration(nullptr)
-		, FloatCompositeMask(0)
 		, DoubleCompositeMask(0)
 		, VariableSizeCompositeOffset(InVariableSizeCompositeOffset)
 		, CompositeSize(0)
@@ -88,9 +87,6 @@
 	/** Stat ID for this property type */
 	TStatId StatID;
 
-	/** A mask of which composite indices pertain to floats */
-	uint32 FloatCompositeMask = 0;
-
 	/** A mask of which composite indices pertain to doubles */
 	uint32 DoubleCompositeMask = 0;
 
diff --git a/Engine/Source/Runtime/MovieScene/Public/EntitySystem/MovieScenePropertySystemTypes.h b/Engine/Source/Runtime/MovieScene/Public/EntitySystem/MovieScenePropertySystemTypes.h
--- a/Engine/Source/Runtime/MovieScene/Public/EntitySystem/MovieScenePropertySystemTypes.h
+++ b/Engine/Source/Runtime/MovieScene/Public/EntitySystem/MovieScenePropertySystemTypes.h
@@ -57,6 +57,53 @@ template<typename PropertyTraits>
 struct TCompositePropertyTypeID : FCompositePropertyTypeID
 {};
 
+namespace Private
+{
+	/** Utility global flag to determine whether a given type has a nested type called CustomAccessorStorageType */
+	template<typename, typename = void>
+	constexpr bool PropertyTraitsHaveCustomAccessorStorageType = false;
+
+	template<typename T>
+	constexpr bool PropertyTraitsHaveCustomAccessorStorageType<T, std::void_t<decltype(sizeof(typename T::CustomAccessorStorageType))>> = true;
+
+	/** Utility class for selecting a property traits' CustomAccessorStorageType, or StorageType if not defined */
+	template<typename PropertyTraits, bool Custom>
+	struct TCustomPropertyAccessorStorageTypeImpl;
+
+	template<typename PropertyTraits>
+	struct TCustomPropertyAccessorStorageTypeImpl<PropertyTraits, false>
+	{
+		using Value = typename PropertyTraits::StorageType;
+	};
+	
+	template<typename PropertyTraits>
+	struct TCustomPropertyAccessorStorageTypeImpl<PropertyTraits, true>
+	{
+		using Value = typename PropertyTraits::CustomAccessorStorageType;
+	};
+	
+	template<typename PropertyTraits>
+	struct TCustomPropertyAccessorStorageType : TCustomPropertyAccessorStorageTypeImpl<PropertyTraits, PropertyTraitsHaveCustomAccessorStorageType<PropertyTraits>>
+	{};
+
+	namespace Tests
+	{
+		struct TestNormal
+		{
+			using StorageType = bool;
+		};
+		static_assert(PropertyTraitsHaveCustomAccessorStorageType<TestNormal> == false, "Normal has no custom storage type");
+		static_assert(TIsSame<TCustomPropertyAccessorStorageType<TestNormal>::Value, bool>::Value, "Normal has bool storage type");
+
+		struct TestCustom
+		{
+			using StorageType = bool;
+			using CustomAccessorStorageType = int;
+		};
+		static_assert(PropertyTraitsHaveCustomAccessorStorageType<TestCustom> == true, "Custom does have custom storage type");
+		static_assert(TIsSame<TCustomPropertyAccessorStorageType<TestCustom>::Value, int>::Value, "Custom has int storage type");
+	}
+}
 
 /**
  * Structure that defines 2 static function pointers that are to be used for retrieving and applying properties of a given type
@@ -67,7 +114,7 @@ struct TCustomPropertyAccessorFunctionsImpl;
 template<typename PropertyTraits, typename ...MetaDataTypes>
 struct TCustomPropertyAccessorFunctionsImpl<PropertyTraits, TPropertyMetaData<MetaDataTypes...>>
 {
-	using StorageType = typename PropertyTraits::StorageType;
+	using StorageType = typename Private::TCustomPropertyAccessorStorageType<PropertyTraits>::Value;
 	using ParamType   = typename TCallTraits<StorageType>::ParamType;
 
 	using GetterFunc = StorageType (*)(const UObject* Object, MetaDataTypes...);
diff --git a/Engine/Source/Runtime/MovieSceneTracks/Private/Channels/MovieSceneFloatPerlinNoiseChannelContainer.cpp b/Engine/Source/Runtime/MovieSceneTracks/Private/Channels/MovieSceneFloatPerlinNoiseChannelContainer.cpp
--- a/Engine/Source/Runtime/MovieSceneTracks/Private/Channels/MovieSceneFloatPerlinNoiseChannelContainer.cpp
+++ b/Engine/Source/Runtime/MovieSceneTracks/Private/Channels/MovieSceneFloatPerlinNoiseChannelContainer.cpp
@@ -15,7 +15,7 @@ void UMovieSceneFloatPerlinNoiseChannelContainer::ImportEntityImpl(int32 Channel
 	OutImportedEntity->AddBuilder(
 		FEntityBuilder()
 		.Add(TracksComponents->FloatPerlinNoiseChannel, FloatPerlinNoiseChannel.GetParam())
-		.Add(BuiltInComponents->FloatResult[ChannelIndex], TNumericLimits<float>::Max())
+		.Add(BuiltInComponents->DoubleResult[ChannelIndex], TNumericLimits<float>::Max())
 	);
 }
 
diff --git a/Engine/Source/Runtime/MovieSceneTracks/Private/EntitySystem/Interrogation/MovieSceneInterrogatedPropertyInstantiator.cpp b/Engine/Source/Runtime/MovieSceneTracks/Private/EntitySystem/Interrogation/MovieSceneInterrogatedPropertyInstantiator.cpp
--- a/Engine/Source/Runtime/MovieSceneTracks/Private/EntitySystem/Interrogation/MovieSceneInterrogatedPropertyInstantiator.cpp
+++ b/Engine/Source/Runtime/MovieSceneTracks/Private/EntitySystem/Interrogation/MovieSceneInterrogatedPropertyInstantiator.cpp
@@ -5,7 +5,7 @@
 #include "EntitySystem/MovieSceneEntitySystemLinker.h"
 #include "EntitySystem/MovieSceneBlenderSystem.h"
 #include "EntitySystem/MovieScenePropertyRegistry.h"
-#include "Systems/MovieScenePiecewiseFloatBlenderSystem.h"
+#include "Systems/MovieScenePiecewiseDoubleBlenderSystem.h"
 
 #include "Algo/IndexOf.h"
 #include "Algo/Find.h"
@@ -131,9 +131,9 @@ UClass* UMovieSceneInterrogatedPropertyInstantiatorSystem::ResolveBlenderClass(c
 		}
 	}
 	
-	if (!ensureMsgf(BlenderClass, TEXT("No default blender class specified on property, and no custom blender specified on entities. Falling back to float blender.")))
+	if (!ensureMsgf(BlenderClass, TEXT("No default blender class specified on property, and no custom blender specified on entities. Falling back to double blender.")))
 	{
-		BlenderClass = UMovieScenePiecewiseFloatBlenderSystem::StaticClass();
+		BlenderClass = UMovieScenePiecewiseDoubleBlenderSystem::StaticClass();
 	}
 	
 	return BlenderClass;
diff --git a/Engine/Source/Runtime/MovieSceneTracks/Private/MovieSceneTracksComponentTypes.cpp b/Engine/Source/Runtime/MovieSceneTracks/Private/MovieSceneTracksComponentTypes.cpp
--- a/Engine/Source/Runtime/MovieSceneTracks/Private/MovieSceneTracksComponentTypes.cpp
+++ b/Engine/Source/Runtime/MovieSceneTracks/Private/MovieSceneTracksComponentTypes.cpp
@@ -13,7 +13,7 @@
 #include "Systems/MovieScenePiecewiseByteBlenderSystem.h"
 #include "Systems/MovieScenePiecewiseEnumBlenderSystem.h"
 #include "Systems/MovieScenePiecewiseIntegerBlenderSystem.h"
-#include "Systems/MovieScenePiecewiseFloatBlenderSystem.h"
+#include "Systems/MovieScenePiecewiseDoubleBlenderSystem.h"
 #include "Systems/MovieScenePiecewiseDoubleBlenderSystem.h"
 #include "EntitySystem/MovieScenePropertyComponentHandler.h"
 #include "EntitySystem/MovieSceneEntityFactoryTemplates.h"
@@ -228,44 +228,38 @@ void SetSkyLightComponentLightColor(UObject* Object, EColorPropertyType InColorT
 	SkyLightComponent->SetLightColor(InColor.GetLinearColor());
 }
 
-float GetSecondFogDataFogDensity(const UObject* Object, bool bIsDouble)
+float GetSecondFogDataFogDensity(const UObject* Object)
 {
-	check(!bIsDouble);
 	const UExponentialHeightFogComponent* ExponentialHeightFogComponent = CastChecked<const UExponentialHeightFogComponent>(Object);
 	return ExponentialHeightFogComponent->SecondFogData.FogDensity;
 }
 
-void SetSecondFogDataFogDensity(UObject* Object, bool bIsDouble, float InFogDensity)
+void SetSecondFogDataFogDensity(UObject* Object, float InFogDensity)
 {
-	check(!bIsDouble);
 	UExponentialHeightFogComponent* ExponentialHeightFogComponent = CastChecked<UExponentialHeightFogComponent>(Object);
 	ExponentialHeightFogComponent->SecondFogData.FogDensity = InFogDensity;
 }
 
-float GetSecondFogDataFogHeightFalloff(const UObject* Object, bool bIsDouble)
+float GetSecondFogDataFogHeightFalloff(const UObject* Object)
 {
-	check(!bIsDouble);
 	const UExponentialHeightFogComponent* ExponentialHeightFogComponent = CastChecked<const UExponentialHeightFogComponent>(Object);
 	return ExponentialHeightFogComponent->SecondFogData.FogHeightFalloff;
 }
 
-void SetSecondFogDataFogHeightFalloff(UObject* Object, bool bIsDouble, float InFogHeightFalloff)
+void SetSecondFogDataFogHeightFalloff(UObject* Object, float InFogHeightFalloff)
 {
-	check(!bIsDouble);
 	UExponentialHeightFogComponent* ExponentialHeightFogComponent = CastChecked<UExponentialHeightFogComponent>(Object);
 	ExponentialHeightFogComponent->SecondFogData.FogHeightFalloff = InFogHeightFalloff;
 }
 
-float GetSecondFogDataFogHeightOffset(const UObject* Object, bool bIsDouble)
+float GetSecondFogDataFogHeightOffset(const UObject* Object)
 {
-	check(!bIsDouble);
 	const UExponentialHeightFogComponent* ExponentialHeightFogComponent = CastChecked<const UExponentialHeightFogComponent>(Object);
 	return ExponentialHeightFogComponent->SecondFogData.FogHeightOffset;
 }
 
-void SetSecondFogDataFogHeightOffset(UObject* Object, bool bIsDouble, float InFogHeightOffset)
+void SetSecondFogDataFogHeightOffset(UObject* Object, float InFogHeightOffset)
 {
-	check(!bIsDouble);
 	UExponentialHeightFogComponent* ExponentialHeightFogComponent = CastChecked<UExponentialHeightFogComponent>(Object);
 	ExponentialHeightFogComponent->SecondFogData.FogHeightOffset = InFogHeightOffset;
 }
@@ -368,34 +362,11 @@ void FComponentDetachParams::ApplyDetach(USceneComponent* ChildComponentToAttach
 static bool GMovieSceneTracksComponentTypesDestroyed = false;
 static TUniquePtr<FMovieSceneTracksComponentTypes> GMovieSceneTracksComponentTypes;
 
-struct FFloatHandler : TPropertyComponentHandler<FFloatPropertyTraits, float>
+struct FFloatHandler : TPropertyComponentHandler<FFloatPropertyTraits, double>
 {
-	virtual void DispatchInitializePropertyMetaDataTasks(const FPropertyDefinition& Definition, FSystemTaskPrerequisites& InPrerequisites, FSystemSubsequentTasks& Subsequents, UMovieSceneEntitySystemLinker* Linker) override
-	{
-		FBuiltInComponentTypes* BuiltInComponents = FBuiltInComponentTypes::Get();
-		FMovieSceneTracksComponentTypes* TrackComponents = FMovieSceneTracksComponentTypes::Get();
-
-		FEntityTaskBuilder()
-		.Read(BuiltInComponents->BoundObject)
-		.Read(BuiltInComponents->PropertyBinding)
-		.Write(TrackComponents->Float.MetaDataComponents.GetType<0>())
-		.FilterAll({ BuiltInComponents->Tags.NeedsLink })
-		.Iterate_PerEntity(&Linker->EntityManager, [](UObject* Object, const FMovieScenePropertyBinding& Binding, bool& OutIsDouble)
-		{
-			FProperty* BoundProperty = FTrackInstancePropertyBindings::FindProperty(Object, Binding.PropertyPath.ToString());
-			if (ensure(BoundProperty))
-			{
-				OutIsDouble = BoundProperty->IsA<FDoubleProperty>();
-			}
-			else
-			{
-				OutIsDouble = false;
-			}
-		});
-	}
 };
 
-struct FColorHandler : TPropertyComponentHandler<FColorPropertyTraits, float, float, float, float>
+struct FColorHandler : TPropertyComponentHandler<FColorPropertyTraits, double, double, double, double>
 {
 	virtual void DispatchInitializePropertyMetaDataTasks(const FPropertyDefinition& Definition, FSystemTaskPrerequisites& InPrerequisites, FSystemSubsequentTasks& Subsequents, UMovieSceneEntitySystemLinker* Linker) override
 	{
@@ -438,7 +409,7 @@ struct FColorHandler : TPropertyComponentHandler<FColorPropertyTraits, float, fl
 };
 
 
-struct FFloatVectorHandler : TPropertyComponentHandler<FFloatVectorPropertyTraits, float, float, float, float>
+struct FFloatVectorHandler : TPropertyComponentHandler<FFloatVectorPropertyTraits, double, double, double, double>
 {
 	virtual void DispatchInitializePropertyMetaDataTasks(const FPropertyDefinition& Definition, FSystemTaskPrerequisites& InPrerequisites, FSystemSubsequentTasks& Subsequents, UMovieSceneEntitySystemLinker* Linker) override
 	{
@@ -455,21 +426,18 @@ struct FFloatVectorHandler : TPropertyComponentHandler<FFloatVectorPropertyTrait
 			FStructProperty* BoundProperty = CastField<FStructProperty>(FTrackInstancePropertyBindings::FindProperty(Object, Binding.PropertyPath.ToString()));
 			if (ensure(BoundProperty && BoundProperty->Struct))
 			{
-				if (BoundProperty->Struct == TBaseStructure<FVector2D>::Get() || BoundProperty->Struct == TVariantStructure<FVector2f>::Get()) // LWC_TODO: Fix LWC type FName case
+				if (BoundProperty->Struct == TVariantStructure<FVector2f>::Get())
 				{
 					OutMetaData.NumChannels = 2;
-					OutMetaData.bIsDouble = (BoundProperty->Struct == TBaseStructure<FVector2D>::Get()); // LWC_TODO: Fix LWC type FName case
 				}
-				else if (BoundProperty->Struct == TBaseStructure<FVector>::Get() || BoundProperty->Struct == TVariantStructure<FVector3f>::Get() || BoundProperty->Struct == TVariantStructure<FVector3d>::Get())
+				else if (BoundProperty->Struct == TVariantStructure<FVector3f>::Get())
 				{
 					OutMetaData.NumChannels = 3;
-					OutMetaData.bIsDouble = (BoundProperty->Struct == TBaseStructure<FVector>::Get() || BoundProperty->Struct == TVariantStructure<FVector3d>::Get());
 				}
 				else
 				{
-					ensure(BoundProperty->Struct == TBaseStructure<FVector4>::Get() || BoundProperty->Struct == TVariantStructure<FVector4f>::Get() || BoundProperty->Struct == TVariantStructure<FVector4d>::Get());
+					ensure(BoundProperty->Struct == TVariantStructure<FVector4f>::Get());
 					OutMetaData.NumChannels = 4;
-					OutMetaData.bIsDouble = (BoundProperty->Struct == TBaseStructure<FVector4>::Get() || BoundProperty->Struct == TVariantStructure<FVector4d>::Get());
 				}
 			}
 			else
@@ -502,13 +470,13 @@ struct FDoubleVectorHandler : TPropertyComponentHandler<FDoubleVectorPropertyTra
 				{
 					OutMetaData.NumChannels = 2;
 				}
-				else if (BoundProperty->Struct->GetFName() == NAME_Vector3d || BoundProperty->Struct->GetFName() == NAME_Vector)
+				else if (BoundProperty->Struct == TBaseStructure<FVector>::Get() || BoundProperty->Struct == TVariantStructure<FVector3d>::Get())
 				{
 					OutMetaData.NumChannels = 3;
 				}
 				else
 				{
-					ensure(BoundProperty->Struct == TBaseStructure<FVector4>::Get() || (BoundProperty->Struct->GetFName() == NAME_Vector4d) || (BoundProperty->Struct->GetFName() == NAME_Vector4));
+					ensure(BoundProperty->Struct == TBaseStructure<FVector4>::Get() || BoundProperty->Struct == TVariantStructure<FVector4d>::Get());
 					OutMetaData.NumChannels = 4;
 				}
 			}
@@ -547,7 +515,6 @@ FMovieSceneTracksComponentTypes::FMovieSceneTracksComponentTypes()
 	ComponentRegistry->NewPropertyType(EulerTransform, TEXT("FEulerTransform"));
 	ComponentRegistry->NewPropertyType(ComponentTransform, TEXT("Component Transform"));
 
-	Float.MetaDataComponents.Initialize(ComponentRegistry, TEXT("Is Double"));
 	Color.MetaDataComponents.Initialize(ComponentRegistry, TEXT("Color Type"));
 	FloatVector.MetaDataComponents.Initialize(ComponentRegistry, TEXT("Num Float Vector Channels"));
 	DoubleVector.MetaDataComponents.Initialize(ComponentRegistry, TEXT("Num Double Vector Channels"));
@@ -629,8 +596,8 @@ FMovieSceneTracksComponentTypes::FMovieSceneTracksComponentTypes()
 	// --------------------------------------------------------------------------------------------
 	// Set up float properties
 	BuiltInComponents->PropertyRegistry.DefineProperty(Float, TEXT("Apply float Properties"))
-	.AddSoleChannel(BuiltInComponents->FloatResult[0])
-	.SetBlenderSystem<UMovieScenePiecewiseFloatBlenderSystem>()
+	.AddSoleChannel(BuiltInComponents->DoubleResult[0])
+	.SetBlenderSystem<UMovieScenePiecewiseDoubleBlenderSystem>()
 	.SetCustomAccessors(&Accessors.Float)
 	.Commit(FFloatHandler());
 
@@ -645,11 +612,11 @@ FMovieSceneTracksComponentTypes::FMovieSceneTracksComponentTypes()
 	// --------------------------------------------------------------------------------------------
 	// Set up color properties
 	BuiltInComponents->PropertyRegistry.DefineCompositeProperty(Color, TEXT("Apply Color Properties"))
-	.AddComposite(BuiltInComponents->FloatResult[0], &FIntermediateColor::R)
-	.AddComposite(BuiltInComponents->FloatResult[1], &FIntermediateColor::G)
-	.AddComposite(BuiltInComponents->FloatResult[2], &FIntermediateColor::B)
-	.AddComposite(BuiltInComponents->FloatResult[3], &FIntermediateColor::A)
-	.SetBlenderSystem<UMovieScenePiecewiseFloatBlenderSystem>()
+	.AddComposite(BuiltInComponents->DoubleResult[0], &FIntermediateColor::R)
+	.AddComposite(BuiltInComponents->DoubleResult[1], &FIntermediateColor::G)
+	.AddComposite(BuiltInComponents->DoubleResult[2], &FIntermediateColor::B)
+	.AddComposite(BuiltInComponents->DoubleResult[3], &FIntermediateColor::A)
+	.SetBlenderSystem<UMovieScenePiecewiseDoubleBlenderSystem>()
 	.SetCustomAccessors(&Accessors.Color)
 	.Commit(FColorHandler());
 
@@ -677,11 +644,11 @@ FMovieSceneTracksComponentTypes::FMovieSceneTracksComponentTypes()
 	// --------------------------------------------------------------------------------------------
 	// Set up vector properties
 	BuiltInComponents->PropertyRegistry.DefineCompositeProperty(FloatVector, TEXT("Apply FloatVector Properties"))
-	.AddComposite(BuiltInComponents->FloatResult[0], &FFloatIntermediateVector::X)
-	.AddComposite(BuiltInComponents->FloatResult[1], &FFloatIntermediateVector::Y)
-	.AddComposite(BuiltInComponents->FloatResult[2], &FFloatIntermediateVector::Z)
-	.AddComposite(BuiltInComponents->FloatResult[3], &FFloatIntermediateVector::W)
-	.SetBlenderSystem<UMovieScenePiecewiseFloatBlenderSystem>()
+	.AddComposite(BuiltInComponents->DoubleResult[0], &FFloatIntermediateVector::X)
+	.AddComposite(BuiltInComponents->DoubleResult[1], &FFloatIntermediateVector::Y)
+	.AddComposite(BuiltInComponents->DoubleResult[2], &FFloatIntermediateVector::Z)
+	.AddComposite(BuiltInComponents->DoubleResult[3], &FFloatIntermediateVector::W)
+	.SetBlenderSystem<UMovieScenePiecewiseDoubleBlenderSystem>()
 	.SetCustomAccessors(&Accessors.FloatVector)
 	.Commit(FFloatVectorHandler());
 
diff --git a/Engine/Source/Runtime/MovieSceneTracks/Private/Systems/FloatChannelEvaluatorSystem.cpp b/Engine/Source/Runtime/MovieSceneTracks/Private/Systems/FloatChannelEvaluatorSystem.cpp
--- a/Engine/Source/Runtime/MovieSceneTracks/Private/Systems/FloatChannelEvaluatorSystem.cpp
+++ b/Engine/Source/Runtime/MovieSceneTracks/Private/Systems/FloatChannelEvaluatorSystem.cpp
@@ -9,6 +9,7 @@
 
 #include "EntitySystem/MovieSceneEvalTimeSystem.h"
 #include "Channels/MovieSceneFloatChannel.h"
+#include "MovieSceneTracksComponentTypes.h"
 
 #include "Algo/Find.h"
 
@@ -25,17 +26,19 @@ namespace MovieScene
 // Do we need to optimize for this case using something like the code below, while pessimizing the common (non-multi-bind) codepath??
 struct FEvaluateFloatChannels
 {
-	void ForEachEntity(FSourceFloatChannel FloatChannel, FFrameTime FrameTime, float& OutResult, FSourceFloatChannelFlags& OutFlags)
+	void ForEachEntity(FSourceFloatChannel FloatChannel, FFrameTime FrameTime, double& OutResult, FSourceFloatChannelFlags& OutFlags)
 	{
 		if (OutFlags.bNeedsEvaluate == false)
 		{
 			return;
 		}
 
-		if (!FloatChannel.Source->Evaluate(FrameTime, OutResult))
+		float FloatResult = 0.f;
+		if (!FloatChannel.Source->Evaluate(FrameTime, FloatResult))
 		{
-			OutResult = MIN_flt;
+			FloatResult = MIN_flt;
 		}
+		OutResult = (double)FloatResult;
 
 		if (FloatChannel.Source->GetTimes().Num() <= 1)
 		{
@@ -64,13 +67,13 @@ UFloatChannelEvaluatorSystem::UFloatChannelEvaluatorSystem(const FObjectInitiali
 
 		for (int32 Index = 0; Index < UE_ARRAY_COUNT(Components->FloatChannel); ++Index)
 		{
-			RegisterChannelType(Components->FloatChannel[Index], Components->FloatChannelFlags[Index], Components->FloatResult[Index]);
+			RegisterChannelType(Components->FloatChannel[Index], Components->FloatChannelFlags[Index], Components->DoubleResult[Index]);
 		}
 		RegisterChannelType(Components->WeightChannel, Components->WeightChannelFlags, Components->WeightResult);
 	}
 }
 
-void UFloatChannelEvaluatorSystem::RegisterChannelType(TComponentTypeID<UE::MovieScene::FSourceFloatChannel> SourceChannelType, TComponentTypeID<UE::MovieScene::FSourceFloatChannelFlags> ChannelFlagsType, TComponentTypeID<float> ResultType)
+void UFloatChannelEvaluatorSystem::RegisterChannelType(TComponentTypeID<UE::MovieScene::FSourceFloatChannel> SourceChannelType, TComponentTypeID<UE::MovieScene::FSourceFloatChannelFlags> ChannelFlagsType, TComponentTypeID<double> ResultType)
 {
 	using namespace UE::MovieScene;
 
diff --git a/Engine/Source/Runtime/MovieSceneTracks/Private/Systems/FloatPerlinNoiseChannelEvaluatorSystem.cpp b/Engine/Source/Runtime/MovieSceneTracks/Private/Systems/FloatPerlinNoiseChannelEvaluatorSystem.cpp
--- a/Engine/Source/Runtime/MovieSceneTracks/Private/Systems/FloatPerlinNoiseChannelEvaluatorSystem.cpp
+++ b/Engine/Source/Runtime/MovieSceneTracks/Private/Systems/FloatPerlinNoiseChannelEvaluatorSystem.cpp
@@ -1,7 +1,7 @@
 // Copyright Epic Games, Inc. All Rights Reserved.
 
 #include "Systems/FloatPerlinNoiseChannelEvaluatorSystem.h"
-#include "Systems/MovieScenePiecewiseFloatBlenderSystem.h"
+#include "Systems/MovieScenePiecewiseDoubleBlenderSystem.h"
 #include "EntitySystem/BuiltInComponentTypes.h"
 #include "EntitySystem/MovieSceneEntitySystemTask.h"
 #include "EntitySystem/MovieSceneEntitySystemLinker.h"
@@ -26,13 +26,14 @@ namespace UE
 				check(InstanceRegistry);
 			}
 
-			void ForEachEntity(FMovieSceneFloatPerlinNoiseChannel FloatPerlinNoiseChannel, FInstanceHandle InstanceHandle, float& OutResult)
+			void ForEachEntity(FMovieSceneFloatPerlinNoiseChannel FloatPerlinNoiseChannel, FInstanceHandle InstanceHandle, double& OutResult)
 			{
 				const FMovieSceneContext& Context = InstanceRegistry->GetContext(InstanceHandle);
 				FFrameTime Time = Context.GetTime();
 				FFrameRate Rate = Context.GetFrameRate();
 
-				OutResult = FloatPerlinNoiseChannel.Evaluate(Rate.AsSeconds(Time));
+				const float NoiseResult = FloatPerlinNoiseChannel.Evaluate(Rate.AsSeconds(Time));
+				OutResult = (double)NoiseResult;
 			}
 		};
 
@@ -53,13 +54,13 @@ UFloatPerlinNoiseChannelEvaluatorSystem::UFloatPerlinNoiseChannelEvaluatorSystem
 	if (HasAnyFlags(RF_ClassDefaultObject))
 	{
 		// Allow writing to all the possible channels
-		for (int32 i = 0; i < UE_ARRAY_COUNT(BuiltInComponents->FloatResult); ++i)
+		for (int32 i = 0; i < UE_ARRAY_COUNT(BuiltInComponents->DoubleResult); ++i)
 		{
-			DefineComponentProducer(UFloatPerlinNoiseChannelEvaluatorSystem::StaticClass(), BuiltInComponents->FloatResult[i]);
+			DefineComponentProducer(UFloatPerlinNoiseChannelEvaluatorSystem::StaticClass(), BuiltInComponents->DoubleResult[i]);
 		}
 
 		DefineImplicitPrerequisite(UMovieSceneEvalTimeSystem::StaticClass(), GetClass());
-		DefineImplicitPrerequisite(GetClass(), UMovieScenePiecewiseFloatBlenderSystem::StaticClass());
+		DefineImplicitPrerequisite(GetClass(), UMovieScenePiecewiseDoubleBlenderSystem::StaticClass());
 	}
 }
 
@@ -70,12 +71,12 @@ void UFloatPerlinNoiseChannelEvaluatorSystem::OnRun(FSystemTaskPrerequisites& In
 	const FBuiltInComponentTypes* BuiltInComponents = FBuiltInComponentTypes::Get();
 	const FMovieSceneTracksComponentTypes* TrackComponents = FMovieSceneTracksComponentTypes::Get();
 
-	for (int32 i = 0; i < UE_ARRAY_COUNT(BuiltInComponents->FloatResult); ++i)
+	for (int32 i = 0; i < UE_ARRAY_COUNT(BuiltInComponents->DoubleResult); ++i)
 	{
 		FEntityTaskBuilder()
 			.Read(TrackComponents->FloatPerlinNoiseChannel)
 			.Read(BuiltInComponents->InstanceHandle)
-			.Write(BuiltInComponents->FloatResult[i])
+			.Write(BuiltInComponents->DoubleResult[i])
 			.FilterNone({ BuiltInComponents->Tags.Ignored })
 			.SetStat(GET_STATID(MovieSceneEval_EvaluateFloatPerlinNoiseChannelTask))
 			.Dispatch_PerEntity<FEvaluateFloatPerlinNoiseChannels>(&Linker->EntityManager, InPrerequisites, &Subsequents, GetLinker()->GetInstanceRegistry());
diff --git a/Engine/Source/Runtime/MovieSceneTracks/Private/Systems/MovieSceneBaseValueEvaluatorSystem.cpp b/Engine/Source/Runtime/MovieSceneTracks/Private/Systems/MovieSceneBaseValueEvaluatorSystem.cpp
--- a/Engine/Source/Runtime/MovieSceneTracks/Private/Systems/MovieSceneBaseValueEvaluatorSystem.cpp
+++ b/Engine/Source/Runtime/MovieSceneTracks/Private/Systems/MovieSceneBaseValueEvaluatorSystem.cpp
@@ -17,12 +17,14 @@ namespace MovieScene
 
 struct FEvaluateBaseFloatValues
 {
-	void ForEachEntity(FSourceFloatChannel FloatChannel, FFrameTime FrameTime, float& OutResult)
+	void ForEachEntity(FSourceFloatChannel FloatChannel, FFrameTime FrameTime, double& OutResult)
 	{
-		if (!FloatChannel.Source->Evaluate(FrameTime, OutResult))
+		float Result;
+		if (!FloatChannel.Source->Evaluate(FrameTime, Result))
 		{
-			OutResult = MIN_flt;
+			Result = MIN_flt;
 		}
+		OutResult = (double)Result;
 	}
 };
 
@@ -61,17 +63,17 @@ void UMovieSceneBaseValueEvaluatorSystem::OnRun(FSystemTaskPrerequisites& InPrer
 	FBuiltInComponentTypes* BuiltInComponents = FBuiltInComponentTypes::Get();
 
 	static_assert(
-			UE_ARRAY_COUNT(BuiltInComponents->BaseFloat) == UE_ARRAY_COUNT(BuiltInComponents->FloatChannel),
+			UE_ARRAY_COUNT(BuiltInComponents->BaseDouble) == UE_ARRAY_COUNT(BuiltInComponents->FloatChannel),
 			"There should be a matching number of float channels and float base values.");
-	for (size_t Index = 0; Index < UE_ARRAY_COUNT(BuiltInComponents->BaseFloat); ++Index)
+	for (size_t Index = 0; Index < UE_ARRAY_COUNT(BuiltInComponents->BaseDouble); ++Index)
 	{
-		const TComponentTypeID<float> BaseFloat = BuiltInComponents->BaseFloat[Index];
+		const TComponentTypeID<double> BaseDouble = BuiltInComponents->BaseDouble[Index];
 		const TComponentTypeID<FSourceFloatChannel> FloatChannel = BuiltInComponents->FloatChannel[Index];
 
 		FEntityTaskBuilder()
 		.Read(FloatChannel)
 		.Read(BuiltInComponents->BaseValueEvalTime)
-		.Write(BaseFloat)
+		.Write(BaseDouble)
 		.FilterAll({ BuiltInComponents->Tags.NeedsLink })
 		.FilterNone({ BuiltInComponents->Tags.Ignored })
 		.RunInline_PerEntity(&Linker->EntityManager, FEvaluateBaseFloatValues());
diff --git a/Engine/Source/Runtime/MovieSceneTracks/Private/Systems/MovieSceneColorPropertySystem.cpp b/Engine/Source/Runtime/MovieSceneTracks/Private/Systems/MovieSceneColorPropertySystem.cpp
--- a/Engine/Source/Runtime/MovieSceneTracks/Private/Systems/MovieSceneColorPropertySystem.cpp
+++ b/Engine/Source/Runtime/MovieSceneTracks/Private/Systems/MovieSceneColorPropertySystem.cpp
@@ -3,7 +3,7 @@
 #include "Systems/MovieSceneColorPropertySystem.h"
 #include "MovieSceneTracksComponentTypes.h"
 #include "Systems/FloatChannelEvaluatorSystem.h"
-#include "Systems/MovieScenePiecewiseFloatBlenderSystem.h"
+#include "Systems/MovieScenePiecewiseDoubleBlenderSystem.h"
 
 
 UMovieSceneColorPropertySystem::UMovieSceneColorPropertySystem(const FObjectInitializer& ObjInit)
@@ -16,7 +16,7 @@ UMovieSceneColorPropertySystem::UMovieSceneColorPropertySystem(const FObjectInit
 	if (HasAnyFlags(RF_ClassDefaultObject))
 	{
 		// We need our floats correctly evaluated and blended, so we are downstream from those systems.
-		DefineImplicitPrerequisite(UMovieScenePiecewiseFloatBlenderSystem::StaticClass(), GetClass());
+		DefineImplicitPrerequisite(UMovieScenePiecewiseDoubleBlenderSystem::StaticClass(), GetClass());
 		DefineImplicitPrerequisite(UFloatChannelEvaluatorSystem::StaticClass(), GetClass());
 	}
 }
diff --git a/Engine/Source/Runtime/MovieSceneTracks/Private/Systems/MovieSceneComponentMaterialSystem.cpp b/Engine/Source/Runtime/MovieSceneTracks/Private/Systems/MovieSceneComponentMaterialSystem.cpp
--- a/Engine/Source/Runtime/MovieSceneTracks/Private/Systems/MovieSceneComponentMaterialSystem.cpp
+++ b/Engine/Source/Runtime/MovieSceneTracks/Private/Systems/MovieSceneComponentMaterialSystem.cpp
@@ -2,11 +2,13 @@
 
 #include "Systems/MovieSceneComponentMaterialSystem.h"
 
+#include "Evaluation/PreAnimatedState/MovieScenePreAnimatedStorageID.inl"
+
 #include "EntitySystem/MovieSceneEntityMutations.h"
 #include "Evaluation/PreAnimatedState/MovieScenePreAnimatedStorageID.inl"
 
 #include "Systems/FloatChannelEvaluatorSystem.h"
-#include "Systems/MovieScenePiecewiseFloatBlenderSystem.h"
+#include "Systems/MovieScenePiecewiseDoubleBlenderSystem.h"
 
 #include "Materials/MaterialInstanceDynamic.h"
 
diff --git a/Engine/Source/Runtime/MovieSceneTracks/Private/Systems/MovieSceneFloatPropertySystem.cpp b/Engine/Source/Runtime/MovieSceneTracks/Private/Systems/MovieSceneFloatPropertySystem.cpp
--- a/Engine/Source/Runtime/MovieSceneTracks/Private/Systems/MovieSceneFloatPropertySystem.cpp
+++ b/Engine/Source/Runtime/MovieSceneTracks/Private/Systems/MovieSceneFloatPropertySystem.cpp
@@ -2,7 +2,7 @@
 
 #include "Systems/MovieSceneFloatPropertySystem.h"
 #include "Systems/FloatChannelEvaluatorSystem.h"
-#include "Systems/MovieScenePiecewiseFloatBlenderSystem.h"
+#include "Systems/MovieScenePiecewiseDoubleBlenderSystem.h"
 
 #include "MovieSceneTracksComponentTypes.h"
 
@@ -16,7 +16,7 @@ UMovieSceneFloatPropertySystem::UMovieSceneFloatPropertySystem(const FObjectInit
 
 	if (HasAnyFlags(RF_ClassDefaultObject))
 	{
-		DefineImplicitPrerequisite(UMovieScenePiecewiseFloatBlenderSystem::StaticClass(), GetClass());
+		DefineImplicitPrerequisite(UMovieScenePiecewiseDoubleBlenderSystem::StaticClass(), GetClass());
 		DefineImplicitPrerequisite(UFloatChannelEvaluatorSystem::StaticClass(), GetClass());
 	}
 }
diff --git a/Engine/Source/Runtime/MovieSceneTracks/Private/Systems/MovieSceneMaterialParameterCollectionSystem.cpp b/Engine/Source/Runtime/MovieSceneTracks/Private/Systems/MovieSceneMaterialParameterCollectionSystem.cpp
--- a/Engine/Source/Runtime/MovieSceneTracks/Private/Systems/MovieSceneMaterialParameterCollectionSystem.cpp
+++ b/Engine/Source/Runtime/MovieSceneTracks/Private/Systems/MovieSceneMaterialParameterCollectionSystem.cpp
@@ -5,13 +5,14 @@
 #include "IMovieScenePlayer.h"
 #include "MovieSceneSequence.h"
 #include "Evaluation/MovieSceneEvaluationTemplateInstance.h"
+#include "Evaluation/PreAnimatedState/MovieScenePreAnimatedStorageID.inl"
 #include "EntitySystem/BuiltInComponentTypes.h"
 #include "EntitySystem/MovieSceneEntityMutations.h"
 #include "Evaluation/PreAnimatedState/MovieScenePreAnimatedStorageID.inl"
 
 #include "MovieSceneTracksComponentTypes.h"
 #include "Systems/FloatChannelEvaluatorSystem.h"
-#include "Systems/MovieScenePiecewiseFloatBlenderSystem.h"
+#include "Systems/MovieScenePiecewiseDoubleBlenderSystem.h"
 
 #include "Logging/MessageLog.h"
 #include "Misc/UObjectToken.h"
diff --git a/Engine/Source/Runtime/MovieSceneTracks/Private/Systems/MovieSceneMaterialParameterSystem.cpp b/Engine/Source/Runtime/MovieSceneTracks/Private/Systems/MovieSceneMaterialParameterSystem.cpp
--- a/Engine/Source/Runtime/MovieSceneTracks/Private/Systems/MovieSceneMaterialParameterSystem.cpp
+++ b/Engine/Source/Runtime/MovieSceneTracks/Private/Systems/MovieSceneMaterialParameterSystem.cpp
@@ -4,7 +4,7 @@
 
 #include "Systems/BoolChannelEvaluatorSystem.h"
 #include "Systems/FloatChannelEvaluatorSystem.h"
-#include "Systems/MovieScenePiecewiseFloatBlenderSystem.h"
+#include "Systems/MovieScenePiecewiseDoubleBlenderSystem.h"
 #include "Systems/MovieSceneHierarchicalBiasSystem.h"
 #include "Systems/MovieSceneMaterialSystem.h"
 
@@ -17,16 +17,16 @@ namespace UE::MovieScene
 /** Apply scalar material parameters */
 struct FApplyScalarParameters
 {
-	static void ForEachEntity(UObject* BoundMaterial, FName ParameterName, float InScalarValue)
+	static void ForEachEntity(UObject* BoundMaterial, FName ParameterName, double InScalarValue)
 	{
 		// WARNING: BoundMaterial may be nullptr here
 		if (UMaterialInstanceDynamic* MID = Cast<UMaterialInstanceDynamic>(BoundMaterial))
 		{
-			MID->SetScalarParameterValue(ParameterName, InScalarValue);
+			MID->SetScalarParameterValue(ParameterName, (float)InScalarValue);
 		}
 		else if (UMaterialParameterCollectionInstance* MPCI = Cast<UMaterialParameterCollectionInstance>(BoundMaterial))
 		{
-			MPCI->SetScalarParameterValue(ParameterName, InScalarValue);
+			MPCI->SetScalarParameterValue(ParameterName, (float)InScalarValue);
 		}
 	}
 };
@@ -37,15 +37,15 @@ struct FApplyVectorParameters
 	static void ForEachAllocation(FEntityAllocationIteratorItem Item,
 		TRead<UObject*> BoundMaterials,
 		TReadOneOrMoreOf<FName, FName> VectorOrColorParameterNames,
-		TReadOneOrMoreOf<float, float, float, float> VectorChannels)
+		TReadOneOrMoreOf<double, double, double, double> VectorChannels)
 	{
 		const int32 Num = Item.GetAllocation()->Num();
 		// Use either the vector parameter name, or the color parameter name
 		const FName* ParameterNames = VectorOrColorParameterNames.Get<0>() ? VectorOrColorParameterNames.Get<0>() : VectorOrColorParameterNames.Get<1>();
-		const float* RESTRICT R = VectorChannels.Get<0>();
-		const float* RESTRICT G = VectorChannels.Get<1>();
-		const float* RESTRICT B = VectorChannels.Get<2>();
-		const float* RESTRICT A = VectorChannels.Get<3>();
+		const double* RESTRICT R = VectorChannels.Get<0>();
+		const double* RESTRICT G = VectorChannels.Get<1>();
+		const double* RESTRICT B = VectorChannels.Get<2>();
+		const double* RESTRICT A = VectorChannels.Get<3>();
 		
 		for (int32 Index = 0; Index < Num; ++Index)
 		{
@@ -56,10 +56,10 @@ struct FApplyVectorParameters
 			}
 
 			FLinearColor Color(
-				R ? R[Index] : 0.f,
-				G ? G[Index] : 0.f,
-				B ? B[Index] : 0.f,
-				A ? A[Index] : 1.f
+				R ? (float)R[Index] : 0.f,
+				G ? (float)G[Index] : 0.f,
+				B ? (float)B[Index] : 0.f,
+				A ? (float)A[Index] : 1.f
 			);
 			if (UMaterialInstanceDynamic* MID = Cast<UMaterialInstanceDynamic>(BoundMaterial))
 			{
@@ -107,14 +107,14 @@ struct FOverlappingMaterialParameterHandler
 		{
 			if (!Output->OutputEntityID)
 			{
-				if (!System->FloatBlenderSystem)
+				if (!System->DoubleBlenderSystem)
 				{
-					System->FloatBlenderSystem = Linker->LinkSystem<UMovieScenePiecewiseFloatBlenderSystem>();
-					Linker->SystemGraph.AddReference(System, System->FloatBlenderSystem);
+					System->DoubleBlenderSystem = Linker->LinkSystem<UMovieScenePiecewiseDoubleBlenderSystem>();
+					Linker->SystemGraph.AddReference(System, System->DoubleBlenderSystem);
 				}
 
 				// Initialize the blend channel ID
-				Output->BlendChannelID = System->FloatBlenderSystem->AllocateBlendChannel();
+				Output->BlendChannelID = System->DoubleBlenderSystem->AllocateBlendChannel();
 
 				// Needs blending
 				Output->OutputEntityID = FEntityBuilder()
@@ -143,9 +143,9 @@ struct FOverlappingMaterialParameterHandler
 				Linker->EntityManager.AddComponent(Output->OutputEntityID, BuiltInComponents->Tags.NeedsUnlink);
 				Output->OutputEntityID = FMovieSceneEntityID();
 
-				if (ensure(System->FloatBlenderSystem))
+				if (ensure(System->DoubleBlenderSystem))
 				{
-					System->FloatBlenderSystem->ReleaseBlendChannel(Output->BlendChannelID);
+					System->DoubleBlenderSystem->ReleaseBlendChannel(Output->BlendChannelID);
 				}
 			}
 
@@ -190,7 +190,7 @@ UMovieSceneMaterialParameterSystem::UMovieSceneMaterialParameterSystem(const FOb
 
 		DefineImplicitPrerequisite(UFloatChannelEvaluatorSystem::StaticClass(), GetClass());
 
-		DefineImplicitPrerequisite(UMovieScenePiecewiseFloatBlenderSystem::StaticClass(), GetClass());
+		DefineImplicitPrerequisite(UMovieScenePiecewiseDoubleBlenderSystem::StaticClass(), GetClass());
 		DefineImplicitPrerequisite(GetClass(), UMovieSceneHierarchicalBiasSystem::StaticClass());
 	}
 }
@@ -200,7 +200,7 @@ void UMovieSceneMaterialParameterSystem::OnLink()
 	using namespace UE::MovieScene;
 
 	// Always reset the float blender system on link to ensure that recycled systems are correctly initialized.
-	FloatBlenderSystem = nullptr;
+	DoubleBlenderSystem = nullptr;
 
 	FOverlappingMaterialParameterHandler Handler(this);
 
@@ -243,7 +243,7 @@ void UMovieSceneMaterialParameterSystem::OnInstantiation()
 		ScalarParameterTracker.Update(Linker, TracksComponents->BoundMaterial, TracksComponents->ScalarParameterName, FEntityComponentFilter());
 
 		FOverlappingMaterialParameterHandler Handler(this);
-		Handler.DefaultComponentMask.Set(BuiltInComponents->FloatResult[0]);
+		Handler.DefaultComponentMask.Set(BuiltInComponents->DoubleResult[0]);
 		ScalarParameterTracker.ProcessInvalidatedOutputs(Linker, Handler);
 	}
 
@@ -258,7 +258,7 @@ void UMovieSceneMaterialParameterSystem::OnInstantiation()
 	if (bHasColors || bHasVectors)
 	{
 		FOverlappingMaterialParameterHandler Handler(this);
-		Handler.DefaultComponentMask.SetAll({ BuiltInComponents->FloatResult[0], BuiltInComponents->FloatResult[1], BuiltInComponents->FloatResult[2], BuiltInComponents->FloatResult[3] });
+		Handler.DefaultComponentMask.SetAll({ BuiltInComponents->DoubleResult[0], BuiltInComponents->DoubleResult[1], BuiltInComponents->DoubleResult[2], BuiltInComponents->DoubleResult[3] });
 		VectorParameterTracker.ProcessInvalidatedOutputs(Linker, Handler);
 	}
 }
@@ -275,7 +275,7 @@ void UMovieSceneMaterialParameterSystem::OnEvaluation(FSystemTaskPrerequisites&
 		FEntityTaskBuilder()
 		.Read(TracksComponents->BoundMaterial)
 		.Read(TracksComponents->ScalarParameterName)
-		.Read(BuiltInComponents->FloatResult[0])
+		.Read(BuiltInComponents->DoubleResult[0])
 		.FilterNone({ BuiltInComponents->BlendChannelInput })
 		.SetDesiredThread(Linker->EntityManager.GetDispatchThread())
 		.Dispatch_PerEntity<FApplyScalarParameters>(&Linker->EntityManager, InPrerequisites, &Subsequents);
@@ -287,9 +287,9 @@ void UMovieSceneMaterialParameterSystem::OnEvaluation(FSystemTaskPrerequisites&
 		FEntityTaskBuilder()
 		.Read(TracksComponents->BoundMaterial)
 		.ReadOneOrMoreOf(TracksComponents->VectorParameterName, TracksComponents->ColorParameterName)
-		.ReadOneOrMoreOf(BuiltInComponents->FloatResult[0], BuiltInComponents->FloatResult[1], BuiltInComponents->FloatResult[2], BuiltInComponents->FloatResult[3])
+		.ReadOneOrMoreOf(BuiltInComponents->DoubleResult[0], BuiltInComponents->DoubleResult[1], BuiltInComponents->DoubleResult[2], BuiltInComponents->DoubleResult[3])
 		.FilterNone({ BuiltInComponents->BlendChannelInput })
 		.SetDesiredThread(Linker->EntityManager.GetDispatchThread())
 		.Dispatch_PerAllocation<FApplyVectorParameters>(&Linker->EntityManager, InPrerequisites, &Subsequents);
 	}
-}
\ No newline at end of file
+}
diff --git a/Engine/Source/Runtime/MovieSceneTracks/Private/Systems/MovieSceneMotionVectorSimulationSystem.cpp b/Engine/Source/Runtime/MovieSceneTracks/Private/Systems/MovieSceneMotionVectorSimulationSystem.cpp
--- a/Engine/Source/Runtime/MovieSceneTracks/Private/Systems/MovieSceneMotionVectorSimulationSystem.cpp
+++ b/Engine/Source/Runtime/MovieSceneTracks/Private/Systems/MovieSceneMotionVectorSimulationSystem.cpp
@@ -12,7 +12,6 @@
 #include "MovieSceneTracksComponentTypes.h"
 
 #include "Systems/FloatChannelEvaluatorSystem.h"
-#include "Systems/MovieScenePiecewiseFloatBlenderSystem.h"
 #include "Systems/MovieSceneComponentTransformSystem.h"
 
 #include "Rendering/MotionVectorSimulation.h"
@@ -153,17 +152,17 @@ void UMovieSceneMotionVectorSimulationSystem::ComputeSimulatedMotion()
 	FEntityTaskBuilder()
 	.Read(BuiltInComponents->BoundObject)
 	.ReadAnyOf(
-		BuiltInComponents->FloatResult[0], BuiltInComponents->FloatResult[1], BuiltInComponents->FloatResult[2],
-		BuiltInComponents->FloatResult[3], BuiltInComponents->FloatResult[4], BuiltInComponents->FloatResult[5],
-		BuiltInComponents->FloatResult[6], BuiltInComponents->FloatResult[7], BuiltInComponents->FloatResult[8])
+		BuiltInComponents->DoubleResult[0], BuiltInComponents->DoubleResult[1], BuiltInComponents->DoubleResult[2],
+		BuiltInComponents->DoubleResult[3], BuiltInComponents->DoubleResult[4], BuiltInComponents->DoubleResult[5],
+		BuiltInComponents->DoubleResult[6], BuiltInComponents->DoubleResult[7], BuiltInComponents->DoubleResult[8])
 	.FilterAll({ TracksComponents->ComponentTransform.PropertyTag })
 	.FilterAny({ BuiltInComponents->CustomPropertyIndex, BuiltInComponents->FastPropertyOffset, BuiltInComponents->SlowProperty })
 	.FilterAny({ 
 		// must have at least one float result component
-		BuiltInComponents->FloatResult[0], BuiltInComponents->FloatResult[1], BuiltInComponents->FloatResult[2],
-		BuiltInComponents->FloatResult[3], BuiltInComponents->FloatResult[4], BuiltInComponents->FloatResult[5],
-		BuiltInComponents->FloatResult[6], BuiltInComponents->FloatResult[7], BuiltInComponents->FloatResult[8] })
-	.Iterate_PerEntity(&Linker->EntityManager, [this](UObject* BoundObject, const float* T_X, const float* T_Y, const float* T_Z, const float* R_X, const float* R_Y, const float* R_Z, const float* S_X, const float* S_Y, const float* S_Z)
+		BuiltInComponents->DoubleResult[0], BuiltInComponents->DoubleResult[1], BuiltInComponents->DoubleResult[2],
+		BuiltInComponents->DoubleResult[3], BuiltInComponents->DoubleResult[4], BuiltInComponents->DoubleResult[5],
+		BuiltInComponents->DoubleResult[6], BuiltInComponents->DoubleResult[7], BuiltInComponents->DoubleResult[8] })
+	.Iterate_PerEntity(&Linker->EntityManager, [this](UObject* BoundObject, const double* T_X, const double* T_Y, const double* T_Z, const double* R_X, const double* R_Y, const double* R_Z, const double* S_X, const double* S_Y, const double* S_Z)
 	{
 		if (USceneComponent* SceneComponent = Cast<USceneComponent>(BoundObject))
 		{
diff --git a/Engine/Source/Runtime/MovieSceneTracks/Private/Systems/MovieScenePiecewiseBlenderSystemImpl.cpp b/Engine/Source/Runtime/MovieSceneTracks/Private/Systems/MovieScenePiecewiseBlenderSystemImpl.cpp
deleted file mode 100644
diff --git a/Engine/Source/Runtime/MovieSceneTracks/Private/Systems/MovieScenePiecewiseDoubleBlenderSystem.cpp b/Engine/Source/Runtime/MovieSceneTracks/Private/Systems/MovieScenePiecewiseDoubleBlenderSystem.cpp
--- a/Engine/Source/Runtime/MovieSceneTracks/Private/Systems/MovieScenePiecewiseDoubleBlenderSystem.cpp
+++ b/Engine/Source/Runtime/MovieSceneTracks/Private/Systems/MovieScenePiecewiseDoubleBlenderSystem.cpp
@@ -2,23 +2,413 @@
 
 #include "Systems/MovieScenePiecewiseDoubleBlenderSystem.h"
 #include "EntitySystem/BuiltInComponentTypes.h"
+#include "EntitySystem/EntityAllocationIterator.h"
+#include "EntitySystem/MovieSceneEntityManager.h"
 #include "EntitySystem/MovieSceneEntitySystemLinker.h"
+#include "EntitySystem/MovieSceneEntitySystemTask.h"
+#include "Systems/FloatChannelEvaluatorSystem.h"
 #include "Systems/DoubleChannelEvaluatorSystem.h"
 #include "Systems/MovieSceneBaseValueEvaluatorSystem.h"
 #include "Systems/MovieSceneQuaternionInterpolationRotationSystem.h"
 #include "Systems/WeightAndEasingEvaluatorSystem.h"
 
-
 DECLARE_CYCLE_STAT(TEXT("Piecewise Double Blender System"),       MovieSceneEval_PiecewiseDoubleBlenderSystem,  STATGROUP_MovieSceneECS);
 DECLARE_CYCLE_STAT(TEXT("Blend double values"),                   MovieSceneEval_BlendDoubleValues,        STATGROUP_MovieSceneECS);
 DECLARE_CYCLE_STAT(TEXT("Default combine blended double values"), MovieSceneEval_BlendCombineDoubleValues, STATGROUP_MovieSceneECS);
  
+namespace UE
+{
+namespace MovieScene
+{
+
+/** Task for accumulating all weighted blend inputs into arrays based on BlendID. Will be run for Absolute, Additive and Relative blend modes*/
+struct FAccumulationTask
+{
+	FAccumulationTask(TSortedMap<FComponentTypeID, TArray<FBlendResult>>* InAccumulationBuffers)
+		: AccumulationBuffers(InAccumulationBuffers)
+	{}
+
+	/** Task entry point - iterates the allocation's headers and accumulates results for any required components */
+	void ForEachAllocation(FEntityAllocationIteratorItem InItem, TRead<FMovieSceneBlendChannelID> BlendIDs, TReadOptional<double> OptionalEasingAndWeights)
+	{
+		const FEntityAllocation* Allocation = InItem;
+		const FComponentMask& AllocationType = InItem;
+
+		for (const FComponentHeader& ComponentHeader : Allocation->GetComponentHeaders())
+		{
+			if (TArray<FBlendResult>* AccumulationBuffer = AccumulationBuffers->Find(ComponentHeader.ComponentType))
+			{
+				ComponentHeader.ReadWriteLock.ReadLock();
+
+				const double* Results = static_cast<const double*>(ComponentHeader.GetValuePtr(0));
+				AccumulateResults(Allocation, Results, BlendIDs, OptionalEasingAndWeights, *AccumulationBuffer);
+
+				ComponentHeader.ReadWriteLock.ReadUnlock();
+			}
+		}
+	}
+
+private:
+
+	void AccumulateResults(const FEntityAllocation* InAllocation, const double* InResults, const FMovieSceneBlendChannelID* BlendIDs, const double* OptionalEasingAndWeights, TArray<FBlendResult>& OutBlendResults)
+	{
+		static const FMovieSceneBlenderSystemID BlenderSystemID = UMovieSceneBlenderSystem::GetBlenderSystemID<UMovieScenePiecewiseDoubleBlenderSystem>();
+
+		const int32 Num = InAllocation->Num();
+		if (OptionalEasingAndWeights)
+		{
+			// We have some easing/weight factors to multiply values with.
+			for (int32 Index = 0; Index < Num; ++Index)
+			{
+				const FMovieSceneBlendChannelID& BlendID(BlendIDs[Index]);
+				ensureMsgf(BlendID.SystemID == BlenderSystemID, TEXT("Overriding the standard blender system of standard types isn't supported."));
+
+				FBlendResult& Result = OutBlendResults[BlendID.ChannelID];
+
+				const float Weight = OptionalEasingAndWeights[Index];
+				Result.Total  += InResults[Index] * Weight;
+				Result.Weight += Weight;
+			}
+		}
+		else
+		{
+			// Faster path for when there's no weight to multiply values with.
+			for (int32 Index = 0; Index < Num; ++Index)
+			{
+				const FMovieSceneBlendChannelID& BlendID(BlendIDs[Index]);
+				ensureMsgf(BlendID.SystemID == BlenderSystemID, TEXT("Overriding the standard blender system of standard types isn't supported."));
+
+				FBlendResult& Result = OutBlendResults[BlendID.ChannelID];
+				Result.Total  += InResults[Index];
+				Result.Weight += 1.f;
+			}
+		}
+	}
+
+	TSortedMap<FComponentTypeID, TArray<FBlendResult>>* AccumulationBuffers;
+};
+
+/** Same as the task above, but also reads a "base value" that is subtracted from all values.
+ *
+ *  Only used by entities with the "additive from base" blend type.
+ */
+struct FAdditiveFromBaseBlendTask
+{
+	TSortedMap<FComponentTypeID, FAdditiveFromBaseBuffer>* AccumulationBuffers;
+
+	void ForEachAllocation(FEntityAllocationIteratorItem InItem, TRead<FMovieSceneBlendChannelID> BlendIDs, TReadOptional<double> EasingAndWeightResults)
+	{
+		FEntityAllocation* Allocation = InItem;
+		const FComponentMask& AllocationType = InItem;
+
+		for (const FComponentHeader& ComponentHeader : Allocation->GetComponentHeaders())
+		{
+			if (FAdditiveFromBaseBuffer* Buffer = AccumulationBuffers->Find(ComponentHeader.ComponentType))
+			{
+				TComponentReader<double> BaseValues = Allocation->ReadComponents(Buffer->BaseComponent.template ReinterpretCast<double>());
+				TComponentReader<double> Results(&ComponentHeader);
+
+				AccumulateResults(Allocation, Results.AsPtr(), BaseValues.AsPtr(), BlendIDs, EasingAndWeightResults, Buffer->Buffer);
+			}
+		}
+	}
+
+private:
+
+	void AccumulateResults(const FEntityAllocation* InAllocation, const double* Results, const double* BaseValues, const FMovieSceneBlendChannelID* BlendIDs, const double* OptionalEasingAndWeights, TArray<FBlendResult>& OutBlendResults)
+	{
+		static const FMovieSceneBlenderSystemID BlenderSystemID = UMovieSceneBlenderSystem::GetBlenderSystemID<UMovieScenePiecewiseDoubleBlenderSystem>();
+
+		const int32 Num = InAllocation->Num();
+
+		if (OptionalEasingAndWeights)
+		{
+			for (int32 Index = 0; Index < Num; ++Index)
+			{
+				const FMovieSceneBlendChannelID& BlendID(BlendIDs[Index]);
+				ensureMsgf(BlendID.SystemID == BlenderSystemID, TEXT("Overriding the standard blender system of standard types isn't supported."));
+
+				FBlendResult& Result = OutBlendResults[BlendID.ChannelID];
+
+				const float Weight = OptionalEasingAndWeights[Index];
+				Result.Total  += (Results[Index] - BaseValues[Index]) * Weight;
+				Result.Weight += Weight;
+			}
+		}
+		else
+		{
+			// Faster path for when there's no weight to multiply values with.
+			for (int32 Index = 0; Index < Num; ++Index)
+			{
+				const FMovieSceneBlendChannelID& BlendID(BlendIDs[Index]);
+				ensureMsgf(BlendID.SystemID == BlenderSystemID, TEXT("Overriding the standard blender system of standard types isn't supported."));
+
+				FBlendResult& Result = OutBlendResults[BlendID.ChannelID];
+				Result.Total  += (Results[Index] - BaseValues[Index]);
+				Result.Weight += 1.f;
+			}
+		}
+	}
+};
+
+void BlendResults(const FAccumulationResult& Results, uint16 BlendID, double& OutFinalBlendResult)
+{
+	FBlendResult AbsoluteResult = Results.GetAbsoluteResult(BlendID);
+	FBlendResult AdditiveResult = Results.GetAdditiveResult(BlendID);
+	FBlendResult AdditiveFromBaseResult = Results.GetAdditiveFromBaseResult(BlendID);
+
+#if DO_GUARD_SLOW
+	ensureMsgf(AbsoluteResult.Weight != 0.f, TEXT("Default blend combine being used for an entity that has no absolute weight. This should have an initial value and should be handled by each system, and excluded by default with UMovieSceneBlenderSystem::FinalCombineExclusionFilter ."));
+#endif
+
+	const float TotalWeight = AbsoluteResult.Weight;
+	if (TotalWeight != 0)
+	{
+		const double Value = AbsoluteResult.Total / AbsoluteResult.Weight + AdditiveResult.Total + AdditiveFromBaseResult.Total;
+		OutFinalBlendResult = Value;
+	}
+}
+
+/** Task that combines all accumulated blends for any tracked property type that has blend inputs/outputs */
+struct FCombineBlendsForProperties
+{
+	explicit FCombineBlendsForProperties(const TBitArray<>& InCachedRelevantProperties, const FAccumulationBuffers* InAccumulationBuffers, FEntityAllocationWriteContext InWriteContext)
+		: CachedRelevantProperties(InCachedRelevantProperties)
+		, AccumulationBuffers(InAccumulationBuffers)
+		, PropertyRegistry(&FBuiltInComponentTypes::Get()->PropertyRegistry)
+		, WriteContext(InWriteContext)
+	{}
+
+	void ForEachAllocation(FEntityAllocationIteratorItem InItem, TRead<FMovieSceneBlendChannelID> BlendIDs)
+	{
+		FEntityAllocation* Allocation = InItem;
+		const FComponentMask& AllocationType = InItem;
+
+		// Find out what kind of property this is
+		for (TConstSetBitIterator<> PropertyIndex(CachedRelevantProperties); PropertyIndex; ++PropertyIndex)
+		{
+			const FPropertyDefinition& PropertyDefinition = PropertyRegistry->GetDefinition(FCompositePropertyTypeID::FromIndex(PropertyIndex.GetIndex()));
+			if (AllocationType.Contains(PropertyDefinition.PropertyType))
+			{
+				ProcessPropertyType(Allocation, AllocationType, PropertyDefinition, BlendIDs);
+				return;
+			}
+		}
+	}
+
+private:
+
+	void ProcessPropertyType(FEntityAllocation* Allocation, const FComponentMask& AllocationType, const FPropertyDefinition& PropertyDefinition, const FMovieSceneBlendChannelID* BlendIDs)
+	{
+		static const FMovieSceneBlenderSystemID BlenderSystemID = UMovieSceneBlenderSystem::GetBlenderSystemID<UMovieScenePiecewiseDoubleBlenderSystem>();
+
+		TArrayView<const FPropertyCompositeDefinition> Composites = PropertyRegistry->GetComposites(PropertyDefinition);
+
+		FOptionalComponentReader OptInitialValues = Allocation->TryReadComponentsErased(PropertyDefinition.InitialValueType);
+
+		for (int32 CompositeIndex = 0; CompositeIndex < Composites.Num(); ++CompositeIndex)
+		{
+			const bool bIsCompositeSupported = ((PropertyDefinition.DoubleCompositeMask & (1 << CompositeIndex)) != 0);
+			if (!bIsCompositeSupported)
+			{
+				continue;
+			}
+
+			TComponentTypeID<double> ResultComponent = Composites[CompositeIndex].ComponentTypeID.template ReinterpretCast<double>();
+			if (!AllocationType.Contains(ResultComponent))
+			{
+				continue;
+			}
+
+			FAccumulationResult Results = AccumulationBuffers->FindResults(ResultComponent);
+			if (Results.IsValid())
+			{
+				const uint16 InitialValueProjectionOffset = Composites[CompositeIndex].CompositeOffset;
+
+				// Open the result channel for write
+				TComponentWriter<double> ValueResults = Allocation->WriteComponents(ResultComponent, WriteContext);
+
+				if (OptInitialValues)
+				{
+					for (int32 Index = 0; Index < Allocation->Num(); ++Index)
+					{
+						ensureMsgf(BlendIDs[Index].SystemID == BlenderSystemID, TEXT("Overriding the standard blender system of standard types isn't supported."));
+						const double InitialValue = *reinterpret_cast<const double*>(static_cast<const uint8*>(OptInitialValues[Index]) + InitialValueProjectionOffset);
+						BlendResultsWithInitial(Results, BlendIDs[Index].ChannelID, InitialValue, ValueResults[Index]);
+					}
+				}
+				else
+				{
+					for (int32 Index = 0; Index < Allocation->Num(); ++Index)
+					{
+						ensureMsgf(BlendIDs[Index].SystemID == BlenderSystemID, TEXT("Overriding the standard blender system of standard types isn't supported."));
+						BlendResults(Results, BlendIDs[Index].ChannelID, ValueResults[Index]);
+					}
+				}
+			}
+		}
+	}
+
+	void BlendResultsWithInitial(const FAccumulationResult& Results, uint16 BlendID, const double InitialValue, double& OutFinalBlendResult)
+	{
+		FBlendResult AbsoluteResult = Results.GetAbsoluteResult(BlendID);
+		FBlendResult RelativeResult = Results.GetRelativeResult(BlendID);
+		FBlendResult AdditiveResult = Results.GetAdditiveResult(BlendID);
+		FBlendResult AdditiveFromBaseResult = Results.GetAdditiveFromBaseResult(BlendID);
+
+		if (RelativeResult.Weight != 0)
+		{
+			RelativeResult.Total += InitialValue * RelativeResult.Weight;
+		}
+
+		FBlendResult TotalAdditiveResult = { AdditiveResult.Total + AdditiveFromBaseResult.Total, AdditiveResult.Weight + AdditiveFromBaseResult.Weight };
+
+		const float TotalWeight = AbsoluteResult.Weight + RelativeResult.Weight;
+		if (TotalWeight != 0)
+		{
+			// If the absolute value has some partial weighting (for ease-in/out for instance), we ramp it from/to the initial value. This means
+			// that the "initial value" adds a contribution to the entire blending process, so we add its weight to the total that we
+			// normalize absolutes and relatives with.
+			//
+			// Note that "partial weighting" means strictly between 0 and 100%. At 100% and above, we don't need to do this thing with the initial
+			// value. At 0%, we have no absolute value (only a relative value) and we therefore don't want to include the initial value either.
+			const bool bInitialValueContributes = (0.f < AbsoluteResult.Weight && AbsoluteResult.Weight < 1.f);
+			const double AbsoluteBlendedValue = bInitialValueContributes ?
+				(InitialValue * (1.f - AbsoluteResult.Weight) + AbsoluteResult.Total) :
+				AbsoluteResult.Total;
+			const float FinalTotalWeight = bInitialValueContributes ? (TotalWeight + (1.f - AbsoluteResult.Weight)) : TotalWeight;
+
+			const double Value = (AbsoluteBlendedValue + RelativeResult.Total) / FinalTotalWeight + TotalAdditiveResult.Total;
+			OutFinalBlendResult = Value;
+		}
+		else if (TotalAdditiveResult.Weight != 0)
+		{
+			OutFinalBlendResult = TotalAdditiveResult.Total + InitialValue;
+		}
+		else
+		{
+			OutFinalBlendResult = InitialValue;
+		}
+	}
+
+	void BlendResults(const FAccumulationResult& Results, uint16 BlendID, double& OutFinalBlendResult)
+	{
+		FBlendResult AbsoluteResult = Results.GetAbsoluteResult(BlendID);
+		FBlendResult AdditiveResult = Results.GetAdditiveResult(BlendID);
+		FBlendResult AdditiveFromBaseResult = Results.GetAdditiveFromBaseResult(BlendID);
+
+#if DO_GUARD_SLOW
+		ensureMsgf(AbsoluteResult.Weight != 0.f, TEXT("Default blend combine being used for an entity that has no absolute weight. This should have an initial value and should be handled by each system, and excluded by default with UMovieSceneBlenderSystem::FinalCombineExclusionFilter ."));
+#endif
+
+		const float TotalWeight = AbsoluteResult.Weight;
+		if (TotalWeight != 0)
+		{
+			const double Value = AbsoluteResult.Total / AbsoluteResult.Weight + AdditiveResult.Total + AdditiveFromBaseResult.Total;
+			OutFinalBlendResult = Value;
+		}
+	}
+
+private:
+
+	TBitArray<> CachedRelevantProperties;
+	const FAccumulationBuffers* AccumulationBuffers;
+	const FPropertyRegistry* PropertyRegistry;
+	FEntityAllocationWriteContext WriteContext;
+};
+
+
+/** Task that combines all accumulated blends for any tracked non-property type that has blend inputs/outputs */
+struct FCombineBlends
+{
+	explicit FCombineBlends(const FAccumulationBuffers* InAccumulationBuffers, FEntityAllocationWriteContext InWriteContext)
+		: AccumulationBuffers(InAccumulationBuffers)
+		, WriteContext(InWriteContext)
+	{}
+
+	void ForEachAllocation(FEntityAllocationIteratorItem InItem, TRead<FMovieSceneBlendChannelID> BlendIDs)
+	{
+		static const FMovieSceneBlenderSystemID BlenderSystemID = UMovieSceneBlenderSystem::GetBlenderSystemID<UMovieScenePiecewiseDoubleBlenderSystem>();
+
+		FEntityAllocation* Allocation = InItem;
+		const FComponentMask& AllocationType = InItem;
+
+		// Find all result types on this allocation
+		TArrayView<TComponentTypeID<double>> ResultComponents(FBuiltInComponentTypes::Get()->DoubleResult);
+		for (int32 ResultIndex = 0; ResultIndex < ResultComponents.Num(); ++ResultIndex)
+		{
+			TComponentTypeID<double> ResultComponent = ResultComponents[ResultIndex];
+			if (!AllocationType.Contains(ResultComponent))
+			{
+				continue;
+			}
+
+			FAccumulationResult Results = AccumulationBuffers->FindResults(ResultComponent);
+			if (!Results.IsValid())
+			{
+				continue;
+			}
+
+			// Open the result channel for write
+			TComponentWriter<double> ValueResults = Allocation->WriteComponents(ResultComponent, WriteContext);
+			for (int32 Index = 0; Index < Allocation->Num(); ++Index)
+			{
+				ensureMsgf(BlendIDs[Index].SystemID == BlenderSystemID, TEXT("Overriding the standard blender system of standard types isn't supported."));
+				BlendResults(Results, BlendIDs[Index].ChannelID, ValueResults[Index]);
+			}
+		}
+	}
+
+private:
+
+	const FAccumulationBuffers* AccumulationBuffers;
+	FEntityAllocationWriteContext WriteContext;
+};
+
+bool FAccumulationBuffers::IsEmpty() const
+{
+	return Absolute.Num() == 0 && Relative.Num() == 0 && Additive.Num() == 0 && AdditiveFromBase.Num() == 0;
+}
+
+void FAccumulationBuffers::Reset()
+{
+	Absolute.Empty();
+	Relative.Empty();
+	Additive.Empty();
+	AdditiveFromBase.Empty();
+}
+
+FAccumulationResult FAccumulationBuffers::FindResults(FComponentTypeID InComponentType) const
+{
+	FAccumulationResult Result;
+	if (const TArray<FBlendResult>* Absolutes = Absolute.Find(InComponentType))
+	{
+		Result.Absolutes = Absolutes->GetData();
+	}
+	if (const TArray<FBlendResult>* Relatives = Relative.Find(InComponentType))
+	{
+		Result.Relatives = Relatives->GetData();
+	}
+	if (const TArray<FBlendResult>* Additives = Additive.Find(InComponentType))
+	{
+		Result.Additives = Additives->GetData();
+	}
+	if (const FAdditiveFromBaseBuffer* AdditivesFromBase = AdditiveFromBase.Find(InComponentType))
+	{
+		Result.AdditivesFromBase = AdditivesFromBase->Buffer.GetData();
+	}
+	return Result;
+}
+
+} // namespace MovieScene
+} // namespace UE
 
 UMovieScenePiecewiseDoubleBlenderSystem::UMovieScenePiecewiseDoubleBlenderSystem(const FObjectInitializer& ObjInit)
 	: Super(ObjInit)
 {
 	if (HasAnyFlags(RF_ClassDefaultObject))
 	{
+		DefineImplicitPrerequisite(UFloatChannelEvaluatorSystem::StaticClass(), GetClass());
 		DefineImplicitPrerequisite(UDoubleChannelEvaluatorSystem::StaticClass(), GetClass());
 		DefineImplicitPrerequisite(UMovieSceneBaseValueEvaluatorSystem::StaticClass(), GetClass());
 		DefineImplicitPrerequisite(UMovieSceneQuaternionInterpolationRotationSystem::StaticClass(), GetClass());
@@ -34,15 +424,331 @@ void UMovieScenePiecewiseDoubleBlenderSystem::OnRun(FSystemTaskPrerequisites& In
 
 	CompactBlendChannels();
 
-	FPiecewiseBlenderSystemImplRunParams Params;
-	Params.MaximumNumBlends = AllocatedBlendChannels.Num();
-	Params.BlendValuesStatId = GET_STATID(MovieSceneEval_BlendDoubleValues);
-	Params.CombineBlendsStatId = GET_STATID(MovieSceneEval_BlendCombineDoubleValues);
-	Impl.Run(Params, Linker->EntityManager, InPrerequisites, Subsequents);
+	FEntityManager& EntityManager = Linker->EntityManager;
+	const TStatId BlendValuesStatId = GET_STATID(MovieSceneEval_BlendDoubleValues);
+	const TStatId CombineBlendsStatId = GET_STATID(MovieSceneEval_BlendCombineDoubleValues);
+
+	// We allocate space for every blend even if there are gaps so we can do a straight index into each array
+	if (AllocatedBlendChannels.Num() == 0)
+	{
+		return;
+	}
+
+	// Update cached channel data if necessary
+	if (ChannelRelevancyCache.Update(EntityManager) == ECachedEntityManagerState::Stale)
+	{
+		ReinitializeAccumulationBuffers();
+	}
+
+	if (AccumulationBuffers.IsEmpty())
+	{
+		return;
+	}
+
+	ZeroAccumulationBuffers();
+
+	FBuiltInComponentTypes* BuiltInComponents = FBuiltInComponentTypes::Get();
+
+	FSystemTaskPrerequisites Prereqs;
+	if (AccumulationBuffers.Absolute.Num() != 0)
+	{
+		FGraphEventRef Task = FEntityTaskBuilder()
+		.Read(BuiltInComponents->BlendChannelInput)
+		.ReadOptional(BuiltInComponents->WeightAndEasingResult)
+		.FilterAll({ BuiltInComponents->Tags.AbsoluteBlend })
+		.FilterAny(BlendedResultMask)
+		.FilterNone({ BuiltInComponents->Tags.Ignored })
+		.SetStat(BlendValuesStatId)
+		.template Dispatch_PerAllocation<FAccumulationTask>(&EntityManager, InPrerequisites, nullptr, &AccumulationBuffers.Absolute);
+
+		if (Task)
+		{
+			Prereqs.AddMasterTask(Task);
+		}
+	}
+
+	if (AccumulationBuffers.Relative.Num() != 0)
+	{
+		FGraphEventRef Task = FEntityTaskBuilder()
+		.Read(BuiltInComponents->BlendChannelInput)
+		.ReadOptional(BuiltInComponents->WeightAndEasingResult)
+		.FilterAll({ BuiltInComponents->Tags.RelativeBlend })
+		.FilterAny(BlendedResultMask)
+		.FilterNone({ BuiltInComponents->Tags.Ignored })
+		.SetStat(BlendValuesStatId)
+		.template Dispatch_PerAllocation<FAccumulationTask>(&EntityManager, InPrerequisites, nullptr, &AccumulationBuffers.Relative);
+
+		if (Task)
+		{
+			Prereqs.AddMasterTask(Task);
+		}
+	}
+
+	if (AccumulationBuffers.Additive.Num() != 0)
+	{
+		FGraphEventRef Task = FEntityTaskBuilder()
+		.Read(BuiltInComponents->BlendChannelInput)
+		.ReadOptional(BuiltInComponents->WeightAndEasingResult)
+		.FilterAll({ BuiltInComponents->Tags.AdditiveBlend })
+		.FilterAny(BlendedResultMask)
+		.FilterNone({ BuiltInComponents->Tags.Ignored })
+		.SetStat(BlendValuesStatId)
+		.template Dispatch_PerAllocation<FAccumulationTask>(&EntityManager, InPrerequisites, nullptr, &AccumulationBuffers.Additive);
+
+		if (Task)
+		{
+			Prereqs.AddMasterTask(Task);
+		}
+	}
+
+	if (AccumulationBuffers.AdditiveFromBase.Num() != 0)
+	{
+		FGraphEventRef Task = FEntityTaskBuilder()
+		.Read(BuiltInComponents->BlendChannelInput)
+		.ReadOptional(BuiltInComponents->WeightAndEasingResult)
+		.FilterAll({ BuiltInComponents->Tags.AdditiveFromBaseBlend })
+		.FilterAny(BlendedResultMask)
+		.FilterNone({ BuiltInComponents->Tags.Ignored })
+		.SetStat(BlendValuesStatId)
+		.template Dispatch_PerAllocation<FAdditiveFromBaseBlendTask>(&EntityManager, InPrerequisites, nullptr, &AccumulationBuffers.AdditiveFromBase);
+
+		if (Task)
+		{
+			Prereqs.AddMasterTask(Task);
+		}
+	}
+
+	// Combine blends for all blend outputs on properties
+	if (BlendedPropertyMask.Find(true) != INDEX_NONE)
+	{
+		FEntityTaskBuilder()
+		.Read(BuiltInComponents->BlendChannelOutput)
+		.FilterAny(BlendedPropertyMask)
+		.SetStat(CombineBlendsStatId)
+		.template Dispatch_PerAllocation<FCombineBlendsForProperties>(&EntityManager, Prereqs, &Subsequents, CachedRelevantProperties, &AccumulationBuffers, FEntityAllocationWriteContext(EntityManager));
+	}
+
+	if (bContainsNonPropertyBlends)
+	{
+		// Blend task that combines vanilla (non-property-based) components
+		FEntityTaskBuilder()
+		.Read(BuiltInComponents->BlendChannelOutput)
+		.FilterAny(BlendedResultMask)
+		.FilterNone(BlendedPropertyMask)
+		.SetStat(CombineBlendsStatId)
+		.template Dispatch_PerAllocation<FCombineBlends>(&EntityManager, Prereqs, &Subsequents, &AccumulationBuffers, FEntityAllocationWriteContext(EntityManager));
+	}
+}
+
+void UMovieScenePiecewiseDoubleBlenderSystem::ReinitializeAccumulationBuffers()
+{
+	using namespace UE::MovieScene;
+
+	FBuiltInComponentTypes* BuiltInComponents = FBuiltInComponentTypes::Get();
+
+	BlendedResultMask.Reset();
+	AccumulationBuffers.Reset();
+
+	TArrayView<TComponentTypeID<double>> BaseComponents(FBuiltInComponentTypes::Get()->BaseDouble);
+	TArrayView<TComponentTypeID<double>> ResultComponents(FBuiltInComponentTypes::Get()->DoubleResult);
+	check(BaseComponents.Num() == ResultComponents.Num());
+
+	// Recompute which result types are blended
+	const int32 NumResults = ResultComponents.Num();
+	const int32 MaximumNumBlends = AllocatedBlendChannels.Num();
+	FEntityManager& EntityManager = Linker->EntityManager;
+	for (int32 Index = 0; Index < NumResults; ++Index)
+	{
+		TComponentTypeID<double> Component = ResultComponents[Index];
+
+		const bool bHasAbsolutes         = EntityManager.Contains(FEntityComponentFilter().All({ Component, BuiltInComponents->BlendChannelInput, BuiltInComponents->Tags.AbsoluteBlend }));
+		const bool bHasRelatives         = EntityManager.Contains(FEntityComponentFilter().All({ Component, BuiltInComponents->BlendChannelInput, BuiltInComponents->Tags.RelativeBlend }));
+		const bool bHasAdditives         = EntityManager.Contains(FEntityComponentFilter().All({ Component, BuiltInComponents->BlendChannelInput, BuiltInComponents->Tags.AdditiveBlend }));
+		const bool bHasAdditivesFromBase = EntityManager.Contains(FEntityComponentFilter().All({ Component, BuiltInComponents->BlendChannelInput, BuiltInComponents->Tags.AdditiveFromBaseBlend }));
+
+		if (!(bHasAbsolutes || bHasRelatives || bHasAdditives || bHasAdditivesFromBase))
+		{
+			continue;
+		}
+
+		BlendedResultMask.Set(Component);
+
+		if (bHasAbsolutes)
+		{
+			TArray<FBlendResult>& Buffer = AccumulationBuffers.Absolute.Add(Component);
+			Buffer.SetNum(MaximumNumBlends);
+		}
+		if (bHasRelatives)
+		{
+			TArray<FBlendResult>& Buffer = AccumulationBuffers.Relative.Add(Component);
+			Buffer.SetNum(MaximumNumBlends);
+		}
+		if (bHasAdditives)
+		{
+			TArray<FBlendResult>& Buffer = AccumulationBuffers.Additive.Add(Component);
+			Buffer.SetNum(MaximumNumBlends);
+		}
+		if (bHasAdditivesFromBase)
+		{
+			FAdditiveFromBaseBuffer& Buffer = AccumulationBuffers.AdditiveFromBase.Add(Component);
+			Buffer.Buffer.SetNum(MaximumNumBlends);
+			Buffer.BaseComponent = BaseComponents[Index];
+		}
+	}
+
+	// Update property relevancy
+	CachedRelevantProperties.Empty();
+
+	// If we have no accumulation buffers, we have nothing more to do
+	if (AccumulationBuffers.IsEmpty())
+	{
+		return;
+	}
+
+	FComponentMask AllPropertyTypes;
+
+	// This code works on the assumption that properties can never be removed (which is safe)
+	FEntityComponentFilter InclusionFilter;
+	TArrayView<const FPropertyDefinition> Properties = BuiltInComponents->PropertyRegistry.GetProperties();
+	for (int32 PropertyTypeIndex = 0; PropertyTypeIndex < Properties.Num(); ++PropertyTypeIndex)
+	{
+		const FPropertyDefinition& PropertyDefinition = Properties[PropertyTypeIndex];
+		const bool bHasAnyComposite = (PropertyDefinition.DoubleCompositeMask != 0);
+		if (bHasAnyComposite)
+		{
+			AllPropertyTypes.Set(PropertyDefinition.PropertyType);
+
+			InclusionFilter.Reset();
+			InclusionFilter.All({ BuiltInComponents->BlendChannelOutput, PropertyDefinition.PropertyType });
+			if (EntityManager.Contains(InclusionFilter))
+			{
+				CachedRelevantProperties.PadToNum(PropertyTypeIndex + 1, false);
+				CachedRelevantProperties[PropertyTypeIndex] = true;
+
+				BlendedPropertyMask.Set(PropertyDefinition.PropertyType);
+			}
+		}
+	}
+
+	bContainsNonPropertyBlends = EntityManager.Contains(FEntityComponentFilter().All({ BuiltInComponents->BlendChannelOutput }).None(AllPropertyTypes));
+}
+
+void UMovieScenePiecewiseDoubleBlenderSystem::ZeroAccumulationBuffers()
+{
+	using namespace UE::MovieScene;
+
+	// Arrays should only ever exist in these containers if they have size (they are always initialized to MaximumNumBlends in ReinitializeAccumulationBuffers)
+	for (TPair<FComponentTypeID, TArray<FBlendResult>>& Pair : AccumulationBuffers.Absolute)
+	{
+		FMemory::Memzero(Pair.Value.GetData(), sizeof(FBlendResult)*Pair.Value.Num());
+	}
+	for (TPair<FComponentTypeID, TArray<FBlendResult>>& Pair : AccumulationBuffers.Relative)
+	{
+		FMemory::Memzero(Pair.Value.GetData(), sizeof(FBlendResult)*Pair.Value.Num());
+	}
+	for (TPair<FComponentTypeID, TArray<FBlendResult>>& Pair : AccumulationBuffers.Additive)
+	{
+		FMemory::Memzero(Pair.Value.GetData(), sizeof(FBlendResult)*Pair.Value.Num());
+	}
+	for (TPair<FComponentTypeID, FAdditiveFromBaseBuffer>& Pair : AccumulationBuffers.AdditiveFromBase)
+	{
+		FMemory::Memzero(Pair.Value.Buffer.GetData(), sizeof(FBlendResult)*Pair.Value.Buffer.Num());
+	}
 }
 
 FGraphEventRef UMovieScenePiecewiseDoubleBlenderSystem::DispatchDecomposeTask(const UE::MovieScene::FValueDecompositionParams& Params, UE::MovieScene::FAlignedDecomposedValue* Output)
 {
-	return Impl.DispatchDecomposeTask(Linker->EntityManager, Params, Output);
+	using namespace UE::MovieScene;
+
+	if (!Params.ResultComponentType)
+	{
+		return nullptr;
+	}
+
+	TComponentTypeID<double> ResultComponentType = Params.ResultComponentType.ReinterpretCast<double>();
+
+	struct FChannelResultTask
+	{
+		TArray<FMovieSceneEntityID, TInlineAllocator<8>> EntitiesToDecompose;
+		FAlignedDecomposedValue* Result;
+		uint16 DecomposeBlendChannel;
+		FComponentTypeID AdditiveBlendTag;
+
+		explicit FChannelResultTask(const FValueDecompositionParams& Params, FAlignedDecomposedValue* InResult)
+			: Result(InResult)
+			, DecomposeBlendChannel(Params.DecomposeBlendChannel)
+			, AdditiveBlendTag(FBuiltInComponentTypes::Get()->Tags.AdditiveBlend)
+		{
+			EntitiesToDecompose.Append(Params.Query.Entities.GetData(), Params.Query.Entities.Num());
+		}
+
+		void ForEachAllocation(const FEntityAllocation* Allocation, TRead<FMovieSceneEntityID> EntityToDecomposeIDs, TRead<FMovieSceneBlendChannelID> BlendChannels, TRead<double> ValueResultComponent, TReadOptional<double> OptionalWeightComponent)
+		{
+			static const FMovieSceneBlenderSystemID BlenderSystemID = UMovieSceneBlenderSystem::GetBlenderSystemID<UMovieScenePiecewiseDoubleBlenderSystem>();
+
+			const bool bAdditive = Allocation->HasComponent(AdditiveBlendTag);
+
+			const int32 Num = Allocation->Num();
+			for (int32 EntityIndex = 0; EntityIndex < Num; ++EntityIndex)
+			{
+				const FMovieSceneBlendChannelID& BlendChannel(BlendChannels[EntityIndex]);
+				ensureMsgf(BlendChannel.SystemID == BlenderSystemID, TEXT("Overriding the standard blender system of standard types isn't supported."));
+
+				if (BlendChannel.ChannelID != DecomposeBlendChannel)
+				{
+					continue;
+				}
+
+				// We've found a contributor for this blend channel
+				const FMovieSceneEntityID EntityToDecompose = EntityToDecomposeIDs[EntityIndex];
+				const float               Weight            = OptionalWeightComponent ? OptionalWeightComponent[EntityIndex] : 1.f;
+				const double              ValueResult       = ValueResultComponent[EntityIndex];
+
+				if (EntitiesToDecompose.Contains(EntityToDecompose))
+				{
+					if (bAdditive)
+					{
+						Result->Value.DecomposedAdditives.Add(MakeTuple(EntityToDecompose, FWeightedValue{ ValueResult, Weight }));
+					}
+					else
+					{
+						Result->Value.DecomposedAbsolutes.Add(MakeTuple(EntityToDecompose, FWeightedValue{ ValueResult, Weight }));
+					}
+				}
+				else if (bAdditive)
+				{
+					Result->Value.Result.Additive += ValueResult * Weight;
+				}
+				else
+				{
+					Result->Value.Result.Absolute.Value  += ValueResult * Weight;
+					Result->Value.Result.Absolute.Weight += Weight;
+				}
+			}
+		}
+	};
+
+	FBuiltInComponentTypes* BuiltInComponents = FBuiltInComponentTypes::Get();
+
+	if (Params.Query.bConvertFromSourceEntityIDs)
+	{
+		return FEntityTaskBuilder()
+			.Read(BuiltInComponents->ParentEntity)
+			.Read(BuiltInComponents->BlendChannelInput)
+			.Read(ResultComponentType)
+			.ReadOptional(BuiltInComponents->WeightAndEasingResult)
+			.FilterAll({ Params.PropertyTag })
+			.template Dispatch_PerAllocation<FChannelResultTask>(&Linker->EntityManager, FSystemTaskPrerequisites(), nullptr, Params, Output);
+	}
+	else
+	{
+		return FEntityTaskBuilder()
+			.ReadEntityIDs()
+			.Read(BuiltInComponents->BlendChannelInput)
+			.Read(ResultComponentType)
+			.ReadOptional(BuiltInComponents->WeightAndEasingResult)
+			.FilterAll({ Params.PropertyTag })
+			.template Dispatch_PerAllocation<FChannelResultTask>(&Linker->EntityManager, FSystemTaskPrerequisites(), nullptr, Params, Output);
+	}
 }
 
diff --git a/Engine/Source/Runtime/MovieSceneTracks/Private/Systems/MovieScenePiecewiseFloatBlenderSystem.cpp b/Engine/Source/Runtime/MovieSceneTracks/Private/Systems/MovieScenePiecewiseFloatBlenderSystem.cpp
deleted file mode 100644
diff --git a/Engine/Source/Runtime/MovieSceneTracks/Private/Systems/MovieScenePiecewiseIntegerBlenderSystem.cpp b/Engine/Source/Runtime/MovieSceneTracks/Private/Systems/MovieScenePiecewiseIntegerBlenderSystem.cpp
--- a/Engine/Source/Runtime/MovieSceneTracks/Private/Systems/MovieScenePiecewiseIntegerBlenderSystem.cpp
+++ b/Engine/Source/Runtime/MovieSceneTracks/Private/Systems/MovieScenePiecewiseIntegerBlenderSystem.cpp
@@ -25,7 +25,7 @@ struct FIntegerAccumulationTask
 	{}
 
 	/** Task entry point - iterates the allocation's headers and accumulates int32 results for any required components */
-	void ForEachAllocation(const FEntityAllocation* Allocation, TRead<FMovieSceneBlendChannelID> BlendIDs, TRead<int32> Integers, TReadOptional<float> EasingAndWeights)
+	void ForEachAllocation(const FEntityAllocation* Allocation, TRead<FMovieSceneBlendChannelID> BlendIDs, TRead<int32> Integers, TReadOptional<double> EasingAndWeights)
 	{
 		static const FMovieSceneBlenderSystemID IntegerBlenderSystemID = UMovieSceneBlenderSystem::GetBlenderSystemID<UMovieScenePiecewiseIntegerBlenderSystem>();
 
@@ -41,7 +41,7 @@ struct FIntegerAccumulationTask
 
 				FIntegerBlendResult& Result = AccumulationBuffer[BlendID.ChannelID];
 
-				const float Weight = EasingAndWeights[Index];
+				const double Weight = EasingAndWeights[Index];
 				Result.Total  += (int32)(Integers[Index] * Weight);
 				Result.Weight += Weight;
 			}
@@ -75,7 +75,7 @@ struct FIntegerAdditiveFromBaseBlendTask
 		: AccumulationBuffer(*InAccumulationBuffer)
 	{}
 
-	void ForEachAllocation(const FEntityAllocation* Allocation, TRead<FMovieSceneBlendChannelID> BlendIDs, TRead<int32> Integers, TRead<int32> BaseValues, TReadOptional<float> EasingAndWeights)
+	void ForEachAllocation(const FEntityAllocation* Allocation, TRead<FMovieSceneBlendChannelID> BlendIDs, TRead<int32> Integers, TRead<int32> BaseValues, TReadOptional<double> EasingAndWeights)
 	{
 		static const FMovieSceneBlenderSystemID IntegerBlenderSystemID = UMovieSceneBlenderSystem::GetBlenderSystemID<UMovieScenePiecewiseIntegerBlenderSystem>();
 
@@ -90,7 +90,7 @@ struct FIntegerAdditiveFromBaseBlendTask
 
 				FIntegerBlendResult& Result = AccumulationBuffer[BlendID.ChannelID];
 
-				const float Weight = EasingAndWeights[Index];
+				const double Weight = EasingAndWeights[Index];
 				Result.Total  += (int32)((Integers[Index] - BaseValues[Index]) * Weight);
 				Result.Weight += Weight;
 			}
@@ -157,7 +157,7 @@ struct FIntegerCombineBlends
 
 		FIntegerBlendResult TotalAdditiveResult = { AdditiveResult.Total + AdditiveFromBaseResult.Total, AdditiveResult.Weight + AdditiveFromBaseResult.Weight };
 
-		const float TotalWeight = AbsoluteResult.Weight + RelativeResult.Weight;
+		const double TotalWeight = AbsoluteResult.Weight + RelativeResult.Weight;
 		if (TotalWeight != 0)
 		{
 			// If the absolute value has some partial weighting (for ease-in/out for instance), we ramp it from/to the initial value. This means
@@ -170,7 +170,7 @@ struct FIntegerCombineBlends
 			const int32 AbsoluteBlendedValue = bInitialValueContributes ?
 				((int32)(InitialValue * (1.f - AbsoluteResult.Weight)) + AbsoluteResult.Total) :
 				AbsoluteResult.Total;
-			const float FinalTotalWeight = bInitialValueContributes ? (TotalWeight + (1.f - AbsoluteResult.Weight)) : TotalWeight;
+			const double FinalTotalWeight = bInitialValueContributes ? (TotalWeight + (1.f - AbsoluteResult.Weight)) : TotalWeight;
 
 			const int32 Value = (int32)((AbsoluteBlendedValue + RelativeResult.Total) / FinalTotalWeight) + TotalAdditiveResult.Total;
 			OutFinalBlendResult = Value;
@@ -195,7 +195,7 @@ struct FIntegerCombineBlends
 		ensureMsgf(AbsoluteResult.Weight != 0.f, TEXT("Default blend combine being used for an entity that has no absolute weight. This should have an initial value and should be handled by each system, and excluded by default with UMovieSceneBlenderSystem::FinalCombineExclusionFilter ."));
 #endif
 
-		const float TotalWeight = AbsoluteResult.Weight;
+		const double TotalWeight = AbsoluteResult.Weight;
 		if (TotalWeight != 0)
 		{
 			const int32 Value = (int32)(AbsoluteResult.Total / AbsoluteResult.Weight) + AdditiveResult.Total + AdditiveFromBaseResult.Total;
diff --git a/Engine/Source/Runtime/MovieSceneTracks/Private/Systems/MovieScenePredictionSystem.cpp b/Engine/Source/Runtime/MovieSceneTracks/Private/Systems/MovieScenePredictionSystem.cpp
--- a/Engine/Source/Runtime/MovieSceneTracks/Private/Systems/MovieScenePredictionSystem.cpp
+++ b/Engine/Source/Runtime/MovieSceneTracks/Private/Systems/MovieScenePredictionSystem.cpp
@@ -2,7 +2,7 @@
 
 #include "Systems/MovieScenePredictionSystem.h"
 #include "Systems/MovieSceneEventSystems.h"
-#include "Systems/MovieScenePiecewiseFloatBlenderSystem.h"
+#include "Systems/MovieScenePiecewiseDoubleBlenderSystem.h"
 #include "Systems/FloatChannelEvaluatorSystem.h"
 
 #include "EntitySystem/BuiltInComponentTypes.h"
@@ -387,7 +387,7 @@ UMovieScenePredictionSystem::UMovieScenePredictionSystem(const FObjectInitialize
 
 	if (HasAnyFlags(RF_ClassDefaultObject))
 	{
-		DefineImplicitPrerequisite(UMovieScenePiecewiseFloatBlenderSystem::StaticClass(), GetClass());
+		DefineImplicitPrerequisite(UMovieScenePiecewiseDoubleBlenderSystem::StaticClass(), GetClass());
 		DefineImplicitPrerequisite(UFloatChannelEvaluatorSystem::StaticClass(), GetClass());
 		DefineImplicitPrerequisite(UMovieScenePostEvalEventSystem::StaticClass(), GetClass());
 		DefineImplicitPrerequisite(UMovieSceneInterrogatedPropertyInstantiatorSystem::StaticClass(), GetClass());
diff --git a/Engine/Source/Runtime/MovieSceneTracks/Private/Systems/MovieScenePropertyInstantiator.cpp b/Engine/Source/Runtime/MovieSceneTracks/Private/Systems/MovieScenePropertyInstantiator.cpp
--- a/Engine/Source/Runtime/MovieSceneTracks/Private/Systems/MovieScenePropertyInstantiator.cpp
+++ b/Engine/Source/Runtime/MovieSceneTracks/Private/Systems/MovieScenePropertyInstantiator.cpp
@@ -7,7 +7,7 @@
 #include "EntitySystem/MovieSceneEntitySystemLinker.h"
 #include "EntitySystem/MovieSceneBlenderSystem.h"
 #include "EntitySystem/MovieScenePropertyRegistry.h"
-#include "Systems/MovieScenePiecewiseFloatBlenderSystem.h"
+#include "Systems/MovieScenePiecewiseDoubleBlenderSystem.h"
 
 #include "Algo/AllOf.h"
 #include "Algo/IndexOf.h"
@@ -121,6 +121,7 @@ void UMovieScenePropertyInstantiatorSystem::OnRun(FSystemTaskPrerequisites& InPr
 
 	if (InvalidatedProperties.Num() != 0)
 	{
+		UpgradeFloatToDoubleProperties(InvalidatedProperties);
 		ProcessInvalidatedProperties(InvalidatedProperties);
 	}
 
@@ -212,6 +213,139 @@ void UMovieScenePropertyInstantiatorSystem::DiscoverInvalidatedProperties(TBitAr
 	.Iterate_PerEntity(&Linker->EntityManager, VisitExpiredEntities);
 }
 
+void UMovieScenePropertyInstantiatorSystem::UpgradeFloatToDoubleProperties(const TBitArray<>& InvalidatedProperties)
+{
+	using namespace UE::MovieScene;
+
+	MOVIESCENE_DETAILED_SCOPE_CYCLE_COUNTER(MovieSceneEval_ProcessInvalidatedProperties);
+
+	const FMovieSceneTracksComponentTypes* TrackComponents = FMovieSceneTracksComponentTypes::Get();
+	TArrayView<const FPropertyDefinition> Properties = BuiltInComponents->PropertyRegistry.GetProperties();
+
+	for (TConstSetBitIterator<> It(InvalidatedProperties); It; ++It)
+	{
+		const int32 PropertyIndex = It.GetIndex();
+		if (!ResolvedProperties.IsValidIndex(PropertyIndex))
+		{
+			continue;
+		}
+
+		FObjectPropertyInfo& PropertyInfo = ResolvedProperties[PropertyIndex];
+		const int32 OldPropertyDefinitionIndex = PropertyInfo.PropertyDefinitionIndex;
+
+		// The first time we encounter a specific property, we need to figure out if it needs type conversion or not.
+		if (!PropertyInfo.bRequiresTypeConversion.IsSet())
+		{
+			// Don't do anything if the property isn't of the kind of type we need to care about. Right now, we only
+			// support dealing with float->double and FVectorXf->FVectorXd.
+			const FPropertyDefinition& PropertyDefinition = Properties[PropertyInfo.PropertyDefinitionIndex];
+			if (PropertyDefinition.PropertyType != TrackComponents->Float.PropertyTag &&
+				PropertyDefinition.PropertyType != TrackComponents->FloatVector.PropertyTag)
+			{
+				PropertyInfo.bRequiresTypeConversion = false;
+				continue;
+			}
+
+			FProperty* BoundProperty = FTrackInstancePropertyBindings::FindProperty(PropertyInfo.BoundObject, PropertyInfo.PropertyBinding.PropertyPath.ToString());
+			if (!BoundProperty)
+			{
+				continue;
+			}
+
+			// Patch the resolved property info to point to the double-precision property definition.
+			if (PropertyDefinition.PropertyType == TrackComponents->Float.PropertyTag)
+			{
+				const bool bIsDouble = BoundProperty->IsA<FDoubleProperty>();
+				PropertyInfo.bRequiresTypeConversion = bIsDouble;
+				if (bIsDouble)
+				{
+					const int32 DoublePropertyDefinitionIndex = Properties.IndexOfByPredicate(
+						[TrackComponents](const FPropertyDefinition& Item) { return Item.PropertyType == TrackComponents->Double.PropertyTag; });
+					ensure(DoublePropertyDefinitionIndex != INDEX_NONE);
+					PropertyInfo.PropertyDefinitionIndex = DoublePropertyDefinitionIndex;
+				}
+			}
+			else if (PropertyDefinition.PropertyType == TrackComponents->FloatVector.PropertyTag)
+			{
+				const UScriptStruct* BoundStructProperty = CastField<FStructProperty>(BoundProperty)->Struct;
+				const bool bIsDouble = (
+					(BoundStructProperty == TBaseStructure<FVector2D>::Get()
+						) ||
+					(
+						BoundStructProperty == TBaseStructure<FVector>::Get() ||
+						BoundStructProperty == TVariantStructure<FVector3d>::Get()
+						) ||
+					(
+						BoundStructProperty == TBaseStructure<FVector4>::Get() ||
+						BoundStructProperty == TVariantStructure<FVector4d>::Get()
+						));
+				PropertyInfo.bRequiresTypeConversion = bIsDouble;
+				if (bIsDouble)
+				{
+					const int32 DoublePropertyDefinitionIndex = Properties.IndexOfByPredicate(
+						[TrackComponents](const FPropertyDefinition& Item) { return Item.PropertyType == TrackComponents->DoubleVector.PropertyTag; });
+					ensure(DoublePropertyDefinitionIndex != INDEX_NONE);
+					PropertyInfo.PropertyDefinitionIndex = DoublePropertyDefinitionIndex;
+				}
+			}
+			else
+			{
+				check(false);
+			}
+		}
+
+		if (!PropertyInfo.bRequiresTypeConversion.Get(false))
+		{
+			continue;
+		}
+
+		// Now we need to patch the contributors so that they have double-precision components.
+		// We only need to do it for the *new* contributors discovered this frame.
+		const FPropertyDefinition& PropertyDefinition = Properties[PropertyInfo.PropertyDefinitionIndex];
+		if (PropertyDefinition.PropertyType == TrackComponents->Double.PropertyTag)
+		{
+			// Swap out the property tag and result components, and add the result widening tag for easily
+			// finding these entities later.
+			int32 NumNewContributors = 0;
+			for (auto ContribIt = NewContributors.CreateKeyIterator(PropertyIndex); ContribIt; ++ContribIt)
+			{
+				const FMovieSceneEntityID CurID(ContribIt.Value());
+				FComponentMask EntityType = Linker->EntityManager.GetEntityType(CurID);
+				EntityType.Remove(TrackComponents->Float.PropertyTag);
+				EntityType.Set(TrackComponents->Double.PropertyTag);
+				Linker->EntityManager.ChangeEntityType(CurID, EntityType);
+
+				++NumNewContributors;
+			}
+
+			PropertyStats[OldPropertyDefinitionIndex].NumProperties -= NumNewContributors;
+			PropertyStats[PropertyInfo.PropertyDefinitionIndex].NumProperties += NumNewContributors;
+		}
+		else if (PropertyDefinition.PropertyType == TrackComponents->DoubleVector.PropertyTag)
+		{
+			// Swap out the property tag and result components for the composites this contributor has.
+			int32 NumNewContributors = 0;
+			for (auto ContribIt = NewContributors.CreateKeyIterator(PropertyIndex); ContribIt; ++ContribIt)
+			{
+				const FMovieSceneEntityID CurID(ContribIt.Value());
+				FComponentMask EntityType = Linker->EntityManager.GetEntityType(CurID);
+				EntityType.Remove(TrackComponents->FloatVector.PropertyTag);
+				EntityType.Set(TrackComponents->DoubleVector.PropertyTag);
+				Linker->EntityManager.ChangeEntityType(CurID, EntityType);
+
+				++NumNewContributors;
+			}
+
+			PropertyStats[OldPropertyDefinitionIndex].NumProperties -= NumNewContributors;
+			PropertyStats[PropertyInfo.PropertyDefinitionIndex].NumProperties += NumNewContributors;
+		}
+		else
+		{
+			check(false);
+		}
+	}
+}
+
 void UMovieScenePropertyInstantiatorSystem::ProcessInvalidatedProperties(const TBitArray<>& InvalidatedProperties)
 {
 	using namespace UE::MovieScene;
@@ -451,7 +585,7 @@ void UMovieScenePropertyInstantiatorSystem::InitializeBlendPath(const FPropertyP
 
 	if (!ensureMsgf(BlenderClass, TEXT("No default blender class specified on property, and no custom blender specified on entities. Falling back to float blender.")))
 	{
-		BlenderClass = UMovieScenePiecewiseFloatBlenderSystem::StaticClass();
+		BlenderClass = UMovieScenePiecewiseDoubleBlenderSystem::StaticClass();
 	}
 
 	UMovieSceneBlenderSystem* ExistingBlender = Params.PropertyInfo->Blender.Get();
diff --git a/Engine/Source/Runtime/MovieSceneTracks/Private/Systems/MovieSceneVectorPropertySystem.cpp b/Engine/Source/Runtime/MovieSceneTracks/Private/Systems/MovieSceneVectorPropertySystem.cpp
--- a/Engine/Source/Runtime/MovieSceneTracks/Private/Systems/MovieSceneVectorPropertySystem.cpp
+++ b/Engine/Source/Runtime/MovieSceneTracks/Private/Systems/MovieSceneVectorPropertySystem.cpp
@@ -8,7 +8,6 @@
 #include "Systems/DoubleChannelEvaluatorSystem.h"
 #include "Systems/FloatChannelEvaluatorSystem.h"
 #include "Systems/MovieScenePiecewiseDoubleBlenderSystem.h"
-#include "Systems/MovieScenePiecewiseFloatBlenderSystem.h"
 #include "Systems/MovieScenePropertyInstantiator.h"
 
 UMovieSceneFloatVectorPropertySystem::UMovieSceneFloatVectorPropertySystem(const FObjectInitializer& ObjInit)
@@ -22,7 +21,7 @@ UMovieSceneFloatVectorPropertySystem::UMovieSceneFloatVectorPropertySystem(const
 	{
 		// We need our floats correctly evaluated and blended, so we are downstream from those systems.
 		DefineImplicitPrerequisite(UFloatChannelEvaluatorSystem::StaticClass(), GetClass());
-		DefineImplicitPrerequisite(UMovieScenePiecewiseFloatBlenderSystem::StaticClass(), GetClass());
+		DefineImplicitPrerequisite(UMovieScenePiecewiseDoubleBlenderSystem::StaticClass(), GetClass());
 	}
 }
 
diff --git a/Engine/Source/Runtime/MovieSceneTracks/Private/Systems/WeightAndEasingEvaluatorSystem.cpp b/Engine/Source/Runtime/MovieSceneTracks/Private/Systems/WeightAndEasingEvaluatorSystem.cpp
--- a/Engine/Source/Runtime/MovieSceneTracks/Private/Systems/WeightAndEasingEvaluatorSystem.cpp
+++ b/Engine/Source/Runtime/MovieSceneTracks/Private/Systems/WeightAndEasingEvaluatorSystem.cpp
@@ -151,10 +151,10 @@ struct FEvaluateEasings
 			const FEntityAllocation* InAllocation, 
 			TRead<FFrameTime> Times,
 			TReadOptional<FEasingComponentData> OptEasing,
-			TReadOptional<float> OptManualWeights,
+			TReadOptional<double> OptManualWeights,
 			TReadOptional<FInstanceHandle> OptInstanceHandles,
 			TReadOptional<FMovieSceneSequenceID> OptHierarchicalEasingProviders,
-			TWrite<float> Results)
+			TWrite<double> Results)
 	{
 		const int32 Num = InAllocation->Num();
 
@@ -169,7 +169,7 @@ struct FEvaluateEasings
 		{
 			for (int32 Idx = 0; Idx < Num; ++Idx)
 			{
-				const float EasingWeight = OptEasing[Idx].Section->EvaluateEasing(Times[Idx]);
+				const double EasingWeight = OptEasing[Idx].Section->EvaluateEasing(Times[Idx]);
 				Results[Idx] *= FMath::Max(EasingWeight, 0.f);
 			}
 		}
@@ -249,7 +249,7 @@ struct FPropagateHierarchicalEasings
 		: EasingChannels(InEasingChannels)
 	{}
 
-	void ForEachAllocation(const FEntityAllocation* Allocation, TRead<uint16> HierarchicalEasingChannels, TWrite<float> WeightAndEasingResults)
+	void ForEachAllocation(const FEntityAllocation* Allocation, TRead<uint16> HierarchicalEasingChannels, TWrite<double> WeightAndEasingResults)
 	{
 		const int32 Num = Allocation->Num();
 		for (int32 Index = 0; Index < Num; ++Index)
@@ -322,7 +322,7 @@ void UWeightAndEasingEvaluatorSystem::ReleaseEasingChannel(uint16 EasingChannelI
 	}
 }
 
-void UWeightAndEasingEvaluatorSystem::SetSubSequenceEasing(UE::MovieScene::FInstanceHandle SubSequenceHandle, float EasingResult)
+void UWeightAndEasingEvaluatorSystem::SetSubSequenceEasing(UE::MovieScene::FInstanceHandle SubSequenceHandle, double EasingResult)
 {
 	using namespace UE::MovieScene;
 
diff --git a/Engine/Source/Runtime/MovieSceneTracks/Public/MovieSceneTracksComponentTypes.h b/Engine/Source/Runtime/MovieSceneTracks/Public/MovieSceneTracksComponentTypes.h
--- a/Engine/Source/Runtime/MovieSceneTracks/Public/MovieSceneTracksComponentTypes.h
+++ b/Engine/Source/Runtime/MovieSceneTracks/Public/MovieSceneTracksComponentTypes.h
@@ -126,92 +126,58 @@ struct FAttachmentComponent
 
 struct FFloatPropertyTraits
 {
-	using StorageType  = float;
-	using MetaDataType = TPropertyMetaData<bool>;
+	using StorageType  = double;
+	using CustomAccessorStorageType = float;
+	using MetaDataType = TPropertyMetaData<>;
+	using TraitsType = FFloatPropertyTraits;
 
 	using FloatTraitsImpl = TDirectPropertyTraits<float>;
-	using DoubleTraitsImpl = TDirectPropertyTraits<double>;
 
-	static void GetObjectPropertyValue(const UObject* InObject, bool bIsDouble, const FCustomPropertyAccessor& BaseCustomAccessor, float& OutValue)
+	static void GetObjectPropertyValue(const UObject* InObject, const FCustomPropertyAccessor& BaseCustomAccessor, double& OutValue)
 	{
-		checkf(!bIsDouble, TEXT("Type mismatch between float and double. Please check for any custom accessors defined on the wrong property type."));
 		const TCustomPropertyAccessor<FFloatPropertyTraits>& CustomAccessor = static_cast<const TCustomPropertyAccessor<FFloatPropertyTraits>&>(BaseCustomAccessor);
-		OutValue = (*CustomAccessor.Functions.Getter)(InObject, bIsDouble);
+		const float GetterValue = (*CustomAccessor.Functions.Getter)(InObject);
+		OutValue = (double)GetterValue;
 	}
-	static void GetObjectPropertyValue(const UObject* InObject, bool bIsDouble, uint16 PropertyOffset, float& OutValue)
+	static void GetObjectPropertyValue(const UObject* InObject, uint16 PropertyOffset, double& OutValue)
 	{
-		if (bIsDouble)
-		{
-			double TempValue;
-			DoubleTraitsImpl::GetObjectPropertyValue(InObject, PropertyOffset, TempValue);
-			OutValue = TempValue;
-		}
-		else
-		{
-			FloatTraitsImpl::GetObjectPropertyValue(InObject, PropertyOffset, OutValue);
-		}
+		float Value;
+		FloatTraitsImpl::GetObjectPropertyValue(InObject, PropertyOffset, Value);
+		OutValue = (double)Value;
 	}
-	static void GetObjectPropertyValue(const UObject* InObject, bool bIsDouble, FTrackInstancePropertyBindings* PropertyBindings, float& OutValue)
+	static void GetObjectPropertyValue(const UObject* InObject, FTrackInstancePropertyBindings* PropertyBindings, double& OutValue)
 	{
-		if (bIsDouble)
-		{
-			double TempValue;
-			DoubleTraitsImpl::GetObjectPropertyValue(InObject, PropertyBindings, TempValue);
-			OutValue = TempValue;
-		}
-		else
-		{
-			FloatTraitsImpl::GetObjectPropertyValue(InObject, PropertyBindings, OutValue);
-		}
+		float Value;
+		FloatTraitsImpl::GetObjectPropertyValue(InObject, PropertyBindings, Value);
+		OutValue = (double)Value;
 	}
-	static void GetObjectPropertyValue(const UObject* InObject, bool bIsDouble, const FName& PropertyPath, float& OutValue)
+	static void GetObjectPropertyValue(const UObject* InObject, const FName& PropertyPath, double& OutValue)
 	{
-		if (bIsDouble)
-		{
-			double TempValue;
-			DoubleTraitsImpl::GetObjectPropertyValue(InObject, PropertyPath, TempValue);
-			OutValue = TempValue;
-		}
-		else
-		{
-			FloatTraitsImpl::GetObjectPropertyValue(InObject, PropertyPath, OutValue);
-		}
+		float Value;
+		FloatTraitsImpl::GetObjectPropertyValue(InObject, PropertyPath, Value);
+		OutValue = (double)Value;
 	}
 
-	static void SetObjectPropertyValue(UObject* InObject, bool bIsDouble, const FCustomPropertyAccessor& BaseCustomAccessor, float InValue)
+	static void SetObjectPropertyValue(UObject* InObject, const FCustomPropertyAccessor& BaseCustomAccessor, double InValue)
 	{
-		checkf(!bIsDouble, TEXT("Type mismatch between float and double. Please check for any custom accessors defined on the wrong vector property type."));
 		const TCustomPropertyAccessor<FFloatPropertyTraits>& CustomAccessor = static_cast<const TCustomPropertyAccessor<FFloatPropertyTraits>&>(BaseCustomAccessor);
-		(*CustomAccessor.Functions.Setter)(InObject, bIsDouble, InValue);
+		const float SetterValue = (float)InValue;
+		(*CustomAccessor.Functions.Setter)(InObject, SetterValue);
 	}
-	static void SetObjectPropertyValue(UObject* InObject, bool bIsDouble, uint16 PropertyOffset, float InValue)
+	static void SetObjectPropertyValue(UObject* InObject, uint16 PropertyOffset, double InValue)
 	{
-		if (bIsDouble)
-		{
-			double TempValue = (float)InValue;
-			DoubleTraitsImpl::SetObjectPropertyValue(InObject, PropertyOffset, TempValue);
-		}
-		else
-		{
-			FloatTraitsImpl::SetObjectPropertyValue(InObject, PropertyOffset, InValue);
-		}
+		const float SetterValue = (float)InValue;
+		FloatTraitsImpl::SetObjectPropertyValue(InObject, PropertyOffset, SetterValue);
 	}
-	static void SetObjectPropertyValue(UObject* InObject, bool bIsDouble, FTrackInstancePropertyBindings* PropertyBindings, float InValue)
+	static void SetObjectPropertyValue(UObject* InObject, FTrackInstancePropertyBindings* PropertyBindings, double InValue)
 	{
-		if (bIsDouble)
-		{
-			double TempValue = (float)InValue;
-			DoubleTraitsImpl::SetObjectPropertyValue(InObject, PropertyBindings, TempValue);
-		}
-		else
-		{
-			FloatTraitsImpl::SetObjectPropertyValue(InObject, PropertyBindings, InValue);
-		}
+		const float SetterValue = (float)InValue;
+		FloatTraitsImpl::SetObjectPropertyValue(InObject, PropertyBindings, SetterValue);
 	}
 
-	static float CombineComposites(bool bIsDouble, float InValue)
+	static float CombineComposites(double InValue)
 	{
-		return InValue;
+		return (float)InValue;
 	}
 };
 
@@ -277,9 +243,9 @@ struct FColorPropertyTraits
 		}
 	}
 
-	static FIntermediateColor CombineComposites(EColorPropertyType InType, float InR, float InG, float InB, float InA)
+	static FIntermediateColor CombineComposites(EColorPropertyType InType, double InR, double InG, double InB, double InA)
 	{
-		return FIntermediateColor(InR, InG, InB, InA);
+		return FIntermediateColor((float)InR, (float)InG, (float)InB, (float)InA);
 	}
 };
 
@@ -362,20 +328,11 @@ struct FFloatVectorPropertyTraits
 
 	static void GetObjectPropertyValue(const UObject* InObject, FVectorPropertyMetaData MetaData, const FCustomPropertyAccessor& BaseCustomAccessor, FFloatIntermediateVector& OutValue)
 	{
-		checkf(!MetaData.bIsDouble, TEXT("Type mismatch between FVectorXf and FVectorXd. Please check for any custom accessors defined on the wrong vector property type."));
 		const TCustomPropertyAccessor<FFloatVectorPropertyTraits>& CustomAccessor = static_cast<const TCustomPropertyAccessor<FFloatVectorPropertyTraits>&>(BaseCustomAccessor);
 		OutValue = (*CustomAccessor.Functions.Getter)(InObject, MetaData);
 	}
 	static void GetObjectPropertyValue(const UObject* InObject, FVectorPropertyMetaData MetaData, uint16 PropertyOffset, FFloatIntermediateVector& OutValue)
 	{
-		if (MetaData.bIsDouble)
-		{
-			FDoubleIntermediateVector TempDoubleValue;
-			FDoubleVectorPropertyTraits::GetObjectPropertyValue(InObject, MetaData, PropertyOffset, TempDoubleValue);
-			OutValue = FFloatIntermediateVector(TempDoubleValue.X, TempDoubleValue.Y, TempDoubleValue.Z, TempDoubleValue.W);
-			return;
-		}
-
 		switch (MetaData.NumChannels)
 		{
 		case 2: TIndirectPropertyTraits<FVector2f, FFloatIntermediateVector>::GetObjectPropertyValue(InObject, PropertyOffset, OutValue); return;
@@ -385,14 +342,6 @@ struct FFloatVectorPropertyTraits
 	}
 	static void GetObjectPropertyValue(const UObject* InObject, FVectorPropertyMetaData MetaData, FTrackInstancePropertyBindings* PropertyBindings, FFloatIntermediateVector& OutValue)
 	{
-		if (MetaData.bIsDouble)
-		{
-			FDoubleIntermediateVector TempDoubleValue;
-			FDoubleVectorPropertyTraits::GetObjectPropertyValue(InObject, MetaData, PropertyBindings, TempDoubleValue);
-			OutValue = FFloatIntermediateVector(TempDoubleValue.X, TempDoubleValue.Y, TempDoubleValue.Z, TempDoubleValue.W);
-			return;
-		}
-
 		switch (MetaData.NumChannels)
 		{
 		case 2: TIndirectPropertyTraits<FVector2f, FFloatIntermediateVector>::GetObjectPropertyValue(InObject, PropertyBindings, OutValue); return;
@@ -402,14 +351,6 @@ struct FFloatVectorPropertyTraits
 	}
 	static void GetObjectPropertyValue(const UObject* InObject, FVectorPropertyMetaData MetaData, const FName& PropertyPath, StorageType& OutValue)
 	{
-		if (MetaData.bIsDouble)
-		{
-			FDoubleIntermediateVector TempDoubleValue;
-			FDoubleVectorPropertyTraits::GetObjectPropertyValue(InObject, MetaData, PropertyPath, TempDoubleValue);
-			OutValue = FFloatIntermediateVector(TempDoubleValue.X, TempDoubleValue.Y, TempDoubleValue.Z, TempDoubleValue.W);
-			return;
-		}
-
 		switch (MetaData.NumChannels)
 		{
 		case 2: TIndirectPropertyTraits<FVector2f, FFloatIntermediateVector>::GetObjectPropertyValue(InObject, PropertyPath, OutValue); return;
@@ -420,19 +361,11 @@ struct FFloatVectorPropertyTraits
 
 	static void SetObjectPropertyValue(UObject* InObject, FVectorPropertyMetaData MetaData, const FCustomPropertyAccessor& BaseCustomAccessor, const FFloatIntermediateVector& InValue)
 	{
-		checkf(!MetaData.bIsDouble, TEXT("Type mismatch between FVectorXf and FVectorXd. Please check for any custom accessors defined on the wrong vector property type."));
 		const TCustomPropertyAccessor<FFloatVectorPropertyTraits>& CustomAccessor = static_cast<const TCustomPropertyAccessor<FFloatVectorPropertyTraits>&>(BaseCustomAccessor);
 		(*CustomAccessor.Functions.Setter)(InObject, MetaData, InValue);
 	}
 	static void SetObjectPropertyValue(UObject* InObject, FVectorPropertyMetaData MetaData, uint16 PropertyOffset, const FFloatIntermediateVector& InValue)
 	{
-		if (MetaData.bIsDouble)
-		{
-			FDoubleIntermediateVector TempDoubleValue(InValue.X, InValue.Y, InValue.Z, InValue.W);
-			FDoubleVectorPropertyTraits::SetObjectPropertyValue(InObject, MetaData, PropertyOffset, TempDoubleValue);
-			return;
-		}
-
 		switch (MetaData.NumChannels)
 		{
 		case 2: TIndirectPropertyTraits<FVector2f, FFloatIntermediateVector>::SetObjectPropertyValue(InObject, PropertyOffset, InValue); return;
@@ -444,13 +377,6 @@ struct FFloatVectorPropertyTraits
 	}
 	static void SetObjectPropertyValue(UObject* InObject, FVectorPropertyMetaData MetaData, FTrackInstancePropertyBindings* PropertyBindings, const FFloatIntermediateVector& InValue)
 	{
-		if (MetaData.bIsDouble)
-		{
-			FDoubleIntermediateVector TempDoubleValue(InValue.X, InValue.Y, InValue.Z, InValue.W);
-			FDoubleVectorPropertyTraits::SetObjectPropertyValue(InObject, MetaData, PropertyBindings, TempDoubleValue);
-			return;
-		}
-
 		switch (MetaData.NumChannels)
 		{
 		case 2: TIndirectPropertyTraits<FVector2f, FFloatIntermediateVector>::SetObjectPropertyValue(InObject, PropertyBindings, InValue); return;
diff --git a/Engine/Source/Runtime/MovieSceneTracks/Public/MovieSceneTracksPropertyTypes.h b/Engine/Source/Runtime/MovieSceneTracks/Public/MovieSceneTracksPropertyTypes.h
--- a/Engine/Source/Runtime/MovieSceneTracks/Public/MovieSceneTracksPropertyTypes.h
+++ b/Engine/Source/Runtime/MovieSceneTracks/Public/MovieSceneTracksPropertyTypes.h
@@ -23,27 +23,26 @@ namespace MovieScene
 struct FVectorPropertyMetaData
 {
 	uint8 NumChannels = 0;
-	bool bIsDouble = false;
 };
 
 /** Intermediate type for the vector property system that lets us store how many dimensions the vector should have */
 struct FFloatIntermediateVector
 {
-	float X, Y, Z, W;
+	double X, Y, Z, W;
 
 	FFloatIntermediateVector()
 		: X(0), Y(0), Z(0), W(0)
 	{}
 
-	FFloatIntermediateVector(float InX, float InY)
+	FFloatIntermediateVector(double InX, double InY)
 		: X(InX), Y(InY), Z(0), W(0)
 	{}
 
-	FFloatIntermediateVector(float InX, float InY, float InZ)
+	FFloatIntermediateVector(double InX, double InY, double InZ)
 		: X(InX), Y(InY), Z(InZ), W(0)
 	{}
 
-	FFloatIntermediateVector(float InX, float InY, float InZ, float InW)
+	FFloatIntermediateVector(double InX, double InY, double InZ, double InW)
 		: X(InX), Y(InY), Z(InZ), W(InW)
 	{}
 };
@@ -90,13 +89,13 @@ template<> struct TComponentDebugType<EColorPropertyType> { static const ECompon
 /** Intermediate type for the color property system that lets us store what kind of color type we should use */
 struct FIntermediateColor
 {
-	float R, G, B, A;
+	double R, G, B, A;
 
 	FIntermediateColor()
 		: R(0.f), G(0.f), B(0.f), A(0.f)
 	{}
 
-	explicit FIntermediateColor(float InR, float InG, float InB, float InA)
+	explicit FIntermediateColor(double InR, double InG, double InB, double InA)
 		: R(InR), G(InG), B(InB), A(InA)
 	{}
 	explicit FIntermediateColor(const FLinearColor& InColor)
@@ -120,7 +119,7 @@ struct FIntermediateColor
 		A = SpecifiedColor.A;
 	}
 
-	float operator[](int32 Index) const
+	double operator[](int32 Index) const
 	{
 		check(Index >= 0 && Index <= 3);
 		return (&R)[Index];
diff --git a/Engine/Source/Runtime/MovieSceneTracks/Public/Systems/FloatChannelEvaluatorSystem.h b/Engine/Source/Runtime/MovieSceneTracks/Public/Systems/FloatChannelEvaluatorSystem.h
--- a/Engine/Source/Runtime/MovieSceneTracks/Public/Systems/FloatChannelEvaluatorSystem.h
+++ b/Engine/Source/Runtime/MovieSceneTracks/Public/Systems/FloatChannelEvaluatorSystem.h
@@ -28,7 +28,6 @@ class MOVIESCENETRACKS_API UFloatChannelEvaluatorSystem : public UMovieSceneEnti
 {
 public:
 
-
 	GENERATED_BODY()
 
 	UFloatChannelEvaluatorSystem(const FObjectInitializer& ObjInit);
@@ -36,7 +35,7 @@ public:
 	virtual void OnRun(FSystemTaskPrerequisites& InPrerequisites, FSystemSubsequentTasks& Subsequents) override;
 	virtual bool IsRelevantImpl(UMovieSceneEntitySystemLinker* InLinker) const override;
 
-	static void RegisterChannelType(TComponentTypeID<UE::MovieScene::FSourceFloatChannel> SourceChannelType, TComponentTypeID<UE::MovieScene::FSourceFloatChannelFlags> ChannelFlagsType, TComponentTypeID<float> ResultType);
+	static void RegisterChannelType(TComponentTypeID<UE::MovieScene::FSourceFloatChannel> SourceChannelType, TComponentTypeID<UE::MovieScene::FSourceFloatChannelFlags> ChannelFlagsType, TComponentTypeID<double> ResultType);
 
 private:
 
@@ -44,7 +43,7 @@ private:
 	{
 		TComponentTypeID<UE::MovieScene::FSourceFloatChannel> ChannelType;
 		TComponentTypeID<UE::MovieScene::FSourceFloatChannelFlags> ChannelFlagsType;
-		TComponentTypeID<float> ResultType;
+		TComponentTypeID<double> ResultType;
 	};
 
 	static TArray<FChannelType, TInlineAllocator<16>> StaticChannelTypes;
diff --git a/Engine/Source/Runtime/MovieSceneTracks/Public/Systems/MovieSceneComponentMaterialSystem.h b/Engine/Source/Runtime/MovieSceneTracks/Public/Systems/MovieSceneComponentMaterialSystem.h
--- a/Engine/Source/Runtime/MovieSceneTracks/Public/Systems/MovieSceneComponentMaterialSystem.h
+++ b/Engine/Source/Runtime/MovieSceneTracks/Public/Systems/MovieSceneComponentMaterialSystem.h
@@ -10,7 +10,7 @@
 
 #include "MovieSceneComponentMaterialSystem.generated.h"
 
-class UMovieScenePiecewiseFloatBlenderSystem;
+class UMovieScenePiecewiseDoubleBlenderSystem;
 
 namespace UE::MovieScene
 {
diff --git a/Engine/Source/Runtime/MovieSceneTracks/Public/Systems/MovieSceneMaterialParameterSystem.h b/Engine/Source/Runtime/MovieSceneTracks/Public/Systems/MovieSceneMaterialParameterSystem.h
--- a/Engine/Source/Runtime/MovieSceneTracks/Public/Systems/MovieSceneMaterialParameterSystem.h
+++ b/Engine/Source/Runtime/MovieSceneTracks/Public/Systems/MovieSceneMaterialParameterSystem.h
@@ -8,7 +8,7 @@
 
 #include "MovieSceneMaterialParameterSystem.generated.h"
 
-class UMovieScenePiecewiseFloatBlenderSystem;
+class UMovieScenePiecewiseDoubleBlenderSystem;
 
 namespace UE::MovieScene
 {
@@ -62,5 +62,5 @@ private:
 public:
 
 	UPROPERTY()
-	TObjectPtr<UMovieScenePiecewiseFloatBlenderSystem> FloatBlenderSystem;
+	TObjectPtr<UMovieScenePiecewiseDoubleBlenderSystem> DoubleBlenderSystem;
 };
diff --git a/Engine/Source/Runtime/MovieSceneTracks/Public/Systems/MovieScenePiecewiseBlenderSystemImpl.h b/Engine/Source/Runtime/MovieSceneTracks/Public/Systems/MovieScenePiecewiseBlenderSystemImpl.h
deleted file mode 100644
diff --git a/Engine/Source/Runtime/MovieSceneTracks/Public/Systems/MovieScenePiecewiseDoubleBlenderSystem.h b/Engine/Source/Runtime/MovieSceneTracks/Public/Systems/MovieScenePiecewiseDoubleBlenderSystem.h
--- a/Engine/Source/Runtime/MovieSceneTracks/Public/Systems/MovieScenePiecewiseDoubleBlenderSystem.h
+++ b/Engine/Source/Runtime/MovieSceneTracks/Public/Systems/MovieScenePiecewiseDoubleBlenderSystem.h
@@ -2,12 +2,87 @@
 
 #pragma once
 
+#include "Containers/SortedMap.h"
 #include "EntitySystem/MovieSceneBlenderSystem.h"
+#include "EntitySystem/MovieSceneCachedEntityFilterResult.h"
 #include "EntitySystem/MovieSceneDecompositionQuery.h"
-#include "Systems/MovieScenePiecewiseBlenderSystemImpl.h"
+#include "EntitySystem/MovieSceneEntitySystem.h"
 
 #include "MovieScenePiecewiseDoubleBlenderSystem.generated.h"
 
+namespace UE
+{
+namespace MovieScene
+{
+
+/** Blend result struct that stores the cumulative sum of pre-weighted values, alongside the total weight */
+struct FBlendResult
+{
+	/** Cumulative sum of blend values pre-multiplied with each value's weight. */
+	double Total = 0.f;
+	/** Cumulative sum of weights. */
+	float Weight = 0.f;
+};
+
+/** Structure for holding the blend results of each blend type */
+struct FAccumulationResult
+{
+	const FBlendResult* Absolutes = nullptr;
+	const FBlendResult* Relatives = nullptr;
+	const FBlendResult* Additives = nullptr;
+	const FBlendResult* AdditivesFromBase = nullptr;
+
+	bool IsValid() const
+	{
+		return Absolutes || Relatives || Additives || AdditivesFromBase;
+	}
+
+	FBlendResult GetAbsoluteResult(uint16 BlendID) const
+	{
+		return Absolutes ? Absolutes[BlendID] : FBlendResult{};
+	}
+	FBlendResult GetRelativeResult(uint16 BlendID) const
+	{
+		return Relatives ? Relatives[BlendID] : FBlendResult{};
+	}
+	FBlendResult GetAdditiveResult(uint16 BlendID) const
+	{
+		return Additives ? Additives[BlendID] : FBlendResult{};
+	}
+	FBlendResult GetAdditiveFromBaseResult(uint16 BlendID) const
+	{
+		return AdditivesFromBase ? AdditivesFromBase[BlendID] : FBlendResult{};
+	}
+};
+
+/** Buffer used for accumulating additive-from-base values */
+struct FAdditiveFromBaseBuffer
+{
+	TArray<FBlendResult> Buffer;
+	TComponentTypeID<double> BaseComponent;
+};
+
+/** Struct that maintains accumulation buffers for each blend type, one buffer per float result component type */
+struct FAccumulationBuffers
+{
+	bool IsEmpty() const;
+
+	void Reset();
+
+	FAccumulationResult FindResults(FComponentTypeID InComponentType) const;
+
+	/** Map from value result component type -> Absolute blend accumulation buffer for that channel type */
+	TSortedMap<FComponentTypeID, TArray<FBlendResult>> Absolute;
+	/** Map from value result component type -> Relative blend accumulation buffer for that channel type */
+	TSortedMap<FComponentTypeID, TArray<FBlendResult>> Relative;
+	/** Map from value result component type -> Additive blend accumulation buffer for that channel type */
+	TSortedMap<FComponentTypeID, TArray<FBlendResult>> Additive;
+	/** Map from value result component type -> Additive From Base blend accumulation buffer for that channel type */
+	TSortedMap<FComponentTypeID, FAdditiveFromBaseBuffer> AdditiveFromBase;
+};
+
+} // namespace MovieScene
+} // namespace UE
 
 UCLASS()
 class MOVIESCENETRACKS_API UMovieScenePiecewiseDoubleBlenderSystem : public UMovieSceneBlenderSystem, public IMovieSceneValueDecomposer
@@ -26,6 +101,25 @@ public:
 
 private:
 
-	UE::MovieScene::TPiecewiseBlenderSystemImpl<double> Impl;
+	void ReinitializeAccumulationBuffers();
+	void ZeroAccumulationBuffers();
+
+	/** Buffers that contain accumulated blend values, separated by blend type */
+	UE::MovieScene::FAccumulationBuffers AccumulationBuffers;
+
+	/** Mask that contains value result components that have BlendChannelInput components */
+	FComponentMask BlendedResultMask;
+
+	/** Mask that contains property tags for any property type that has has at least one BlendChannelOutput */
+	FComponentMask BlendedPropertyMask;
+
+	/** Cache state that is used to invalidate and reset the accumulation buffers when the entity manager has structurally changed */
+	UE::MovieScene::FCachedEntityManagerState ChannelRelevancyCache;
+
+	/** Bit array specifying FCompositePropertyTypeID's for properties contained within BlendedPropertyMask */
+	TBitArray<> CachedRelevantProperties;
+
+	/** Whether the current entity manager contains any non-property based blends */
+	bool bContainsNonPropertyBlends = false;
 };
 
diff --git a/Engine/Source/Runtime/MovieSceneTracks/Public/Systems/MovieScenePiecewiseFloatBlenderSystem.h b/Engine/Source/Runtime/MovieSceneTracks/Public/Systems/MovieScenePiecewiseFloatBlenderSystem.h
deleted file mode 100644
diff --git a/Engine/Source/Runtime/MovieSceneTracks/Public/Systems/MovieScenePiecewiseIntegerBlenderSystem.h b/Engine/Source/Runtime/MovieSceneTracks/Public/Systems/MovieScenePiecewiseIntegerBlenderSystem.h
--- a/Engine/Source/Runtime/MovieSceneTracks/Public/Systems/MovieScenePiecewiseIntegerBlenderSystem.h
+++ b/Engine/Source/Runtime/MovieSceneTracks/Public/Systems/MovieScenePiecewiseIntegerBlenderSystem.h
@@ -19,7 +19,7 @@ namespace MovieScene
 		/** Cumulative sum of blend values pre-multiplied with each value's weight. */
 		int32 Total = 0;
 		/** Cumulative sum of weights. */
-		float Weight = 0.f;
+		double Weight = 0.f;
 	};
 
 	/** Struct that maintains accumulation buffers for each blend type */
diff --git a/Engine/Source/Runtime/MovieSceneTracks/Public/Systems/MovieScenePropertyInstantiator.h b/Engine/Source/Runtime/MovieSceneTracks/Public/Systems/MovieScenePropertyInstantiator.h
--- a/Engine/Source/Runtime/MovieSceneTracks/Public/Systems/MovieScenePropertyInstantiator.h
+++ b/Engine/Source/Runtime/MovieSceneTracks/Public/Systems/MovieScenePropertyInstantiator.h
@@ -131,7 +131,9 @@ private:
 		uint16 BlendChannel;
 		/** The index of this property within FPropertyRegistry::GetProperties. */
 		int32 PropertyDefinitionIndex;
-		/** true if any of the contributors to this property need restore state.. */
+		/** true if this property has already been inspected for float-to-double conversion */
+		TOptional<bool> bRequiresTypeConversion;
+		/** true if any of the contributors to this property need restore state. */
 		bool bWantsRestoreState;
 	};
 
@@ -148,6 +150,7 @@ private:
 		int32 PropertyInfoIndex;
 	};
 	void DiscoverInvalidatedProperties(TBitArray<>& OutInvalidatedProperties);
+	void UpgradeFloatToDoubleProperties(const TBitArray<>& InvalidatedProperties);
 	void ProcessInvalidatedProperties(const TBitArray<>& InvalidatedProperties);
 	void UpdatePropertyInfo(const FPropertyParameters& Params);
 	bool PropertySupportsFastPath(const FPropertyParameters& Params) const;
diff --git a/Engine/Source/Runtime/MovieSceneTracks/Public/Systems/WeightAndEasingEvaluatorSystem.h b/Engine/Source/Runtime/MovieSceneTracks/Public/Systems/WeightAndEasingEvaluatorSystem.h
--- a/Engine/Source/Runtime/MovieSceneTracks/Public/Systems/WeightAndEasingEvaluatorSystem.h
+++ b/Engine/Source/Runtime/MovieSceneTracks/Public/Systems/WeightAndEasingEvaluatorSystem.h
@@ -24,13 +24,13 @@ namespace MovieScene
 	struct FHierarchicalEasingChannelContributorData
 	{
 		UE::MovieScene::FInstanceHandle SubSequenceHandle;
-		float EasingResult;
+		double EasingResult;
 	};
 
 	struct FHierarchicalEasingChannelData
 	{
 		TArray<FHierarchicalEasingChannelContributorData, TInlineAllocator<8>> Contributors;
-		float FinalEasingResult;
+		double FinalEasingResult;
 	};
 
 }  // namespace MovieScene
@@ -73,7 +73,7 @@ public:
 	uint16 AllocateEasingChannel(UE::MovieScene::FInstanceHandle SubSequenceHandle);
 	void ReleaseEasingChannel(uint16 EasingChannelID);
 
-	void SetSubSequenceEasing(UE::MovieScene::FInstanceHandle SubSequenceHandle, float EasingResult);
+	void SetSubSequenceEasing(UE::MovieScene::FInstanceHandle SubSequenceHandle, double EasingResult);
 
 private:
 
diff --git a/Engine/Source/Runtime/UMG/Private/Animation/MovieScene2DTransformPropertySystem.cpp b/Engine/Source/Runtime/UMG/Private/Animation/MovieScene2DTransformPropertySystem.cpp
--- a/Engine/Source/Runtime/UMG/Private/Animation/MovieScene2DTransformPropertySystem.cpp
+++ b/Engine/Source/Runtime/UMG/Private/Animation/MovieScene2DTransformPropertySystem.cpp
@@ -4,7 +4,7 @@
 
 #include "Systems/MovieScenePropertyInstantiator.h"
 #include "Systems/FloatChannelEvaluatorSystem.h"
-#include "Systems/MovieScenePiecewiseFloatBlenderSystem.h"
+#include "Systems/MovieScenePiecewiseDoubleBlenderSystem.h"
 
 #include "Animation/MovieSceneUMGComponentTypes.h"
 
@@ -18,7 +18,7 @@ UMovieScene2DTransformPropertySystem::UMovieScene2DTransformPropertySystem(const
 
 	if (HasAnyFlags(RF_ClassDefaultObject))
 	{
-		DefineImplicitPrerequisite(UMovieScenePiecewiseFloatBlenderSystem::StaticClass(), GetClass());
+		DefineImplicitPrerequisite(UMovieScenePiecewiseDoubleBlenderSystem::StaticClass(), GetClass());
 		DefineImplicitPrerequisite(UFloatChannelEvaluatorSystem::StaticClass(), GetClass());
 	}
 }
diff --git a/Engine/Source/Runtime/UMG/Private/Animation/MovieSceneMarginPropertySystem.cpp b/Engine/Source/Runtime/UMG/Private/Animation/MovieSceneMarginPropertySystem.cpp
--- a/Engine/Source/Runtime/UMG/Private/Animation/MovieSceneMarginPropertySystem.cpp
+++ b/Engine/Source/Runtime/UMG/Private/Animation/MovieSceneMarginPropertySystem.cpp
@@ -3,7 +3,7 @@
 #include "Animation/MovieSceneMarginPropertySystem.h"
 #include "Animation/MovieSceneUMGComponentTypes.h"
 #include "Systems/FloatChannelEvaluatorSystem.h"
-#include "Systems/MovieScenePiecewiseFloatBlenderSystem.h"
+#include "Systems/MovieScenePiecewiseDoubleBlenderSystem.h"
 
 
 UMovieSceneMarginPropertySystem::UMovieSceneMarginPropertySystem(const FObjectInitializer& ObjInit)
@@ -15,7 +15,7 @@ UMovieSceneMarginPropertySystem::UMovieSceneMarginPropertySystem(const FObjectIn
 
 	if (HasAnyFlags(RF_ClassDefaultObject))
 {
-		DefineImplicitPrerequisite(UMovieScenePiecewiseFloatBlenderSystem::StaticClass(), GetClass());
+		DefineImplicitPrerequisite(UMovieScenePiecewiseDoubleBlenderSystem::StaticClass(), GetClass());
 		DefineImplicitPrerequisite(UFloatChannelEvaluatorSystem::StaticClass(), GetClass());
 	}
 	}
diff --git a/Engine/Source/Runtime/UMG/Private/Animation/MovieSceneUMGComponentTypes.cpp b/Engine/Source/Runtime/UMG/Private/Animation/MovieSceneUMGComponentTypes.cpp
--- a/Engine/Source/Runtime/UMG/Private/Animation/MovieSceneUMGComponentTypes.cpp
+++ b/Engine/Source/Runtime/UMG/Private/Animation/MovieSceneUMGComponentTypes.cpp
@@ -6,7 +6,7 @@
 #include "EntitySystem/MovieSceneEntityFactoryTemplates.h"
 #include "EntitySystem/MovieScenePropertyComponentHandler.h"
 #include "MovieSceneTracksComponentTypes.h"
-#include "Systems/MovieScenePiecewiseFloatBlenderSystem.h"
+#include "Systems/MovieScenePiecewiseDoubleBlenderSystem.h"
 
 #include "Components/Widget.h"
 
@@ -40,15 +40,29 @@ void ConvertOperationalProperty(const FWidgetTransform& In, FIntermediateWidgetT
 	Out.ShearY = In.Shear.Y;
 }
 
-static float GetRenderOpacity(const UObject* Object, bool bIsDouble)
+void ConvertOperationalProperty(const FIntermediateMargin& In, FMargin& Out)
+{
+	Out.Top = In.Top;
+	Out.Right = In.Right;
+	Out.Bottom = In.Bottom;
+	Out.Left = In.Left;
+}
+
+void ConvertOperationalProperty(const FMargin& In, FIntermediateMargin& Out)
+{
+	Out.Top = In.Top;
+	Out.Right = In.Right;
+	Out.Bottom = In.Bottom;
+	Out.Left = In.Left;
+}
+
+static float GetRenderOpacity(const UObject* Object)
 {
-	check(!bIsDouble);
 	return CastChecked<const UWidget>(Object)->GetRenderOpacity();
 }
 
-static void SetRenderOpacity(UObject* Object, bool bIsDouble, float InRenderOpacity)
+static void SetRenderOpacity(UObject* Object, float InRenderOpacity)
 {
-	check(!bIsDouble);
 	CastChecked<UWidget>(Object)->SetRenderOpacity(InRenderOpacity);
 }
 
@@ -87,22 +101,22 @@ FMovieSceneUMGComponentTypes::FMovieSceneUMGComponentTypes()
 	CustomWidgetTransformAccessors.Add(UWidget::StaticClass(), "RenderTransform", &GetRenderTransform, &SetRenderTransform);
 
 	BuiltInComponents->PropertyRegistry.DefineCompositeProperty(Margin, TEXT("Apply FMargin Properties"))
-	.AddComposite(BuiltInComponents->FloatResult[0], &FMargin::Left)
-	.AddComposite(BuiltInComponents->FloatResult[1], &FMargin::Top)
-	.AddComposite(BuiltInComponents->FloatResult[2], &FMargin::Right)
-	.AddComposite(BuiltInComponents->FloatResult[3], &FMargin::Bottom)
-	.SetBlenderSystem<UMovieScenePiecewiseFloatBlenderSystem>()
+	.AddComposite(BuiltInComponents->DoubleResult[0], &FIntermediateMargin::Left)
+	.AddComposite(BuiltInComponents->DoubleResult[1], &FIntermediateMargin::Top)
+	.AddComposite(BuiltInComponents->DoubleResult[2], &FIntermediateMargin::Right)
+	.AddComposite(BuiltInComponents->DoubleResult[3], &FIntermediateMargin::Bottom)
+	.SetBlenderSystem<UMovieScenePiecewiseDoubleBlenderSystem>()
 	.Commit();
 
 	BuiltInComponents->PropertyRegistry.DefineCompositeProperty(WidgetTransform, TEXT("Call UUserWidget::SetRenderTransform"))
-	.AddComposite(BuiltInComponents->FloatResult[0], &FIntermediateWidgetTransform::TranslationX)
-	.AddComposite(BuiltInComponents->FloatResult[1], &FIntermediateWidgetTransform::TranslationY)
-	.AddComposite(BuiltInComponents->FloatResult[2], &FIntermediateWidgetTransform::Rotation)
-	.AddComposite(BuiltInComponents->FloatResult[3], &FIntermediateWidgetTransform::ScaleX)
-	.AddComposite(BuiltInComponents->FloatResult[4], &FIntermediateWidgetTransform::ScaleY)
-	.AddComposite(BuiltInComponents->FloatResult[5], &FIntermediateWidgetTransform::ShearX)
-	.AddComposite(BuiltInComponents->FloatResult[6], &FIntermediateWidgetTransform::ShearY)
-	.SetBlenderSystem<UMovieScenePiecewiseFloatBlenderSystem>()
+	.AddComposite(BuiltInComponents->DoubleResult[0], &FIntermediateWidgetTransform::TranslationX)
+	.AddComposite(BuiltInComponents->DoubleResult[1], &FIntermediateWidgetTransform::TranslationY)
+	.AddComposite(BuiltInComponents->DoubleResult[2], &FIntermediateWidgetTransform::Rotation)
+	.AddComposite(BuiltInComponents->DoubleResult[3], &FIntermediateWidgetTransform::ScaleX)
+	.AddComposite(BuiltInComponents->DoubleResult[4], &FIntermediateWidgetTransform::ScaleY)
+	.AddComposite(BuiltInComponents->DoubleResult[5], &FIntermediateWidgetTransform::ShearX)
+	.AddComposite(BuiltInComponents->DoubleResult[6], &FIntermediateWidgetTransform::ShearY)
+	.SetBlenderSystem<UMovieScenePiecewiseDoubleBlenderSystem>()
 	.SetCustomAccessors(&CustomWidgetTransformAccessors)
 	.Commit();
 }
diff --git a/Engine/Source/Runtime/UMG/Private/Animation/MovieSceneWidgetMaterialSystem.cpp b/Engine/Source/Runtime/UMG/Private/Animation/MovieSceneWidgetMaterialSystem.cpp
--- a/Engine/Source/Runtime/UMG/Private/Animation/MovieSceneWidgetMaterialSystem.cpp
+++ b/Engine/Source/Runtime/UMG/Private/Animation/MovieSceneWidgetMaterialSystem.cpp
@@ -5,10 +5,11 @@
 #include "Animation/WidgetMaterialTrackUtilities.h"
 
 #include "EntitySystem/MovieSceneEntityMutations.h"
+
 #include "Evaluation/PreAnimatedState/MovieScenePreAnimatedStorageID.inl"
 
 #include "Systems/FloatChannelEvaluatorSystem.h"
-#include "Systems/MovieScenePiecewiseFloatBlenderSystem.h"
+#include "Systems/MovieScenePiecewiseDoubleBlenderSystem.h"
 
 #include "Materials/MaterialInstanceDynamic.h"
 
diff --git a/Engine/Source/Runtime/UMG/Public/Animation/MovieSceneUMGComponentTypes.h b/Engine/Source/Runtime/UMG/Public/Animation/MovieSceneUMGComponentTypes.h
--- a/Engine/Source/Runtime/UMG/Public/Animation/MovieSceneUMGComponentTypes.h
+++ b/Engine/Source/Runtime/UMG/Public/Animation/MovieSceneUMGComponentTypes.h
@@ -29,18 +29,28 @@ struct FWidgetMaterialPath
 
 struct FIntermediateWidgetTransform
 {
-	float TranslationX;
-	float TranslationY;
-	float Rotation;
-	float ScaleX;
-	float ScaleY;
-	float ShearX;
-	float ShearY;
+	double TranslationX;
+	double TranslationY;
+	double Rotation;
+	double ScaleX;
+	double ScaleY;
+	double ShearX;
+	double ShearY;
 };
 UMG_API void ConvertOperationalProperty(const FIntermediateWidgetTransform& In, FWidgetTransform& Out);
 UMG_API void ConvertOperationalProperty(const FWidgetTransform& In, FIntermediateWidgetTransform& Out);
 
-using FMarginTraits = TDirectPropertyTraits<FMargin>;
+struct FIntermediateMargin
+{
+	double Left;
+	double Top;
+	double Right;
+	double Bottom;
+};
+UMG_API void ConvertOperationalProperty(const FIntermediateMargin& In, FMargin& Out);
+UMG_API void ConvertOperationalProperty(const FMargin& In, FIntermediateMargin& Out);
+
+using FMarginTraits = TIndirectPropertyTraits<FMargin, FIntermediateMargin>;
 using FWidgetTransformPropertyTraits = TIndirectPropertyTraits<FWidgetTransform, FIntermediateWidgetTransform>;
 
 struct UMG_API FMovieSceneUMGComponentTypes
diff --git a/Engine/Source/Runtime/UMG/Public/Animation/MovieSceneWidgetMaterialSystem.h b/Engine/Source/Runtime/UMG/Public/Animation/MovieSceneWidgetMaterialSystem.h
--- a/Engine/Source/Runtime/UMG/Public/Animation/MovieSceneWidgetMaterialSystem.h
+++ b/Engine/Source/Runtime/UMG/Public/Animation/MovieSceneWidgetMaterialSystem.h
@@ -11,7 +11,7 @@
 #include "MovieSceneWidgetMaterialSystem.generated.h"
 
 class UWidget;
-class UMovieScenePiecewiseFloatBlenderSystem;
+class UMovieScenePiecewiseDoubleBlenderSystem;
 
 namespace UE::MovieScene
 {
