Sequencer: Fixed pre-animated material parameters not being exposed to the reference graph

diff --git a/Engine/Source/Runtime/MovieSceneTracks/Public/Systems/MovieSceneMaterialSystem.h b/Engine/Source/Runtime/MovieSceneTracks/Public/Systems/MovieSceneMaterialSystem.h
--- a/Engine/Source/Runtime/MovieSceneTracks/Public/Systems/MovieSceneMaterialSystem.h
+++ b/Engine/Source/Runtime/MovieSceneTracks/Public/Systems/MovieSceneMaterialSystem.h
@@ -11,23 +11,34 @@
 
 #include "EntitySystem/BuiltInComponentTypes.h"
 #include "EntitySystem/MovieSceneEntitySystemRunner.h"
+#include "EntitySystem/MovieSceneComponentTypeInfo.h"
 #include "MovieSceneTracksComponentTypes.h"
 
-namespace UE::MovieScene
-{
+#include "MovieSceneMaterialSystem.generated.h"
 
-struct FPreAnimatedMaterialParameters
+USTRUCT()
+struct FMovieScenePreAnimatedMaterialParameters
 {
-	UMaterialInterface* PreviousMaterial = nullptr;
-	UMaterialInterface* PreviousParameterContainer = nullptr;
+	GENERATED_BODY()
+
+	UPROPERTY()
+	TObjectPtr<UMaterialInterface> PreviousMaterial = nullptr;
+
+	UPROPERTY()
+	TObjectPtr<UMaterialInterface> PreviousParameterContainer = nullptr;
 };
 
+namespace UE::MovieScene
+{
+
 template<typename AccessorType, typename... RequiredComponents>
 struct TPreAnimatedMaterialTraits : FBoundObjectPreAnimatedStateTraits
 {
 	using KeyType     = typename AccessorType::KeyType;
 	using StorageType = UMaterialInterface*;
 
+	static_assert(THasAddReferencedObjectForComponent<StorageType>::Value, "StorageType is not correctly exposed to the reference graph!");
+
 	static UMaterialInterface* CachePreAnimatedValue(typename TCallTraits<RequiredComponents>::ParamType... InRequiredComponents)
 	{
 		return AccessorType{ InRequiredComponents... }.GetMaterial();
@@ -43,13 +54,15 @@ template<typename AccessorType, typename... RequiredComponents>
 struct TPreAnimatedMaterialParameterTraits : FBoundObjectPreAnimatedStateTraits
 {
 	using KeyType     = typename AccessorType::KeyType;
-	using StorageType = FPreAnimatedMaterialParameters;
+	using StorageType = FMovieScenePreAnimatedMaterialParameters;
+
+	static_assert(THasAddReferencedObjectForComponent<StorageType>::Value, "StorageType is not correctly exposed to the reference graph!");
 
-	static FPreAnimatedMaterialParameters CachePreAnimatedValue(typename TCallTraits<RequiredComponents>::ParamType... InRequiredComponents)
+	static FMovieScenePreAnimatedMaterialParameters CachePreAnimatedValue(typename TCallTraits<RequiredComponents>::ParamType... InRequiredComponents)
 	{
 		AccessorType Accessor{ InRequiredComponents... };
 
-		FPreAnimatedMaterialParameters Parameters;
+		FMovieScenePreAnimatedMaterialParameters Parameters;
 		Parameters.PreviousMaterial = Accessor.GetMaterial();
 
 		// If the current material we're overriding is already a material instance dynamic, copy it since we will be modifying the data. 
@@ -63,7 +76,7 @@ struct TPreAnimatedMaterialParameterTraits : FBoundObjectPreAnimatedStateTraits
 		return Parameters;
 	}
 
-	static void RestorePreAnimatedValue(const KeyType& InKey, const FPreAnimatedMaterialParameters& PreAnimatedValue, const FRestoreStateParams& Params)
+	static void RestorePreAnimatedValue(const KeyType& InKey, const FMovieScenePreAnimatedMaterialParameters& PreAnimatedValue, const FRestoreStateParams& Params)
 	{
 		AccessorType Accessor{ InKey };
 		if (PreAnimatedValue.PreviousParameterContainer != nullptr)
