Sequencer: Custom accessor for mass scale since SetMassScale takes a BoneName.

diff --git a/Engine/Source/Runtime/MovieSceneTracks/Private/MovieSceneTracksComponentTypes.cpp b/Engine/Source/Runtime/MovieSceneTracks/Private/MovieSceneTracksComponentTypes.cpp
--- a/Engine/Source/Runtime/MovieSceneTracks/Private/MovieSceneTracksComponentTypes.cpp
+++ b/Engine/Source/Runtime/MovieSceneTracks/Private/MovieSceneTracksComponentTypes.cpp
@@ -3,6 +3,7 @@
 #include "MovieSceneTracksComponentTypes.h"
 #include "Components/ExponentialHeightFogComponent.h"
 #include "Components/LightComponent.h"
+#include "Components/PrimitiveComponent.h"
 #include "Components/SkyLightComponent.h"
 #include "EntitySystem/BuiltInComponentTypes.h"
 #include "EntitySystem/MovieSceneEntityManager.h"
@@ -25,6 +26,7 @@
 #include "GameFramework/Actor.h"
 #include "Materials/MaterialParameterCollection.h"
 #include "Misc/App.h"
+#include "PhysicsEngine/BodyInstance.h"
 
 namespace UE
 {
@@ -228,6 +230,18 @@ void SetSkyLightComponentLightColor(UObject* Object, EColorPropertyType InColorT
 	SkyLightComponent->SetLightColor(InColor.GetLinearColor());
 }
 
+float GetMassScale(const UObject* Object)
+{
+	const UPrimitiveComponent* PrimitiveComponent = CastChecked<const UPrimitiveComponent>(Object);
+	return PrimitiveComponent->GetMassScale();
+}
+
+void SetMassScale(UObject* Object, float InMassScale)
+{
+	UPrimitiveComponent* PrimitiveComponent = CastChecked<UPrimitiveComponent>(Object);
+	PrimitiveComponent->SetMassScale(NAME_None, InMassScale);
+}
+
 float GetSecondFogDataFogDensity(const UObject* Object)
 {
 	const UExponentialHeightFogComponent* ExponentialHeightFogComponent = CastChecked<const UExponentialHeightFogComponent>(Object);
@@ -628,6 +642,11 @@ FMovieSceneTracksComponentTypes::FMovieSceneTracksComponentTypes()
 			USkyLightComponent::StaticClass(), GET_MEMBER_NAME_CHECKED(USkyLightComponent, LightColor), 
 			GetSkyLightComponentLightColor, SetSkyLightComponentLightColor);
 	
+	const FString MassScalePath = FString::Printf(TEXT("%s.%s"), GET_MEMBER_NAME_STRING_CHECKED(UPrimitiveComponent, BodyInstance), GET_MEMBER_NAME_STRING_CHECKED(FBodyInstance, MassScale));
+	Accessors.Float.Add(
+		UPrimitiveComponent::StaticClass(), *MassScalePath,
+		GetMassScale, SetMassScale);
+
 	const FString SecondFogDataFogDensityPath = FString::Printf(TEXT("%s.%s"), GET_MEMBER_NAME_STRING_CHECKED(UExponentialHeightFogComponent, SecondFogData), GET_MEMBER_NAME_STRING_CHECKED(FExponentialHeightFogData, FogDensity));
 	Accessors.Float.Add(
 			UExponentialHeightFogComponent::StaticClass(), *SecondFogDataFogDensityPath,
