Remove Matinee (1/2)

diff --git a/Engine/Source/Runtime/Engine/Public/Tests/AutomationCommon.h b/Engine/Source/Runtime/Engine/Public/Tests/AutomationCommon.h
--- a/Engine/Source/Runtime/Engine/Public/Tests/AutomationCommon.h
+++ b/Engine/Source/Runtime/Engine/Public/Tests/AutomationCommon.h
@@ -5,7 +5,6 @@
 #include "CoreMinimal.h"
 #include "Misc/AutomationTest.h"
 
-class AMatineeActor;
 class SWindow;
 class SWidget;
 
@@ -125,17 +124,16 @@ DEFINE_ENGINE_LATENT_AUTOMATION_COMMAND(FWaitForMapToLoadCommand);
 */
 DEFINE_ENGINE_LATENT_AUTOMATION_COMMAND_ONE_PARAMETER(FWaitForSpecifiedMapToLoadCommand, FString, MapName);
 
-
 /**
 * Force a matinee to not loop and request that it play
 */
-DEFINE_ENGINE_LATENT_AUTOMATION_COMMAND_ONE_PARAMETER(FPlayMatineeLatentCommand, AMatineeActor*, MatineeActor);
+DEFINE_ENGINE_LATENT_AUTOMATION_COMMAND_ONE_PARAMETER(FPlayMatineeLatentCommand, AActor*, MatineeActor);
 
 
 /**
 * Wait for a particular matinee actor to finish playing
 */
-DEFINE_ENGINE_LATENT_AUTOMATION_COMMAND_ONE_PARAMETER(FWaitForMatineeToCompleteLatentCommand, AMatineeActor*, MatineeActor);
+DEFINE_ENGINE_LATENT_AUTOMATION_COMMAND_ONE_PARAMETER(FWaitForMatineeToCompleteLatentCommand, AActor*, MatineeActor);
 
 
 /**
diff --git a/Engine/Source/Runtime/MovieSceneTracks/Private/Evaluation/MovieSceneCameraShakeTemplate.cpp b/Engine/Source/Runtime/MovieSceneTracks/Private/Evaluation/MovieSceneCameraShakeTemplate.cpp
new file mode 100644
--- /dev/null
+++ b/Engine/Source/Runtime/MovieSceneTracks/Private/Evaluation/MovieSceneCameraShakeTemplate.cpp
@@ -0,0 +1,412 @@
+// Copyright Epic Games, Inc. All Rights Reserved.
+
+#include "Evaluation/MovieSceneCameraShakeTemplate.h"
+#include "Camera/CameraComponent.h"
+#include "GameFramework/Actor.h"
+#include "Engine/World.h"
+#include "MovieSceneCommonHelpers.h"
+#include "Evaluation/MovieSceneEvaluation.h"
+#include "UObject/Package.h"
+#include "IMovieScenePlayer.h"
+
+DECLARE_CYCLE_STAT(TEXT("Camera Anim Track Token Execute"), MovieSceneEval_CameraAnimTrack_TokenExecute, STATGROUP_MovieSceneEval);
+
+/** Structure that holds blended post processing settings */
+struct FBlendedPostProcessSettings : FPostProcessSettings
+{
+	FBlendedPostProcessSettings() : Weight(0.f) {}
+	FBlendedPostProcessSettings(float InWeight, const FPostProcessSettings& InSettings) : FPostProcessSettings(InSettings), Weight(InWeight) {}
+
+	/** The weighting to apply to these settings */
+	float Weight;
+};
+
+/** Persistent data that exists as long as a given camera track is being evaluated */
+struct FMovieSceneAdditiveCameraData : IPersistentEvaluationData
+{
+	FMovieSceneAdditiveCameraData()
+		: bApplyTransform(false)
+		, bApplyPostProcessing(false)
+		, TotalTransform(FTransform::Identity)
+		, TotalFOVOffset(0.f)
+	{}
+
+	static FMovieSceneSharedDataId GetSharedDataID()
+	{
+		static FMovieSceneSharedDataId SharedDataId = FMovieSceneSharedDataId::Allocate();
+		return SharedDataId;
+	}
+
+	static FMovieSceneAdditiveCameraData& Get(const FMovieSceneEvaluationOperand& Operand, FPersistentEvaluationData& PersistentData)
+	{
+		return PersistentData.GetOrAdd<FMovieSceneAdditiveCameraData>(FSharedPersistentDataKey(GetSharedDataID(), Operand));
+	}
+
+	/** Reset the additive camera values */
+	void Reset()
+	{
+		TotalFOVOffset = 0.f;
+		TotalTransform = FTransform::Identity;
+		BlendedPostProcessSettings.Reset();
+
+		bApplyTransform = false;
+		bApplyPostProcessing = false;
+	}
+
+	/** Accumulate the given post processing settings for this frame */
+	void AccumulatePostProcessing(const FPostProcessSettings& InPostProcessSettings, float Weight)
+	{
+		if (Weight > 0.f)
+		{
+			BlendedPostProcessSettings.Add(FBlendedPostProcessSettings(Weight, InPostProcessSettings));
+		}
+
+		bApplyPostProcessing = true;
+	}
+
+	/** Accumulate the transform and FOV offset */
+	void AccumulateOffset(const FTransform& AdditiveOffset, float AdditiveFOVOffset)
+	{
+		TotalTransform = TotalTransform * AdditiveOffset;
+		TotalFOVOffset += AdditiveFOVOffset;
+
+		bApplyTransform = true;
+	}
+
+	/** Apply any cumulative animation states */
+	void ApplyCumulativeAnimation(UCameraComponent& CameraComponent) const
+	{
+		if (bApplyPostProcessing)
+		{
+			CameraComponent.ClearExtraPostProcessBlends();
+			for (const FBlendedPostProcessSettings& Settings : BlendedPostProcessSettings)
+			{
+				CameraComponent.AddExtraPostProcessBlend(Settings, Settings.Weight);
+			}
+		}
+
+		if (bApplyTransform)
+		{
+			CameraComponent.ClearAdditiveOffset();
+			CameraComponent.AddAdditiveOffset(TotalTransform, TotalFOVOffset);
+		}
+	}
+
+private:
+
+	bool bApplyTransform, bApplyPostProcessing;
+	TArray<FBlendedPostProcessSettings, TInlineAllocator<2>> BlendedPostProcessSettings;
+	FTransform TotalTransform;
+	float TotalFOVOffset;
+};
+
+/** Pre animated token that restores a camera component's additive transform */
+struct FPreAnimatedCameraTransformTokenProducer : IMovieScenePreAnimatedTokenProducer
+{
+
+	static FMovieSceneAnimTypeID GetAnimTypeID()
+	{
+		static FMovieSceneAnimTypeID AnimTypeID = TMovieSceneAnimTypeID<FPreAnimatedCameraTransformTokenProducer, 0>();
+		return AnimTypeID;
+	}
+
+	virtual IMovieScenePreAnimatedTokenPtr CacheExistingState(UObject& Object) const override
+	{
+		struct FRestoreToken : IMovieScenePreAnimatedToken
+		{
+			virtual void RestoreState(UObject& InObject, const UE::MovieScene::FRestoreStateParams& Params) override
+			{
+				UCameraComponent* CameraComponent = CastChecked<UCameraComponent>(&InObject);
+				CameraComponent->ClearAdditiveOffset();
+			}
+		};
+
+		return FRestoreToken();
+	}
+};
+
+/** Pre animated token that restores a camera component's blended post processing settings */
+struct FPreAnimatedPostProcessingBlendsTokenProducer : IMovieScenePreAnimatedTokenProducer
+{
+	virtual IMovieScenePreAnimatedTokenPtr CacheExistingState(UObject& Object) const override
+	{
+		struct FRestoreToken : IMovieScenePreAnimatedToken
+		{
+			virtual void RestoreState(UObject& InObject, const UE::MovieScene::FRestoreStateParams& Params) override
+			{
+				UCameraComponent* CameraComponent = CastChecked<UCameraComponent>(&InObject);
+				CameraComponent->ClearExtraPostProcessBlends();
+			}
+		};
+
+		return FRestoreToken();
+	}
+};
+
+/** A movie scene execution token that applies camera cuts */
+struct FMovieSceneApplyAdditiveCameraDataExecutionToken : IMovieSceneSharedExecutionToken
+{
+	static void EnsureSharedToken(const FMovieSceneEvaluationOperand& Operand, FMovieSceneExecutionTokens& ExecutionTokens)
+	{
+		// Add a shared token that will apply the blended camera anims
+		FMovieSceneSharedDataId TokenID = FMovieSceneAdditiveCameraData::GetSharedDataID();
+		// It's safe to cast here as only FMovieSceneApplyAdditiveCameraDataExecutionTokens can have this token ID
+		if (FMovieSceneApplyAdditiveCameraDataExecutionToken* ExistingToken = static_cast<FMovieSceneApplyAdditiveCameraDataExecutionToken*>(ExecutionTokens.FindShared(TokenID)))
+		{
+			ExistingToken->Operands.Add(Operand);
+		}
+		else
+		{
+			ExecutionTokens.AddShared(TokenID, FMovieSceneApplyAdditiveCameraDataExecutionToken(Operand));
+		}
+	}
+
+	FMovieSceneApplyAdditiveCameraDataExecutionToken(const FMovieSceneEvaluationOperand& InOperand)
+	{
+		Operands.Add(InOperand);
+		// Evaluate after everything else in the group
+		Order = 1000;
+	}
+
+	virtual void Execute(FPersistentEvaluationData& PersistentData, IMovieScenePlayer& Player) override
+	{
+		for (FMovieSceneEvaluationOperand& Operand : Operands)
+		{
+			FMovieSceneAdditiveCameraData& SharedData = FMovieSceneAdditiveCameraData::Get(Operand, PersistentData);
+
+			for (TWeakObjectPtr<> ObjectWP : Player.FindBoundObjects(Operand))
+			{
+				UCameraComponent* CameraComponent = MovieSceneHelpers::CameraComponentFromRuntimeObject(ObjectWP.Get());
+				if (CameraComponent)
+				{
+					SharedData.ApplyCumulativeAnimation(*CameraComponent);
+				}
+			}
+
+			SharedData.Reset();
+		}
+	}
+
+	TSet<FMovieSceneEvaluationOperand> Operands;
+};
+
+/** Utility class for additive camera execution tokens */
+template<typename Impl, typename UserDataStruct>
+struct TAccumulateCameraAnimExecutionToken : IMovieSceneExecutionToken
+{
+	virtual void Execute(const FMovieSceneContext& Context, const FMovieSceneEvaluationOperand& Operand, FPersistentEvaluationData& PersistentData, IMovieScenePlayer& Player) override
+	{
+		UserDataStruct UserData;
+		if (!static_cast<Impl*>(this)->Impl::EnsureSetup(Operand, PersistentData, Player, UserData))
+		{
+			return;
+		}
+
+		FMovieSceneAdditiveCameraData& AdditiveSharedData = FMovieSceneAdditiveCameraData::Get(Operand, PersistentData);
+
+		for (TWeakObjectPtr<> ObjectWP : Player.FindBoundObjects(Operand))
+		{
+			UCameraComponent* CameraComponent = MovieSceneHelpers::CameraComponentFromRuntimeObject(ObjectWP.Get());
+			if (!CameraComponent)
+			{
+				continue;
+			}
+
+			FMinimalViewInfo POV;
+			POV.Location = CameraComponent->GetComponentLocation();
+			POV.Rotation = CameraComponent->GetComponentRotation();
+			POV.FOV = CameraComponent->FieldOfView;
+
+			float PostProcessBlendWeight = 0.f;
+			FPostProcessSettings PostProcessSettings;
+
+			if (!static_cast<Impl*>(this)->Impl::UpdateCamera(Context, Operand, PersistentData, Player, UserData, POV, PostProcessSettings, PostProcessBlendWeight))
+			{
+				continue;
+			}
+
+			// Grab transform and FOV changes.
+			FTransform WorldToBaseCamera = CameraComponent->GetComponentToWorld().Inverse();
+			float BaseFOV = CameraComponent->FieldOfView;
+			FTransform NewCameraToWorld(POV.Rotation, POV.Location);
+			float NewFOV = POV.FOV;
+
+			FTransform NewCameraToBaseCamera = NewCameraToWorld * WorldToBaseCamera;
+
+			float NewFOVToBaseFOV = BaseFOV - NewFOV;
+
+			{
+				static FMovieSceneAnimTypeID TransformAnimTypeID = TMovieSceneAnimTypeID<TAccumulateCameraAnimExecutionToken<Impl, UserDataStruct>, 0>();
+				Player.SavePreAnimatedState(*CameraComponent, TransformAnimTypeID, FPreAnimatedCameraTransformTokenProducer());
+
+				// Accumumulate the offsets into the track data for application as part of the track execution token
+				AdditiveSharedData.AccumulateOffset(NewCameraToBaseCamera, NewFOVToBaseFOV);
+			}
+
+			// Grab post process changes.
+			if (PostProcessBlendWeight > 0.f)
+			{
+				static FMovieSceneAnimTypeID PostAnimTypeID = TMovieSceneAnimTypeID<TAccumulateCameraAnimExecutionToken<Impl, UserDataStruct>, 1>();
+				Player.SavePreAnimatedState(*CameraComponent, PostAnimTypeID, FPreAnimatedPostProcessingBlendsTokenProducer());
+
+				AdditiveSharedData.AccumulatePostProcessing(PostProcessSettings, PostProcessBlendWeight);
+			}
+		}
+	}
+};
+
+// Initialize shake evaluator registry.
+TArray<FMovieSceneBuildShakeEvaluator> FMovieSceneCameraShakeEvaluatorRegistry::ShakeEvaluatorBuilders;
+
+/** Persistent data that exists as long as a given camera anim section is being evaluated */
+struct FMovieSceneCameraShakeSectionInstanceData : IPersistentEvaluationData
+{
+	/** Camera shake instance */
+	TStrongObjectPtr<UCameraShakeBase> CameraShakeInstance;
+	
+	/** Custom evaluator for the shake (optional) */
+	TStrongObjectPtr<UMovieSceneCameraShakeEvaluator> CameraShakeEvaluator;
+};
+
+/** Pre animated token that restores a camera animation */
+struct FPreAnimatedCameraShakeTokenProducer : IMovieScenePreAnimatedTokenProducer
+{
+	virtual IMovieScenePreAnimatedTokenPtr CacheExistingState(UObject& Object) const override
+	{
+		struct FRestoreToken : IMovieScenePreAnimatedToken
+		{
+			virtual void RestoreState(UObject& InObject, const UE::MovieScene::FRestoreStateParams& Params) override
+			{
+				UCameraShakeBase* CameraShake = CastChecked<UCameraShakeBase>(&InObject);
+				if (!CameraShake->IsFinished())	
+				{
+					CameraShake->StopShake(true);
+				}
+				CameraShake->TeardownShake();
+			}
+		};
+
+		return FRestoreToken();
+	}
+};
+
+/** A movie scene execution token that applies camera shakes */
+struct FCameraShakeExecutionToken : TAccumulateCameraAnimExecutionToken<FCameraShakeExecutionToken, bool>
+{
+	FCameraShakeExecutionToken(const FMovieSceneCameraShakeSectionData& InSourceData, FFrameNumber InSectionStartTime)
+		: SourceData(InSourceData) 
+		, SectionStartTime(InSectionStartTime)
+	{}
+
+	bool EnsureSetup(const FMovieSceneEvaluationOperand& Operand, FPersistentEvaluationData& PersistentData, IMovieScenePlayer& Player, bool& Dummy)
+	{
+		// Get the camera anim instance from the section data (local to this specific section)
+		FMovieSceneCameraShakeSectionInstanceData& SectionData = PersistentData.GetOrAddSectionData<FMovieSceneCameraShakeSectionInstanceData>();
+		UCameraShakeBase* CameraShakeInstance = SectionData.CameraShakeInstance.Get();
+
+		if (!CameraShakeInstance)
+		{
+			if (!*SourceData.ShakeClass)
+			{
+				return false;
+			}
+
+			UObject* OuterObject = Player.GetPlaybackContext() ? Player.GetPlaybackContext() : GetTransientPackage();
+			CameraShakeInstance = NewObject<UCameraShakeBase>(OuterObject, SourceData.ShakeClass);
+			if (CameraShakeInstance)
+			{
+				// Store the anim instance with the section and always remove it when we've finished evaluating
+				{
+					FMovieSceneAnimTypeID AnimTypeID = TMovieSceneAnimTypeID<FCameraShakeExecutionToken>();
+
+					FScopedPreAnimatedCaptureSource CaptureSource(&Player.PreAnimatedState, PersistentData.GetSectionKey(), true);
+					Player.PreAnimatedState.SavePreAnimatedState(*CameraShakeInstance, AnimTypeID, FPreAnimatedCameraShakeTokenProducer());
+				}
+
+				// Custom logic, if any.
+				UMovieSceneCameraShakeEvaluator* CameraShakeEvaluator = FMovieSceneCameraShakeEvaluatorRegistry::BuildShakeEvaluator(CameraShakeInstance);
+				if (CameraShakeEvaluator)
+				{
+					CameraShakeEvaluator->Setup(Operand, PersistentData, Player, CameraShakeInstance);
+				}
+				SectionData.CameraShakeEvaluator.Reset(CameraShakeEvaluator);
+
+				// Start the shake.
+				CameraShakeInstance->StartShake(nullptr, SourceData.PlayScale, SourceData.PlaySpace, SourceData.UserDefinedPlaySpace);
+			}
+			SectionData.CameraShakeInstance.Reset(CameraShakeInstance);
+		}
+		else
+		{
+			// We have a camera shake instance, but we need to check that it's still active. This is because
+			// our shake could have been stopped and torn down by a recompilation of the sequence (when the user
+			// edits it), an auto-save kicking in, etc.
+			if (!CameraShakeInstance->IsActive())
+			{
+				if (UMovieSceneCameraShakeEvaluator* CameraShakeEvaluator = SectionData.CameraShakeEvaluator.Get())
+				{
+					CameraShakeEvaluator->Setup(Operand, PersistentData, Player, CameraShakeInstance);
+				}
+				CameraShakeInstance->StartShake(nullptr, SourceData.PlayScale, SourceData.PlaySpace, SourceData.UserDefinedPlaySpace);
+			}
+		}
+
+		// If we failed to create the camera shake instance, we're doomed
+		return ensure(CameraShakeInstance);
+	}
+
+	bool UpdateCamera(const FMovieSceneContext& Context, const FMovieSceneEvaluationOperand& Operand, FPersistentEvaluationData& PersistentData, IMovieScenePlayer& Player, bool Dummy, FMinimalViewInfo& OutPOV, FPostProcessSettings& OutPostProcessSettings, float& OutPostProcessBlendWeight)
+	{
+		// Get the camera anim instance from the section data (local to this specific section)
+		FMovieSceneCameraShakeSectionInstanceData& SectionData = PersistentData.GetOrAddSectionData<FMovieSceneCameraShakeSectionInstanceData>();
+		UCameraShakeBase* CameraShakeInstance = SectionData.CameraShakeInstance.Get();
+
+		if (!ensure(CameraShakeInstance))
+		{
+			return false;
+		}
+
+		// Run custom shake logic if any.
+		if (SectionData.CameraShakeEvaluator)
+		{
+			SectionData.CameraShakeEvaluator->Evaluate(Context, Operand, PersistentData, Player, CameraShakeInstance);
+		}
+
+		// Update shake to the new time.
+		const FFrameTime NewShakeTime = Context.GetTime() - SectionStartTime;
+		CameraShakeInstance->ScrubAndApplyCameraShake(NewShakeTime / Context.GetFrameRate(), 1.f, OutPOV);
+
+		// TODO: post process settings
+
+		return true;
+	}
+
+	FMovieSceneCameraShakeSectionData SourceData;
+	FFrameNumber SectionStartTime;
+};
+
+FMovieSceneCameraShakeSectionTemplate::FMovieSceneCameraShakeSectionTemplate()
+	: SourceData()
+{
+}
+
+FMovieSceneCameraShakeSectionTemplate::FMovieSceneCameraShakeSectionTemplate(const UMovieSceneCameraShakeSection& Section)
+	: SourceData(Section.ShakeData)
+	, SectionStartTime(Section.HasStartFrame() ? Section.GetInclusiveStartFrame() : 0)
+{
+	RequiresInitialization();
+}
+
+void FMovieSceneCameraShakeSectionTemplate::Initialize(const FMovieSceneEvaluationOperand& Operand, const FMovieSceneContext& Context, FPersistentEvaluationData& PersistentData, IMovieScenePlayer& Player) const
+{
+	FMovieSceneCameraShakeSectionInstanceData& SectionData = PersistentData.GetOrAddSectionData<FMovieSceneCameraShakeSectionInstanceData>();
+	SectionData.CameraShakeInstance = nullptr;
+	SectionData.CameraShakeEvaluator = nullptr;
+}
+
+void FMovieSceneCameraShakeSectionTemplate::Evaluate(const FMovieSceneEvaluationOperand& Operand, const FMovieSceneContext& Context, const FPersistentEvaluationData& PersistentData, FMovieSceneExecutionTokens& ExecutionTokens) const
+{
+	ExecutionTokens.Add(FCameraShakeExecutionToken(SourceData, SectionStartTime));
+	FMovieSceneApplyAdditiveCameraDataExecutionToken::EnsureSharedToken(Operand, ExecutionTokens);
+}
+
diff --git a/Engine/Source/Runtime/MovieSceneTracks/Private/Evaluation/MovieSceneSlomoTemplate.cpp b/Engine/Source/Runtime/MovieSceneTracks/Private/Evaluation/MovieSceneSlomoTemplate.cpp
--- a/Engine/Source/Runtime/MovieSceneTracks/Private/Evaluation/MovieSceneSlomoTemplate.cpp
+++ b/Engine/Source/Runtime/MovieSceneTracks/Private/Evaluation/MovieSceneSlomoTemplate.cpp
@@ -36,7 +36,7 @@ struct FSlomoTrackToken
 
 		if (WorldSettings)
 		{
-			WorldSettings->MatineeTimeDilation = SlomoValue;
+			WorldSettings->CinematicTimeDilation = SlomoValue;
 			WorldSettings->ForceNetUpdate();
 		}
 	}
@@ -75,7 +75,7 @@ struct FSlomoPreAnimatedGlobalTokenProducer : IMovieScenePreAnimatedGlobalTokenP
 	{
 		if (AWorldSettings* WorldSettings = Player.GetPlaybackContext()->GetWorld()->GetWorldSettings())
 		{
-			return FSlomoPreAnimatedGlobalToken(WorldSettings->MatineeTimeDilation);
+			return FSlomoPreAnimatedGlobalToken(WorldSettings->CinematicTimeDilation);
 		}
 		return IMovieScenePreAnimatedGlobalTokenPtr();
 	}
diff --git a/Engine/Source/Runtime/MovieSceneTracks/Private/Tracks/MovieSceneCameraShakeTrack.cpp b/Engine/Source/Runtime/MovieSceneTracks/Private/Tracks/MovieSceneCameraShakeTrack.cpp
--- a/Engine/Source/Runtime/MovieSceneTracks/Private/Tracks/MovieSceneCameraShakeTrack.cpp
+++ b/Engine/Source/Runtime/MovieSceneTracks/Private/Tracks/MovieSceneCameraShakeTrack.cpp
@@ -3,7 +3,7 @@
 #include "Tracks/MovieSceneCameraShakeTrack.h"
 #include "Sections/MovieSceneCameraShakeSection.h"
 #include "Evaluation/PersistentEvaluationData.h"
-#include "Evaluation/MovieSceneCameraAnimTemplate.h"
+#include "Evaluation/MovieSceneCameraShakeTemplate.h"
 #include "Evaluation/MovieSceneEvaluationTrack.h"
 #include "Evaluation/MovieSceneEvaluationTemplate.h"
 #include "Compilation/MovieSceneSegmentCompiler.h"
diff --git a/Engine/Source/Runtime/Renderer/Private/SceneRendering.cpp b/Engine/Source/Runtime/Renderer/Private/SceneRendering.cpp
--- a/Engine/Source/Runtime/Renderer/Private/SceneRendering.cpp
+++ b/Engine/Source/Runtime/Renderer/Private/SceneRendering.cpp
@@ -29,7 +29,6 @@
 #include "CompositionLighting/CompositionLighting.h"
 #include "LegacyScreenPercentageDriver.h"
 #include "SceneViewExtension.h"
-#include "Matinee/MatineeActor.h"
 #include "ComponentRecreateRenderStateContext.h"
 #include "PostProcess/PostProcessSubsurface.h"
 #include "PhysicsField/PhysicsFieldComponent.h"
@@ -5443,4 +5442,4 @@ FRDGTextureRef CreateHalfResolutionDepthCheckerboardMinMax(FRDGBuilder& GraphBui
 	}
 
 	return SmallDepthTexture;
-}
\ No newline at end of file
+}
