make sure we use the SetActive function to also notify the tick function from state changes

diff --git a/Engine/Source/Runtime/MovieSceneTracks/Private/Systems/MovieSceneConstraintSystem.cpp b/Engine/Source/Runtime/MovieSceneTracks/Private/Systems/MovieSceneConstraintSystem.cpp
--- a/Engine/Source/Runtime/MovieSceneTracks/Private/Systems/MovieSceneConstraintSystem.cpp
+++ b/Engine/Source/Runtime/MovieSceneTracks/Private/Systems/MovieSceneConstraintSystem.cpp
@@ -45,7 +45,7 @@ struct FPreAnimatedConstraintTraits : FBoundObjectPreAnimatedStateTraits
 	{
 		if (UTickableConstraint* Constraint = OldValue.WeakConstraint.Get())
 		{
-			Constraint->Active = OldValue.bPreviouslyEnabled;
+			Constraint->SetActive(OldValue.bPreviouslyEnabled);
 		}
 	}
 };
@@ -113,11 +113,9 @@ void UMovieSceneConstraintSystem::OnRun(FSystemTaskPrerequisites& InPrerequisite
 				UTickableConstraint* Constraint = Controller->GetConstraint(ConstraintChannel.ConstraintName);
 				if (Constraint)
 				{
-					bool Result;
-					if (ConstraintChannel.Channel->Evaluate(FrameTime, Result))
-					{
-						Constraint->Active = Result;
-					}
+					bool Result = false;
+					ConstraintChannel.Channel->Evaluate(FrameTime, Result);
+					Constraint->SetActive(Result);
 				}
 			}
 			FConstraintsManagerController* Controller;
