Sequencer: Check for valid player

The fixed some crashes that occurred when saving a level sequence while in a multi-user session. I don't have a simple repro to test against but this was tested in the Partner-Otter branch. It's possible that the PlayerUseCounts are being reset or not even being added when it gets to here.

diff --git a/Engine/Source/Runtime/MovieSceneTracks/Private/TrackInstances/MovieSceneCameraCutTrackInstance.cpp b/Engine/Source/Runtime/MovieSceneTracks/Private/TrackInstances/MovieSceneCameraCutTrackInstance.cpp
--- a/Engine/Source/Runtime/MovieSceneTracks/Private/TrackInstances/MovieSceneCameraCutTrackInstance.cpp
+++ b/Engine/Source/Runtime/MovieSceneTracks/Private/TrackInstances/MovieSceneCameraCutTrackInstance.cpp
@@ -378,11 +378,13 @@ void UMovieSceneCameraCutTrackInstance::OnAnimate()
 		if (FCameraCutAnimator::AnimateBlendedCameraCut(FinalCameraCut, CameraCutCache, Context, *Player))
 		{
 			// Track whether this ever evaluated to take control. If so, we'll want to remove control in OnDestroyed.
-			FCameraCutUseData& PlayerUseCount = PlayerUseCounts.FindChecked(Player);
-			PlayerUseCount.bValid = true;
-			// Remember whether we had blending support the last time we took control of the viewport. This is also
-			// for OnDestroyed.
-			PlayerUseCount.bCanBlend = FinalCameraCut.bCanBlend;
+			if (FCameraCutUseData* PlayerUseCount = PlayerUseCounts.Find(Player))
+			{
+				PlayerUseCount->bValid = true;
+				// Remember whether we had blending support the last time we took control of the viewport. This is also
+				// for OnDestroyed.
+				PlayerUseCount->bCanBlend = FinalCameraCut.bCanBlend;
+			}
 		}
 	}
 }
@@ -407,11 +409,13 @@ void UMovieSceneCameraCutTrackInstance::OnInputRemoved(const FMovieSceneTrackIns
 	const FSequenceInstance& SequenceInstance = InstanceRegistry->GetInstance(InInput.InstanceHandle);
 	IMovieScenePlayer* Player = SequenceInstance.GetPlayer();
 
-	int32& UseCount = PlayerUseCounts.FindChecked(Player).UseCount;
-	--UseCount;
-	if (UseCount == 0)
+	if (FCameraCutUseData* PlayerUseCount = PlayerUseCounts.Find(Player))
 	{
-		PlayerUseCounts.Remove(Player);
+		PlayerUseCount->UseCount--;
+		if (PlayerUseCount->UseCount == 0)
+		{
+			PlayerUseCounts.Remove(Player);
+		}
 	}
 }
 
