Sequencer: Backwards compatibility is now maintained for material parameter tracks that have completely empty channels

Previously such channels would always receive 1.f, but the new code would either set them to 0.f (if they didn't exist), or MIN_flt (if they existed but were never evaluated)

diff --git a/Engine/Source/Runtime/MovieSceneTracks/Private/Sections/MovieSceneParameterSection.cpp b/Engine/Source/Runtime/MovieSceneTracks/Private/Sections/MovieSceneParameterSection.cpp
--- a/Engine/Source/Runtime/MovieSceneTracks/Private/Sections/MovieSceneParameterSection.cpp
+++ b/Engine/Source/Runtime/MovieSceneTracks/Private/Sections/MovieSceneParameterSection.cpp
@@ -626,76 +626,97 @@ void UMovieSceneParameterSection::ImportEntityImpl(UMovieSceneEntitySystemLinker
 		{
 			const FScalarParameterNameAndCurve& Scalar = ScalarParameterNamesAndCurves[EntityIndex];
 
-			OutImportedEntity->AddBuilder(
-				BaseBuilder
-				.Add(TracksComponentTypes->ScalarParameterName, Scalar.ParameterName)
-				.Add(BuiltInComponentTypes->FloatChannel[0], &Scalar.ParameterCurve)
-			);
+			if (Scalar.ParameterCurve.HasAnyData())
+			{
+				OutImportedEntity->AddBuilder(
+					BaseBuilder
+					.Add(TracksComponentTypes->ScalarParameterName, Scalar.ParameterName)
+					.Add(BuiltInComponentTypes->FloatChannel[0], &Scalar.ParameterCurve)
+				);
+			}
 			break;
 		}
 		case 1:
 		{
 			const FBoolParameterNameAndCurve& Bool = BoolParameterNamesAndCurves[EntityIndex];
 
-			OutImportedEntity->AddBuilder(
-				BaseBuilder
-				.Add(TracksComponentTypes->BoolParameterName, Bool.ParameterName)
-				.Add(BuiltInComponentTypes->BoolChannel, &Bool.ParameterCurve)
-			);
+			if (Bool.ParameterCurve.HasAnyData())
+			{
+				OutImportedEntity->AddBuilder(
+					BaseBuilder
+					.Add(TracksComponentTypes->BoolParameterName, Bool.ParameterName)
+					.Add(BuiltInComponentTypes->BoolChannel, &Bool.ParameterCurve)
+				);
+			}
 			break;
 		}
 		case 2:
 		{
 			const FVector2DParameterNameAndCurves& Vector2D = Vector2DParameterNamesAndCurves[EntityIndex];
-			OutImportedEntity->AddBuilder(
-				BaseBuilder
-				.Add(TracksComponentTypes->Vector2DParameterName, Vector2D.ParameterName)
-				.Add(BuiltInComponentTypes->FloatChannel[0], &Vector2D.XCurve)
-				.Add(BuiltInComponentTypes->FloatChannel[1], &Vector2D.YCurve)
-			);
+			if (Vector2D.XCurve.HasAnyData() || Vector2D.YCurve.HasAnyData())
+			{
+				OutImportedEntity->AddBuilder(
+					BaseBuilder
+					.Add(TracksComponentTypes->Vector2DParameterName, Vector2D.ParameterName)
+					.AddConditional(BuiltInComponentTypes->FloatChannel[0], &Vector2D.XCurve, Vector2D.XCurve.HasAnyData())
+					.AddConditional(BuiltInComponentTypes->FloatChannel[1], &Vector2D.YCurve, Vector2D.YCurve.HasAnyData())
+				);
+			}
 			break;
 		}
 		case 3:
 		{
 			const FVectorParameterNameAndCurves& Vector = VectorParameterNamesAndCurves[EntityIndex];
-			OutImportedEntity->AddBuilder(
-				BaseBuilder
-				.Add(TracksComponentTypes->VectorParameterName, Vector.ParameterName)
-				.Add(BuiltInComponentTypes->FloatChannel[0], &Vector.XCurve)
-				.Add(BuiltInComponentTypes->FloatChannel[1], &Vector.YCurve)
-				.Add(BuiltInComponentTypes->FloatChannel[2], &Vector.ZCurve)
-			);
+
+			if (Vector.XCurve.HasAnyData() || Vector.YCurve.HasAnyData() || Vector.ZCurve.HasAnyData())
+			{
+				OutImportedEntity->AddBuilder(
+					BaseBuilder
+					.Add(TracksComponentTypes->VectorParameterName, Vector.ParameterName)
+					.AddConditional(BuiltInComponentTypes->FloatChannel[0], &Vector.XCurve, Vector.XCurve.HasAnyData())
+					.AddConditional(BuiltInComponentTypes->FloatChannel[1], &Vector.YCurve, Vector.YCurve.HasAnyData())
+					.AddConditional(BuiltInComponentTypes->FloatChannel[2], &Vector.ZCurve, Vector.ZCurve.HasAnyData())
+				);
+			}
 			break;
 		}
 		case 4:
 		{
 			const FColorParameterNameAndCurves& Color = ColorParameterNamesAndCurves[EntityIndex];
-			OutImportedEntity->AddBuilder(
-				BaseBuilder
-				.Add(TracksComponentTypes->ColorParameterName, Color.ParameterName)
-				.Add(BuiltInComponentTypes->FloatChannel[0], &Color.RedCurve)
-				.Add(BuiltInComponentTypes->FloatChannel[1], &Color.GreenCurve)
-				.Add(BuiltInComponentTypes->FloatChannel[2], &Color.BlueCurve)
-				.Add(BuiltInComponentTypes->FloatChannel[3], &Color.AlphaCurve)
-			);
+			if (Color.RedCurve.HasAnyData() || Color.GreenCurve.HasAnyData() || Color.BlueCurve.HasAnyData() || Color.AlphaCurve.HasAnyData())
+			{
+				OutImportedEntity->AddBuilder(
+					BaseBuilder
+					.Add(TracksComponentTypes->ColorParameterName, Color.ParameterName)
+					.AddConditional(BuiltInComponentTypes->FloatChannel[0], &Color.RedCurve, Color.RedCurve.HasAnyData())
+					.AddConditional(BuiltInComponentTypes->FloatChannel[1], &Color.GreenCurve, Color.GreenCurve.HasAnyData())
+					.AddConditional(BuiltInComponentTypes->FloatChannel[2], &Color.BlueCurve, Color.BlueCurve.HasAnyData())
+					.AddConditional(BuiltInComponentTypes->FloatChannel[3], &Color.AlphaCurve, Color.AlphaCurve.HasAnyData())
+				);
+			}
 			break;
 		}
 		case 5:
 		{
 			const FTransformParameterNameAndCurves& Transform = TransformParameterNamesAndCurves[EntityIndex];
-			OutImportedEntity->AddBuilder(
-				BaseBuilder
-				.Add(TracksComponentTypes->TransformParameterName, Transform.ParameterName)
-				.Add(BuiltInComponentTypes->FloatChannel[0], &Transform.Translation[0])
-				.Add(BuiltInComponentTypes->FloatChannel[1], &Transform.Translation[1])
-				.Add(BuiltInComponentTypes->FloatChannel[2], &Transform.Translation[2])
-				.Add(BuiltInComponentTypes->FloatChannel[3], &Transform.Rotation[0])
-				.Add(BuiltInComponentTypes->FloatChannel[4], &Transform.Rotation[1])
-				.Add(BuiltInComponentTypes->FloatChannel[5], &Transform.Rotation[2])
-				.Add(BuiltInComponentTypes->FloatChannel[6], &Transform.Scale[0])
-				.Add(BuiltInComponentTypes->FloatChannel[7], &Transform.Scale[1])
-				.Add(BuiltInComponentTypes->FloatChannel[8], &Transform.Scale[2])
-			);
+			if (Transform.Translation[0].HasAnyData() || Transform.Translation[1].HasAnyData() || Transform.Translation[2].HasAnyData() ||
+				Transform.Rotation[0].HasAnyData() || Transform.Rotation[1].HasAnyData() || Transform.Rotation[2].HasAnyData()
+				|| Transform.Scale[0].HasAnyData() || Transform.Scale[1].HasAnyData() || Transform.Scale[2].HasAnyData())
+			{
+				OutImportedEntity->AddBuilder(
+					BaseBuilder
+					.Add(TracksComponentTypes->TransformParameterName, Transform.ParameterName)
+					.AddConditional(BuiltInComponentTypes->FloatChannel[0], &Transform.Translation[0], Transform.Translation[0].HasAnyData())
+					.AddConditional(BuiltInComponentTypes->FloatChannel[1], &Transform.Translation[1], Transform.Translation[1].HasAnyData())
+					.AddConditional(BuiltInComponentTypes->FloatChannel[2], &Transform.Translation[2], Transform.Translation[2].HasAnyData())
+					.AddConditional(BuiltInComponentTypes->FloatChannel[3], &Transform.Rotation[0], Transform.Rotation[0].HasAnyData())
+					.AddConditional(BuiltInComponentTypes->FloatChannel[4], &Transform.Rotation[1], Transform.Rotation[1].HasAnyData())
+					.AddConditional(BuiltInComponentTypes->FloatChannel[5], &Transform.Rotation[2], Transform.Rotation[2].HasAnyData())
+					.AddConditional(BuiltInComponentTypes->FloatChannel[6], &Transform.Scale[0], Transform.Scale[0].HasAnyData())
+					.AddConditional(BuiltInComponentTypes->FloatChannel[7], &Transform.Scale[1], Transform.Scale[1].HasAnyData())
+					.AddConditional(BuiltInComponentTypes->FloatChannel[8], &Transform.Scale[2], Transform.Scale[2].HasAnyData())
+				);
+			}
 			break;
 		}
 	}
diff --git a/Engine/Source/Runtime/MovieSceneTracks/Private/Systems/MovieSceneMaterialParameterSystem.cpp b/Engine/Source/Runtime/MovieSceneTracks/Private/Systems/MovieSceneMaterialParameterSystem.cpp
--- a/Engine/Source/Runtime/MovieSceneTracks/Private/Systems/MovieSceneMaterialParameterSystem.cpp
+++ b/Engine/Source/Runtime/MovieSceneTracks/Private/Systems/MovieSceneMaterialParameterSystem.cpp
@@ -55,10 +55,11 @@ struct FApplyVectorParameters
 				continue;
 			}
 
+			// Default to white for unanimated channels for backwards compatibility with the legacy track
 			FLinearColor Color(
-				R ? (float)R[Index] : 0.f,
-				G ? (float)G[Index] : 0.f,
-				B ? (float)B[Index] : 0.f,
+				R ? (float)R[Index] : 1.f,
+				G ? (float)G[Index] : 1.f,
+				B ? (float)B[Index] : 1.f,
 				A ? (float)A[Index] : 1.f
 			);
 			if (UMaterialInstanceDynamic* MID = Cast<UMaterialInstanceDynamic>(BoundMaterial))
