Sequencer: Handle non-existent runners in material systems with an ensure

diff --git a/Engine/Source/Runtime/MovieSceneTracks/Private/Systems/MovieSceneMaterialParameterSystem.cpp b/Engine/Source/Runtime/MovieSceneTracks/Private/Systems/MovieSceneMaterialParameterSystem.cpp
--- a/Engine/Source/Runtime/MovieSceneTracks/Private/Systems/MovieSceneMaterialParameterSystem.cpp
+++ b/Engine/Source/Runtime/MovieSceneTracks/Private/Systems/MovieSceneMaterialParameterSystem.cpp
@@ -216,7 +216,13 @@ void UMovieSceneMaterialParameterSystem::OnRun(FSystemTaskPrerequisites& InPrere
 {
 	using namespace UE::MovieScene;
 
-	ESystemPhase CurrentPhase = Linker->GetActiveRunner()->GetCurrentPhase();
+	FMovieSceneEntitySystemRunner* ActiveRunner = Linker->GetActiveRunner();
+	if (!ensure(ActiveRunner))
+	{
+		return;
+	}
+
+	ESystemPhase CurrentPhase = ActiveRunner->GetCurrentPhase();
 	if (CurrentPhase == ESystemPhase::Instantiation)
 	{
 		OnInstantiation();
diff --git a/Engine/Source/Runtime/MovieSceneTracks/Public/Systems/MovieSceneMaterialSystem.h b/Engine/Source/Runtime/MovieSceneTracks/Public/Systems/MovieSceneMaterialSystem.h
--- a/Engine/Source/Runtime/MovieSceneTracks/Public/Systems/MovieSceneMaterialSystem.h
+++ b/Engine/Source/Runtime/MovieSceneTracks/Public/Systems/MovieSceneMaterialSystem.h
@@ -266,7 +266,13 @@ void TMovieSceneMaterialSystem<AccessorType, RequiredComponents...>::OnRun(UMovi
 	FBuiltInComponentTypes*          BuiltInComponents = FBuiltInComponentTypes::Get();
 	FMovieSceneTracksComponentTypes* TracksComponents  = FMovieSceneTracksComponentTypes::Get();
 
-	ESystemPhase CurrentPhase = Linker->GetActiveRunner()->GetCurrentPhase();
+	FMovieSceneEntitySystemRunner* ActiveRunner = Linker->GetActiveRunner();
+	if (!ensure(ActiveRunner))
+	{
+		return;
+	}
+
+	ESystemPhase CurrentPhase = ActiveRunner->GetCurrentPhase();
 	if (CurrentPhase == ESystemPhase::Instantiation)
 	{
 		// Only mutate things that are tagged as requiring linking
