Fix crash in Sidebar when scrolling through User Lists with the mouse wheel.

#rnx
#Tests Ran game in a single client and navigated to the friends list.  Tried scrolling up and down the user list repeatedly to make sure there was no crash.

diff --git a/Engine/Source/Runtime/MovieSceneTracks/Private/Systems/MovieScenePiecewiseDoubleBlenderSystem.cpp b/Engine/Source/Runtime/MovieSceneTracks/Private/Systems/MovieScenePiecewiseDoubleBlenderSystem.cpp
--- a/Engine/Source/Runtime/MovieSceneTracks/Private/Systems/MovieScenePiecewiseDoubleBlenderSystem.cpp
+++ b/Engine/Source/Runtime/MovieSceneTracks/Private/Systems/MovieScenePiecewiseDoubleBlenderSystem.cpp
@@ -63,10 +63,16 @@ private:
 				const FMovieSceneBlendChannelID& BlendID(BlendIDs[Index]);
 				ensureMsgf(BlendID.SystemID == BlenderSystemID, TEXT("Overriding the standard blender system of standard types isn't supported."));
 
+				if (!OutBlendResults.IsValidIndex(BlendID.ChannelID))
+				{
+					const bool bShrinkArray = false;
+					OutBlendResults.SetNum(BlendID.ChannelID + 1, bShrinkArray);
+				}
+
 				FBlendResult& Result = OutBlendResults[BlendID.ChannelID];
 
 				const float Weight = OptionalEasingAndWeights[Index];
-				Result.Total  += InResults[Index] * Weight;
+				Result.Total += InResults[Index] * Weight;
 				Result.Weight += Weight;
 			}
 		}
@@ -78,8 +84,14 @@ private:
 				const FMovieSceneBlendChannelID& BlendID(BlendIDs[Index]);
 				ensureMsgf(BlendID.SystemID == BlenderSystemID, TEXT("Overriding the standard blender system of standard types isn't supported."));
 
+				if (!OutBlendResults.IsValidIndex(BlendID.ChannelID))
+				{
+					const bool bShrinkArray = false;
+					OutBlendResults.SetNum(BlendID.ChannelID + 1, bShrinkArray);
+				}
+
 				FBlendResult& Result = OutBlendResults[BlendID.ChannelID];
-				Result.Total  += InResults[Index];
+				Result.Total += InResults[Index];
 				Result.Weight += 1.f;
 			}
 		}
