Sequencer: Follow up fix to LWC float->double conversion during property instantiation

ROBOMERGE-AUTHOR: Max.Chen

Remove float result components and the float blender system, do everything in double precision. Handle LWC float->double conversion during property instantiation.

diff --git a/Engine/Source/Runtime/MovieSceneTracks/Private/Systems/MovieScenePropertyInstantiator.cpp b/Engine/Source/Runtime/MovieSceneTracks/Private/Systems/MovieScenePropertyInstantiator.cpp
--- a/Engine/Source/Runtime/MovieSceneTracks/Private/Systems/MovieScenePropertyInstantiator.cpp
+++ b/Engine/Source/Runtime/MovieSceneTracks/Private/Systems/MovieScenePropertyInstantiator.cpp
@@ -306,7 +306,6 @@ void UMovieScenePropertyInstantiatorSystem::UpgradeFloatToDoubleProperties(const
 		{
 			// Swap out the property tag and result components, and add the result widening tag for easily
 			// finding these entities later.
-			int32 NumNewContributors = 0;
 			for (auto ContribIt = NewContributors.CreateKeyIterator(PropertyIndex); ContribIt; ++ContribIt)
 			{
 				const FMovieSceneEntityID CurID(ContribIt.Value());
@@ -314,17 +313,14 @@ void UMovieScenePropertyInstantiatorSystem::UpgradeFloatToDoubleProperties(const
 				EntityType.Remove(TrackComponents->Float.PropertyTag);
 				EntityType.Set(TrackComponents->Double.PropertyTag);
 				Linker->EntityManager.ChangeEntityType(CurID, EntityType);
-
-				++NumNewContributors;
 			}
 
-			PropertyStats[OldPropertyDefinitionIndex].NumProperties -= NumNewContributors;
-			PropertyStats[PropertyInfo.PropertyDefinitionIndex].NumProperties += NumNewContributors;
+			--PropertyStats[OldPropertyDefinitionIndex].NumProperties;
+			++PropertyStats[PropertyInfo.PropertyDefinitionIndex].NumProperties;
 		}
 		else if (PropertyDefinition.PropertyType == TrackComponents->DoubleVector.PropertyTag)
 		{
 			// Swap out the property tag and result components for the composites this contributor has.
-			int32 NumNewContributors = 0;
 			for (auto ContribIt = NewContributors.CreateKeyIterator(PropertyIndex); ContribIt; ++ContribIt)
 			{
 				const FMovieSceneEntityID CurID(ContribIt.Value());
@@ -332,12 +328,10 @@ void UMovieScenePropertyInstantiatorSystem::UpgradeFloatToDoubleProperties(const
 				EntityType.Remove(TrackComponents->FloatVector.PropertyTag);
 				EntityType.Set(TrackComponents->DoubleVector.PropertyTag);
 				Linker->EntityManager.ChangeEntityType(CurID, EntityType);
-
-				++NumNewContributors;
 			}
 
-			PropertyStats[OldPropertyDefinitionIndex].NumProperties -= NumNewContributors;
-			PropertyStats[PropertyInfo.PropertyDefinitionIndex].NumProperties += NumNewContributors;
+			--PropertyStats[OldPropertyDefinitionIndex].NumProperties;
+			++PropertyStats[PropertyInfo.PropertyDefinitionIndex].NumProperties;
 		}
 		else
 		{
