Sequencer: Temporarily disabled material parameter blending

Currently there is an issue where materials are re-assigned during evaluation, causing stale overlapping entity data which is not trivial to correct.
It can be re-enabled using the Sequencer.MaterialParameterBlending cvar for debugging.

diff --git a/Engine/Source/Runtime/MovieSceneTracks/Private/Systems/MovieSceneMaterialParameterSystem.cpp b/Engine/Source/Runtime/MovieSceneTracks/Private/Systems/MovieSceneMaterialParameterSystem.cpp
--- a/Engine/Source/Runtime/MovieSceneTracks/Private/Systems/MovieSceneMaterialParameterSystem.cpp
+++ b/Engine/Source/Runtime/MovieSceneTracks/Private/Systems/MovieSceneMaterialParameterSystem.cpp
@@ -14,6 +14,14 @@
 namespace UE::MovieScene
 {
 
+bool GMaterialParameterBlending = false;
+FAutoConsoleVariableRef CVarMaterialParameterBlending(
+	TEXT("Sequencer.MaterialParameterBlending"),
+	GMaterialParameterBlending,
+	TEXT("(Default: false) Defines whether material parameter blending should be enabled or not.\n"),
+	ECVF_Default
+);
+
 /** Apply scalar material parameters */
 struct FApplyScalarParameters
 {
@@ -203,13 +211,16 @@ void UMovieSceneMaterialParameterSystem::OnLink()
 	// Always reset the float blender system on link to ensure that recycled systems are correctly initialized.
 	DoubleBlenderSystem = nullptr;
 
-	FOverlappingMaterialParameterHandler Handler(this);
+	if (GMaterialParameterBlending)
+	{
+		FOverlappingMaterialParameterHandler Handler(this);
 
-	ScalarParameterTracker.Destroy(Handler);
-	VectorParameterTracker.Destroy(Handler);
+		ScalarParameterTracker.Destroy(Handler);
+		VectorParameterTracker.Destroy(Handler);
 
-	ScalarParameterTracker.Initialize(this);
-	VectorParameterTracker.Initialize(this);
+		ScalarParameterTracker.Initialize(this);
+		VectorParameterTracker.Initialize(this);
+	}
 }
 
 void UMovieSceneMaterialParameterSystem::OnRun(FSystemTaskPrerequisites& InPrerequisites, FSystemSubsequentTasks& Subsequents)
@@ -241,32 +252,35 @@ void UMovieSceneMaterialParameterSystem::OnInstantiation()
 	FBuiltInComponentTypes*          BuiltInComponents = FBuiltInComponentTypes::Get();
 	FMovieSceneTracksComponentTypes* TracksComponents  = FMovieSceneTracksComponentTypes::Get();
 
-	const bool bHasScalars = Linker->EntityManager.ContainsComponent(TracksComponents->ScalarParameterName);
-	const bool bHasColors = Linker->EntityManager.ContainsComponent(TracksComponents->ColorParameterName);
-	const bool bHasVectors = Linker->EntityManager.ContainsComponent(TracksComponents->VectorParameterName);
-
-	if (bHasScalars)
+	if (GMaterialParameterBlending)
 	{
-		ScalarParameterTracker.Update(Linker, TracksComponents->BoundMaterial, TracksComponents->ScalarParameterName, FEntityComponentFilter());
+		const bool bHasScalars = Linker->EntityManager.ContainsComponent(TracksComponents->ScalarParameterName);
+		const bool bHasColors = Linker->EntityManager.ContainsComponent(TracksComponents->ColorParameterName);
+		const bool bHasVectors = Linker->EntityManager.ContainsComponent(TracksComponents->VectorParameterName);
 
-		FOverlappingMaterialParameterHandler Handler(this);
-		Handler.DefaultComponentMask.Set(BuiltInComponents->DoubleResult[0]);
-		ScalarParameterTracker.ProcessInvalidatedOutputs(Linker, Handler);
-	}
+		if (bHasScalars)
+		{
+			ScalarParameterTracker.Update(Linker, TracksComponents->BoundMaterial, TracksComponents->ScalarParameterName, FEntityComponentFilter());
 
-	if (bHasColors)
-	{
-		VectorParameterTracker.Update(Linker, TracksComponents->BoundMaterial, TracksComponents->ColorParameterName, FEntityComponentFilter());
-	}
-	if (bHasVectors)
-	{
-		VectorParameterTracker.Update(Linker, TracksComponents->BoundMaterial, TracksComponents->VectorParameterName, FEntityComponentFilter());
-	}
-	if (bHasColors || bHasVectors)
-	{
-		FOverlappingMaterialParameterHandler Handler(this);
-		Handler.DefaultComponentMask.SetAll({ BuiltInComponents->DoubleResult[0], BuiltInComponents->DoubleResult[1], BuiltInComponents->DoubleResult[2], BuiltInComponents->DoubleResult[3] });
-		VectorParameterTracker.ProcessInvalidatedOutputs(Linker, Handler);
+			FOverlappingMaterialParameterHandler Handler(this);
+			Handler.DefaultComponentMask.Set(BuiltInComponents->DoubleResult[0]);
+			ScalarParameterTracker.ProcessInvalidatedOutputs(Linker, Handler);
+		}
+
+		if (bHasColors)
+		{
+			VectorParameterTracker.Update(Linker, TracksComponents->BoundMaterial, TracksComponents->ColorParameterName, FEntityComponentFilter());
+		}
+		if (bHasVectors)
+		{
+			VectorParameterTracker.Update(Linker, TracksComponents->BoundMaterial, TracksComponents->VectorParameterName, FEntityComponentFilter());
+		}
+		if (bHasColors || bHasVectors)
+		{
+			FOverlappingMaterialParameterHandler Handler(this);
+			Handler.DefaultComponentMask.SetAll({ BuiltInComponents->DoubleResult[0], BuiltInComponents->DoubleResult[1], BuiltInComponents->DoubleResult[2], BuiltInComponents->DoubleResult[3] });
+			VectorParameterTracker.ProcessInvalidatedOutputs(Linker, Handler);
+		}
 	}
 }
 
