Sequencer: Fix incorrect property stats updating when upgrading floats to doubles

The "old property index" was only properly set when we encounter a property for the first time. The next time, it's already set to the new index, so we don't actually properly update the property stats.

This change keeps the old property index in the property info struct to solve this. And then it also streamlines and cleans up things a bit.

diff --git a/Engine/Source/Runtime/MovieSceneTracks/Private/Systems/MovieScenePropertyInstantiator.cpp b/Engine/Source/Runtime/MovieSceneTracks/Private/Systems/MovieScenePropertyInstantiator.cpp
--- a/Engine/Source/Runtime/MovieSceneTracks/Private/Systems/MovieScenePropertyInstantiator.cpp
+++ b/Engine/Source/Runtime/MovieSceneTracks/Private/Systems/MovieScenePropertyInstantiator.cpp
@@ -231,10 +231,9 @@ void UMovieScenePropertyInstantiatorSystem::UpgradeFloatToDoubleProperties(const
 		}
 
 		FObjectPropertyInfo& PropertyInfo = ResolvedProperties[PropertyIndex];
-		const int32 OldPropertyDefinitionIndex = PropertyInfo.PropertyDefinitionIndex;
 
 		// The first time we encounter a specific property, we need to figure out if it needs type conversion or not.
-		if (!PropertyInfo.bRequiresTypeConversion.IsSet())
+		if (!PropertyInfo.ConvertedFromPropertyDefinitionIndex.IsSet())
 		{
 			// Don't do anything if the property isn't of the kind of type we need to care about. Right now, we only
 			// support dealing with float->double and FVectorXf->FVectorXd.
@@ -242,7 +241,7 @@ void UMovieScenePropertyInstantiatorSystem::UpgradeFloatToDoubleProperties(const
 			if (PropertyDefinition.PropertyType != TrackComponents->Float.PropertyTag &&
 				PropertyDefinition.PropertyType != TrackComponents->FloatVector.PropertyTag)
 			{
-				PropertyInfo.bRequiresTypeConversion = false;
+				PropertyInfo.ConvertedFromPropertyDefinitionIndex = INDEX_NONE;
 				continue;
 			}
 
@@ -253,15 +252,16 @@ void UMovieScenePropertyInstantiatorSystem::UpgradeFloatToDoubleProperties(const
 			}
 
 			// Patch the resolved property info to point to the double-precision property definition.
+			PropertyInfo.ConvertedFromPropertyDefinitionIndex = INDEX_NONE;
 			if (PropertyDefinition.PropertyType == TrackComponents->Float.PropertyTag)
 			{
 				const bool bIsDouble = BoundProperty->IsA<FDoubleProperty>();
-				PropertyInfo.bRequiresTypeConversion = bIsDouble;
 				if (bIsDouble)
 				{
 					const int32 DoublePropertyDefinitionIndex = Properties.IndexOfByPredicate(
 						[TrackComponents](const FPropertyDefinition& Item) { return Item.PropertyType == TrackComponents->Double.PropertyTag; });
 					ensure(DoublePropertyDefinitionIndex != INDEX_NONE);
+					PropertyInfo.ConvertedFromPropertyDefinitionIndex = PropertyInfo.PropertyDefinitionIndex;
 					PropertyInfo.PropertyDefinitionIndex = DoublePropertyDefinitionIndex;
 				}
 			}
@@ -279,12 +279,12 @@ void UMovieScenePropertyInstantiatorSystem::UpgradeFloatToDoubleProperties(const
 						BoundStructProperty == TBaseStructure<FVector4>::Get() ||
 						BoundStructProperty == TVariantStructure<FVector4d>::Get()
 						));
-				PropertyInfo.bRequiresTypeConversion = bIsDouble;
 				if (bIsDouble)
 				{
 					const int32 DoublePropertyDefinitionIndex = Properties.IndexOfByPredicate(
 						[TrackComponents](const FPropertyDefinition& Item) { return Item.PropertyType == TrackComponents->DoubleVector.PropertyTag; });
 					ensure(DoublePropertyDefinitionIndex != INDEX_NONE);
+					PropertyInfo.ConvertedFromPropertyDefinitionIndex = PropertyInfo.PropertyDefinitionIndex;
 					PropertyInfo.PropertyDefinitionIndex = DoublePropertyDefinitionIndex;
 				}
 			}
@@ -294,49 +294,42 @@ void UMovieScenePropertyInstantiatorSystem::UpgradeFloatToDoubleProperties(const
 			}
 		}
 
-		if (!PropertyInfo.bRequiresTypeConversion.Get(false))
+		const int32 OldPropertyDefinitionIndex = PropertyInfo.ConvertedFromPropertyDefinitionIndex.Get(INDEX_NONE);
+		if (OldPropertyDefinitionIndex == INDEX_NONE)
 		{
 			continue;
 		}
 
 		// Now we need to patch the contributors so that they have double-precision components.
 		// We only need to do it for the *new* contributors discovered this frame.
+		FComponentTypeID OldPropertyTag, NewPropertyTag;
 		const FPropertyDefinition& PropertyDefinition = Properties[PropertyInfo.PropertyDefinitionIndex];
 		if (PropertyDefinition.PropertyType == TrackComponents->Double.PropertyTag)
 		{
-			// Swap out the property tag and result components, and add the result widening tag for easily
-			// finding these entities later.
-			for (auto ContribIt = NewContributors.CreateKeyIterator(PropertyIndex); ContribIt; ++ContribIt)
-			{
-				const FMovieSceneEntityID CurID(ContribIt.Value());
-				FComponentMask EntityType = Linker->EntityManager.GetEntityType(CurID);
-				EntityType.Remove(TrackComponents->Float.PropertyTag);
-				EntityType.Set(TrackComponents->Double.PropertyTag);
-				Linker->EntityManager.ChangeEntityType(CurID, EntityType);
-			}
-
-			--PropertyStats[OldPropertyDefinitionIndex].NumProperties;
-			++PropertyStats[PropertyInfo.PropertyDefinitionIndex].NumProperties;
+			OldPropertyTag = TrackComponents->Float.PropertyTag;
+			NewPropertyTag = TrackComponents->Double.PropertyTag;
 		}
 		else if (PropertyDefinition.PropertyType == TrackComponents->DoubleVector.PropertyTag)
 		{
-			// Swap out the property tag and result components for the composites this contributor has.
+			OldPropertyTag = TrackComponents->FloatVector.PropertyTag;
+			NewPropertyTag = TrackComponents->DoubleVector.PropertyTag;
+		}
+
+		if (ensure(OldPropertyTag && NewPropertyTag))
+		{
+			// Swap out the property tag, update contributor info and stats.
 			for (auto ContribIt = NewContributors.CreateKeyIterator(PropertyIndex); ContribIt; ++ContribIt)
 			{
 				const FMovieSceneEntityID CurID(ContribIt.Value());
 				FComponentMask EntityType = Linker->EntityManager.GetEntityType(CurID);
-				EntityType.Remove(TrackComponents->FloatVector.PropertyTag);
-				EntityType.Set(TrackComponents->DoubleVector.PropertyTag);
+				EntityType.Remove(OldPropertyTag);
+				EntityType.Set(NewPropertyTag);
 				Linker->EntityManager.ChangeEntityType(CurID, EntityType);
 			}
 
 			--PropertyStats[OldPropertyDefinitionIndex].NumProperties;
 			++PropertyStats[PropertyInfo.PropertyDefinitionIndex].NumProperties;
 		}
-		else
-		{
-			check(false);
-		}
 	}
 }
 
diff --git a/Engine/Source/Runtime/MovieSceneTracks/Public/Systems/MovieScenePropertyInstantiator.h b/Engine/Source/Runtime/MovieSceneTracks/Public/Systems/MovieScenePropertyInstantiator.h
--- a/Engine/Source/Runtime/MovieSceneTracks/Public/Systems/MovieScenePropertyInstantiator.h
+++ b/Engine/Source/Runtime/MovieSceneTracks/Public/Systems/MovieScenePropertyInstantiator.h
@@ -131,8 +131,8 @@ private:
 		uint16 BlendChannel;
 		/** The index of this property within FPropertyRegistry::GetProperties. */
 		int32 PropertyDefinitionIndex;
-		/** true if this property has already been inspected for float-to-double conversion */
-		TOptional<bool> bRequiresTypeConversion;
+		/** Index of a float-based property if this property has been set for float-to-double conversion */
+		TOptional<int32> ConvertedFromPropertyDefinitionIndex;
 		/** true if any of the contributors to this property need restore state. */
 		bool bWantsRestoreState;
 	};
