Sequencer: only update property stats the first time a property is encountered for float->double conversion

diff --git a/Engine/Source/Runtime/MovieSceneTracks/Private/Systems/MovieScenePropertyInstantiator.cpp b/Engine/Source/Runtime/MovieSceneTracks/Private/Systems/MovieScenePropertyInstantiator.cpp
--- a/Engine/Source/Runtime/MovieSceneTracks/Private/Systems/MovieScenePropertyInstantiator.cpp
+++ b/Engine/Source/Runtime/MovieSceneTracks/Private/Systems/MovieScenePropertyInstantiator.cpp
@@ -231,9 +231,10 @@ void UMovieScenePropertyInstantiatorSystem::UpgradeFloatToDoubleProperties(const
 		}
 
 		FObjectPropertyInfo& PropertyInfo = ResolvedProperties[PropertyIndex];
+		const bool bHadConversionInfo = PropertyInfo.ConvertedFromPropertyDefinitionIndex.IsSet();
 
 		// The first time we encounter a specific property, we need to figure out if it needs type conversion or not.
-		if (!PropertyInfo.ConvertedFromPropertyDefinitionIndex.IsSet())
+		if (!bHadConversionInfo)
 		{
 			// Don't do anything if the property isn't of the kind of type we need to care about. Right now, we only
 			// support dealing with float->double and FVectorXf->FVectorXd.
@@ -317,7 +318,7 @@ void UMovieScenePropertyInstantiatorSystem::UpgradeFloatToDoubleProperties(const
 
 		if (ensure(OldPropertyTag && NewPropertyTag))
 		{
-			// Swap out the property tag, update contributor info and stats.
+			// Swap out the property tag
 			for (auto ContribIt = NewContributors.CreateKeyIterator(PropertyIndex); ContribIt; ++ContribIt)
 			{
 				const FMovieSceneEntityID CurID(ContribIt.Value());
@@ -327,8 +328,12 @@ void UMovieScenePropertyInstantiatorSystem::UpgradeFloatToDoubleProperties(const
 				Linker->EntityManager.ChangeEntityType(CurID, EntityType);
 			}
 
-			--PropertyStats[OldPropertyDefinitionIndex].NumProperties;
-			++PropertyStats[PropertyInfo.PropertyDefinitionIndex].NumProperties;
+			// Update contributor info and stats. This is only done the first time this property is encountered.
+			if (!bHadConversionInfo)
+			{
+				--PropertyStats[OldPropertyDefinitionIndex].NumProperties;
+				++PropertyStats[PropertyInfo.PropertyDefinitionIndex].NumProperties;
+			}
 		}
 	}
 }
