Sequencer: phase two of overridable channels

- The channel override registry now handles channels by mapping them to a name (such as TranslationX or RotationY).
- Introduce some infrastructure for figuring out what channel types are compatible with what other channel types.
- Pass some context when a custom channel container creates ECS entities, so that it knows what channel it's overriding.
- Channel groups and other sequencer outliner view-models now have menu entries for managing channel overrides
- Perlin noise channels use the new evaluation overloads that provide the parent section, so they can find the tick resolution of the owning MovieScene and figure out seconds from ticks.
- Perlin noise channels now have detail view customizations.
- Add channel interface classes for perlin noise channels.
- Add overridable channel support to transform and float property sections.
- Rebuild key editor widget when one of its row's channels have changed, which can happen overrides are added/removed/changed.

diff --git a/Engine/Source/Editor/MovieSceneTools/Private/Channels/PerlinNoiseChannelDetailsCustomization.cpp b/Engine/Source/Editor/MovieSceneTools/Private/Channels/PerlinNoiseChannelDetailsCustomization.cpp
new file mode 100644
--- /dev/null
+++ b/Engine/Source/Editor/MovieSceneTools/Private/Channels/PerlinNoiseChannelDetailsCustomization.cpp
@@ -0,0 +1,47 @@
+// Copyright Epic Games, Inc. All Rights Reserved.
+
+#include "Channels/PerlinNoiseChannelDetailsCustomization.h"
+#include "DetailLayoutBuilder.h"
+
+TSharedRef<IDetailCustomization> FMovieSceneFloatPerlinNoiseChannelDetailsCustomization::MakeInstance()
+{
+	return MakeShareable(new FMovieSceneFloatPerlinNoiseChannelDetailsCustomization());
+}
+
+FName FMovieSceneFloatPerlinNoiseChannelDetailsCustomization::GetPerlinNoiseChannelParamsPropertyName()
+{
+	static const FString ChannelProperty = GET_MEMBER_NAME_STRING_CHECKED(UMovieSceneFloatPerlinNoiseChannelContainer, PerlinNoiseChannel);
+	static const FString ParamsProperty = GET_MEMBER_NAME_STRING_CHECKED(FMovieSceneFloatPerlinNoiseChannel, PerlinNoiseParams);
+	static FName FullProperty(ChannelProperty + TEXT(".") + ParamsProperty);
+	return FullProperty;
+}
+
+void FMovieSceneFloatPerlinNoiseChannelDetailsCustomization::CustomizeDetails(IDetailLayoutBuilder& DetailBuilder)
+{
+	const FName PropertyName = GetPerlinNoiseChannelParamsPropertyName();
+	TSharedRef<IPropertyHandle> ParamsProperty = DetailBuilder.GetProperty(PropertyName, UMovieSceneFloatPerlinNoiseChannelContainer::StaticClass());
+	IDetailPropertyRow& Row = DetailBuilder.AddPropertyToCategory(ParamsProperty);	
+	DetailBuilder.HideProperty(GET_MEMBER_NAME_CHECKED(UMovieSceneFloatPerlinNoiseChannelContainer, PerlinNoiseChannel));
+}
+
+TSharedRef<IDetailCustomization> FMovieSceneDoublePerlinNoiseChannelDetailsCustomization::MakeInstance()
+{
+	return MakeShareable(new FMovieSceneDoublePerlinNoiseChannelDetailsCustomization());
+}
+
+FName FMovieSceneDoublePerlinNoiseChannelDetailsCustomization::GetPerlinNoiseChannelParamsPropertyName()
+{
+	static const FString ChannelProperty = GET_MEMBER_NAME_STRING_CHECKED(UMovieSceneDoublePerlinNoiseChannelContainer, PerlinNoiseChannel);
+	static const FString ParamsProperty = GET_MEMBER_NAME_STRING_CHECKED(FMovieSceneDoublePerlinNoiseChannel, PerlinNoiseParams);
+	static FName FullProperty(ChannelProperty + TEXT(".") + ParamsProperty);
+	return FullProperty;
+}
+
+void FMovieSceneDoublePerlinNoiseChannelDetailsCustomization::CustomizeDetails(IDetailLayoutBuilder& DetailBuilder)
+{
+	const FName PropertyName = GetPerlinNoiseChannelParamsPropertyName();
+	TSharedRef<IPropertyHandle> ParamsProperty = DetailBuilder.GetProperty(PropertyName, UMovieSceneDoublePerlinNoiseChannelContainer::StaticClass());
+	IDetailPropertyRow& Row = DetailBuilder.AddPropertyToCategory(ParamsProperty);	
+	DetailBuilder.HideProperty(GET_MEMBER_NAME_CHECKED(UMovieSceneDoublePerlinNoiseChannelContainer, PerlinNoiseChannel));
+}
+
diff --git a/Engine/Source/Editor/MovieSceneTools/Private/MovieSceneToolsModule.cpp b/Engine/Source/Editor/MovieSceneTools/Private/MovieSceneToolsModule.cpp
--- a/Engine/Source/Editor/MovieSceneTools/Private/MovieSceneToolsModule.cpp
+++ b/Engine/Source/Editor/MovieSceneTools/Private/MovieSceneToolsModule.cpp
@@ -54,6 +54,8 @@
 #include "TrackEditors/CameraShakeSourceShakeTrackEditor.h"
 #include "TrackEditors/CVarTrackEditor.h"
 
+#include "Channels/PerlinNoiseChannelInterface.h"
+
 #include "MVVM/ViewModels/CameraCutTrackModel.h"
 #include "MVVM/ViewModels/CinematicShotTrackModel.h"
 
@@ -77,6 +79,8 @@
 #include "Channels/MovieSceneObjectPathChannel.h"
 #include "Channels/MovieSceneCameraShakeSourceTriggerChannel.h"
 #include "Channels/MovieSceneStringChannel.h"
+#include "Channels/MovieSceneFloatPerlinNoiseChannel.h"
+#include "Channels/MovieSceneDoublePerlinNoiseChannel.h"
 #include "Channels/EventChannelCurveModel.h"
 #include "Channels/SCurveEditorEventChannelView.h"
 #include "Channels/MovieSceneAudioTriggerChannel.h"
@@ -86,6 +90,10 @@
 #include "Channels/SCurveEditorKeyBarView.h"
 #include "Sections/MovieSceneEventSection.h"
 
+#include "Channels/MovieSceneDoublePerlinNoiseChannelContainer.h"
+#include "Channels/MovieSceneFloatPerlinNoiseChannelContainer.h"
+#include "Channels/PerlinNoiseChannelDetailsCustomization.h"
+
 #include "MovieSceneEventUtils.h"
 
 #include "EntitySystem/MovieSceneEntityManager.h"
@@ -176,6 +184,8 @@ void FMovieSceneToolsModule::StartupModule()
 		FPropertyEditorModule& PropertyModule = FModuleManager::LoadModuleChecked<FPropertyEditorModule>("PropertyEditor");
 		PropertyModule.RegisterCustomClassLayout("MovieSceneToolsProjectSettings", FOnGetDetailCustomizationInstance::CreateStatic(&FMovieSceneToolsProjectSettingsCustomization::MakeInstance));
 		PropertyModule.RegisterCustomClassLayout("MovieSceneBuiltInEasingFunction", FOnGetDetailCustomizationInstance::CreateLambda(&MakeShared<FMovieSceneBuiltInEasingFunctionCustomization>));
+		PropertyModule.RegisterCustomClassLayout("MovieSceneFloatPerlinNoiseChannelContainer", FOnGetDetailCustomizationInstance::CreateStatic(&FMovieSceneFloatPerlinNoiseChannelDetailsCustomization::MakeInstance));
+		PropertyModule.RegisterCustomClassLayout("MovieSceneDoublePerlinNoiseChannelContainer", FOnGetDetailCustomizationInstance::CreateStatic(&FMovieSceneDoublePerlinNoiseChannelDetailsCustomization::MakeInstance));
 		PropertyModule.RegisterCustomPropertyTypeLayout("MovieSceneObjectBindingID", FOnGetPropertyTypeCustomizationInstance::CreateLambda(&MakeShared<FMovieSceneObjectBindingIDCustomization>));
 		PropertyModule.RegisterCustomPropertyTypeLayout("MovieSceneEvent", FOnGetPropertyTypeCustomizationInstance::CreateStatic(&FMovieSceneEventCustomization::MakeInstance));
 		PropertyModule.RegisterCustomPropertyTypeLayout("MovieSceneCVarOverrides", FOnGetPropertyTypeCustomizationInstance::CreateStatic(&UE::MovieScene::FCVarOverridesPropertyTypeCustomization::MakeInstance));
@@ -199,6 +209,9 @@ void FMovieSceneToolsModule::StartupModule()
 
 		SequencerModule.RegisterChannelInterface<FMovieSceneConstraintChannel>();
 
+		SequencerModule.RegisterChannelInterface<FMovieSceneFloatPerlinNoiseChannel>(MakeUnique<TPerlinNoiseChannelInterface<UMovieSceneFloatPerlinNoiseChannelContainer>>());
+		SequencerModule.RegisterChannelInterface<FMovieSceneDoublePerlinNoiseChannel>(MakeUnique<TPerlinNoiseChannelInterface<UMovieSceneDoublePerlinNoiseChannelContainer>>());
+
 		ICurveEditorModule& CurveEditorModule = FModuleManager::LoadModuleChecked<ICurveEditorModule>("CurveEditor");
 
 		FEventChannelCurveModel::EventView = CurveEditorModule.RegisterView(FOnCreateCurveEditorView::CreateStatic(
diff --git a/Engine/Source/Editor/MovieSceneTools/Public/Channels/PerlinNoiseChannelDetailsCustomization.h b/Engine/Source/Editor/MovieSceneTools/Public/Channels/PerlinNoiseChannelDetailsCustomization.h
new file mode 100644
--- /dev/null
+++ b/Engine/Source/Editor/MovieSceneTools/Public/Channels/PerlinNoiseChannelDetailsCustomization.h
@@ -0,0 +1,35 @@
+// Copyright Epic Games, Inc. All Rights Reserved.
+
+#pragma once
+
+#include "CoreMinimal.h"
+#include "IDetailCustomization.h"
+#include "IDetailChildrenBuilder.h"
+#include "DetailLayoutBuilder.h"
+#include "PropertyCustomizationHelpers.h"
+#include "Channels/MovieSceneDoublePerlinNoiseChannelContainer.h"
+#include "Channels/MovieSceneFloatPerlinNoiseChannelContainer.h"
+#include "Styling/SlateTypes.h"
+
+class IDetailLayoutBuilder;
+
+class FMovieSceneFloatPerlinNoiseChannelDetailsCustomization : public IDetailCustomization
+{
+public:
+	static TSharedRef<IDetailCustomization> MakeInstance();
+	static FName GetPerlinNoiseChannelParamsPropertyName();
+
+	/** IDetailCustomization interface */
+	virtual void CustomizeDetails(IDetailLayoutBuilder& DetailBuilder) override;
+};
+
+class FMovieSceneDoublePerlinNoiseChannelDetailsCustomization : public IDetailCustomization
+{
+public:
+	static TSharedRef<IDetailCustomization> MakeInstance();
+	static FName GetPerlinNoiseChannelParamsPropertyName();
+
+	/** IDetailCustomization interface */
+	virtual void CustomizeDetails(IDetailLayoutBuilder& DetailBuilder) override;
+};
+
diff --git a/Engine/Source/Editor/MovieSceneTools/Public/Channels/PerlinNoiseChannelInterface.h b/Engine/Source/Editor/MovieSceneTools/Public/Channels/PerlinNoiseChannelInterface.h
new file mode 100644
--- /dev/null
+++ b/Engine/Source/Editor/MovieSceneTools/Public/Channels/PerlinNoiseChannelInterface.h
@@ -0,0 +1,230 @@
+// Copyright Epic Games, Inc. All Rights Reserved.
+
+#pragma once
+
+#include "ISequencerChannelInterface.h"
+#include "Channels/IMovieSceneChannelOverrideProvider.h"
+#include "Channels/MovieSceneSectionChannelOverrideRegistry.h"
+#include "CurveKeyEditors/SequencerKeyEditor.h"
+#include "Framework/MultiBox/MultiBoxBuilder.h"
+#include "Framework/MultiBox/MultiBoxExtender.h"
+#include "Modules/ModuleManager.h"
+#include "MovieSceneSection.h"
+#include "PropertyEditorModule.h"
+#include "SequencerSectionPainter.h"
+#include "TimeToPixel.h"
+#include "ToolMenus.h"
+#include "Widgets/Input/NumericTypeInterface.h"
+#include "Widgets/Text/STextBlock.h"
+
+template<typename ChannelType, typename NumericType>
+class SNumericTextBlockKeyEditor : public SCompoundWidget
+{
+	using FKeyEditor = TSequencerKeyEditor<ChannelType, NumericType>;
+	
+	FKeyEditor KeyEditor;
+	TSharedPtr<INumericTypeInterface<NumericType>> Interface;
+
+	SLATE_BEGIN_ARGS(SNumericTextBlockKeyEditor){}
+	SLATE_END_ARGS()
+
+	void Construct(const FArguments& InArgs, const FKeyEditor& InKeyEditor)
+	{
+		KeyEditor = InKeyEditor;
+		Interface = MakeShareable(new TDefaultNumericTypeInterface<NumericType>);
+
+		ChildSlot
+		[
+			SNew(STextBlock)
+			.Text(this, &SNumericTextBlockKeyEditor<ChannelType, NumericType>::GetValueAsText)
+		];
+	}
+
+	FText GetValueAsText() const
+	{
+		const FString ValueString = Interface->ToString(KeyEditor.GetCurrentValue());
+		return FText::FromString(ValueString);
+	}
+};
+
+template<typename ChannelContainerType>
+struct TPerlinNoiseChannelSectionMenuExtension : TSharedFromThis<TPerlinNoiseChannelSectionMenuExtension<ChannelContainerType>>
+{
+	using ChannelType = typename ChannelContainerType::ChannelType;
+
+	TPerlinNoiseChannelSectionMenuExtension(TArrayView<UMovieSceneSection* const> InSections)
+		: Sections(InSections)
+	{}
+
+	void ExtendMenu(FMenuBuilder& MenuBuilder)
+	{
+		TSharedRef<TPerlinNoiseChannelSectionMenuExtension<ChannelContainerType>> SharedThis = this->AsShared();
+
+		MenuBuilder.AddSubMenu(
+			NSLOCTEXT("PerlinNoiseChannelInterface", "PerlinNoiseChannelsMenu", "Perlin Noise Channels"),
+			NSLOCTEXT("PerlinNoiseChannelInterface", "PerlinNoiseChannelsMenuToolTip", "Edit parameters for Perlin Noise channels"),
+			FNewMenuDelegate::CreateLambda([SharedThis](FMenuBuilder& MenuBuilder) { SharedThis->BuildChannelOverrideParametersMenu(MenuBuilder); })
+		);
+	}
+
+private:
+
+	void BuildChannelOverrideParametersMenu(FMenuBuilder& MenuBuilder)
+	{
+		FPropertyEditorModule& EditModule = FModuleManager::Get().GetModuleChecked<FPropertyEditorModule>("PropertyEditor");
+
+		FDetailsViewArgs DetailsViewArgs;
+		DetailsViewArgs.bAllowSearch = false;
+		DetailsViewArgs.NameAreaSettings = FDetailsViewArgs::HideNameArea;
+		DetailsViewArgs.bHideSelectionTip = true;
+		DetailsViewArgs.bShowOptions = false;
+		DetailsViewArgs.bShowScrollBar = false;
+
+		TSharedRef<IDetailsView> DetailsView = EditModule.CreateDetailView(DetailsViewArgs);
+
+		TArray<UObject*> Objects;
+		for (TWeakObjectPtr<UMovieSceneSection> WeakSection : Sections)
+		{
+			UMovieSceneSection* Section = WeakSection.Get();
+			if (!Section)
+			{
+				continue;
+			}
+
+			IMovieSceneChannelOverrideProvider* OverrideProvider = Cast<IMovieSceneChannelOverrideProvider>(Section);
+			if (!OverrideProvider)
+			{
+				continue;
+			}
+			
+			UMovieSceneSectionChannelOverrideRegistry* OverrideRegistry = OverrideProvider->GetChannelOverrideRegistry(false);
+			if (!OverrideRegistry)
+			{
+				continue;
+			}
+	
+			TArray<ChannelContainerType*> ChannelContainers;
+			OverrideRegistry->GetChannels<ChannelContainerType>(ChannelContainers);
+			Objects.Append(ChannelContainers);
+		}
+
+		DetailsView->SetObjects(Objects, true);
+
+		MenuBuilder.AddWidget(DetailsView, FText(), true, false);
+	}
+
+
+private:
+
+	TArray<TWeakObjectPtr<UMovieSceneSection>> Sections;
+};
+
+template<typename ChannelContainerType>
+struct TPerlinNoiseChannelInterface : ISequencerChannelInterface
+{
+	using ChannelType = typename ChannelContainerType::ChannelType;
+
+	virtual FKeyHandle AddOrUpdateKey_Raw(FMovieSceneChannel* Channel, UMovieSceneSection* SectionToKey, const void* ExtendedEditorData, FFrameNumber InTime, ISequencer& InSequencer, const FGuid& ObjectBindingID, FTrackInstancePropertyBindings* PropertyBindings) const override
+	{
+		return FKeyHandle::Invalid();
+	}
+
+	virtual void CopyKeys_Raw(FMovieSceneChannel* Channel, const UMovieSceneSection* Section, FName KeyAreaName, FMovieSceneClipboardBuilder& ClipboardBuilder, TArrayView<const FKeyHandle> KeyMask) const override
+	{
+	}
+
+	virtual void PasteKeys_Raw(FMovieSceneChannel* Channel, UMovieSceneSection* Section, const FMovieSceneClipboardKeyTrack& KeyTrack, const FMovieSceneClipboardEnvironment& SrcEnvironment, const FSequencerPasteEnvironment& DstEnvironment, TArray<FKeyHandle>& OutPastedKeys) const override
+	{
+	}
+
+	virtual TSharedPtr<FStructOnScope> GetKeyStruct_Raw(FMovieSceneChannelHandle Channel, FKeyHandle KeyHandle) const override
+	{
+		return nullptr;
+	}
+
+	virtual bool CanCreateKeyEditor_Raw(const FMovieSceneChannel* Channel) const override
+	{
+		return true;
+	}
+
+	virtual TSharedRef<SWidget> CreateKeyEditor_Raw(const FMovieSceneChannelHandle& Channel, UMovieSceneSection* Section, const FGuid& InObjectBindingID, TWeakPtr<FTrackInstancePropertyBindings> PropertyBindings, TWeakPtr<ISequencer> Sequencer) const override
+	{
+		const TMovieSceneExternalValue<typename ChannelType::CurveValueType>* ExternalValue = Channel.Cast<ChannelType>().GetExtendedEditorData();
+		if (!ExternalValue)
+		{
+			return SNullWidget::NullWidget;
+		}
+
+		TSequencerKeyEditor<ChannelType, typename ChannelType::CurveValueType> KeyEditor(
+			InObjectBindingID, Channel.Cast<ChannelType>(),
+			Section, Sequencer, PropertyBindings, ExternalValue->OnGetExternalValue
+			);
+
+		using KeyEditorType = SNumericTextBlockKeyEditor<ChannelType, typename ChannelType::CurveValueType>;
+		return SNew(KeyEditorType, KeyEditor);
+	}
+
+	virtual void ExtendKeyMenu_Raw(FMenuBuilder& MenuBuilder, TSharedPtr<FExtender> MenuExtender, TArrayView<const FExtendKeyMenuParams> Parameters, TWeakPtr<ISequencer> InSequencer) const override
+	{
+	}
+
+	virtual void ExtendSectionMenu_Raw(FMenuBuilder& MenuBuilder, TSharedPtr<FExtender> MenuExtender, TArrayView<const FMovieSceneChannelHandle> Channels, TArrayView<UMovieSceneSection* const> Sections, TWeakPtr<ISequencer> InSequencer) const override
+	{
+		using MenuExtenderType = TPerlinNoiseChannelSectionMenuExtension<ChannelContainerType>;
+		TSharedRef<MenuExtenderType> Extension = MakeShared<MenuExtenderType>(Sections);
+
+		MenuExtender->AddMenuExtension("SequencerChannels", EExtensionHook::First, nullptr, FMenuExtensionDelegate::CreateLambda([Extension](FMenuBuilder& MenuBuilder) { Extension->ExtendMenu(MenuBuilder); }));
+	}
+
+	virtual void DrawKeys_Raw(FMovieSceneChannel* Channel, TArrayView<const FKeyHandle> InKeyHandles, const UMovieSceneSection* InOwner, TArrayView<FKeyDrawParams> OutKeyDrawParams) const override
+	{
+	}
+
+	virtual bool ShouldShowCurve_Raw(const FMovieSceneChannel* Channel, UMovieSceneSection* InSection) const override
+	{
+		return true;
+	}
+
+	virtual bool SupportsCurveEditorModels_Raw(const FMovieSceneChannelHandle& InChannel) const override
+	{
+		return false;
+	}
+
+	virtual TUniquePtr<FCurveModel> CreateCurveEditorModel_Raw(const FMovieSceneChannelHandle& Channel, UMovieSceneSection* OwningSection, TSharedRef<ISequencer> InSequencer) const override
+	{
+		return nullptr;
+	}
+
+	virtual void DrawExtra_Raw(FMovieSceneChannel* InChannel, const UMovieSceneSection* InOwner, const FGeometry& InKeyGeometry, FSequencerSectionPainter& Painter) const override
+	{
+		using namespace UE::Sequencer;
+
+		FColor FillColor(0, 0, 0, 51);
+		FTimeToPixel TimeToPixelConverter = Painter.GetTimeConverter();
+		ChannelType* TypedChannel = static_cast<ChannelType*>(InChannel);
+
+		TArray<FVector2D> CurvePoints;
+		CurvePoints.Reserve(InKeyGeometry.Size.X / 2.0);
+
+		const double Amplitude = TypedChannel->GetParams().Amplitude;
+		const double YOffset = InKeyGeometry.Size.Y / 2.0;
+		const double YScale = (Amplitude != 0) ? (InKeyGeometry.Size.Y / Amplitude / 2.0) : 1;
+
+		for (double X = 0; X < InKeyGeometry.Size.X; X += 2.f)
+		{
+			double Seconds = TimeToPixelConverter.PixelToSeconds(X);
+			double Y = (double)TypedChannel->Evaluate(Seconds);
+			CurvePoints.Add(FVector2D(X, (Y * YScale) + YOffset));
+		}
+
+		FSlateDrawElement::MakeLines(
+				Painter.DrawElements,
+				Painter.LayerId++,
+				InKeyGeometry.ToPaintGeometry(),
+				CurvePoints,
+				ESlateDrawEffect::PreMultipliedAlpha,
+				FillColor,
+				true);
+	}
+};
+
diff --git a/Engine/Source/Editor/Sequencer/Private/MVVM/ViewModels/ChannelModel.cpp b/Engine/Source/Editor/Sequencer/Private/MVVM/ViewModels/ChannelModel.cpp
--- a/Engine/Source/Editor/Sequencer/Private/MVVM/ViewModels/ChannelModel.cpp
+++ b/Engine/Source/Editor/Sequencer/Private/MVVM/ViewModels/ChannelModel.cpp
@@ -6,8 +6,15 @@
 #include "MVVM/ViewModels/ViewModelIterators.h"
 #include "MVVM/ViewModels/SequencerModelUtils.h"
 #include "MVVM/ViewModels/SequencerEditorViewModel.h"
+#include "MVVM/ViewModels/OutlinerViewModel.h"
+#include "MVVM/Views/SOutlinerView.h"
 #include "MVVM/Views/SOutlinerItemViewBase.h"
 #include "MVVM/Views/SSequencerKeyNavigationButtons.h"
+#include "Channels/IMovieSceneChannelOverrideProvider.h"
+#include "Channels/MovieSceneChannel.h"
+#include "Channels/MovieSceneChannelOverrideContainer.h"
+#include "Channels/MovieSceneSectionChannelOverrideRegistry.h"
+#include "IDetailsView.h"
 #include "MVVM/Views/KeyRenderer.h"
 #include "MVVM/Views/SChannelView.h"
 #include "IKeyArea.h"
@@ -20,6 +27,8 @@
 #include "SSequencerSection.h"
 #include "Channels/MovieSceneChannel.h"
 #include "CurveModel.h"
+#include "Modules/ModuleManager.h"
+#include "PropertyEditorModule.h"
 #include "Styling/AppStyle.h"
 
 #define LOCTEXT_NAMESPACE "SequencerChannelModel"
@@ -223,7 +232,8 @@ FLinearColor FChannelModel::GetKeyBarColor() const
 }
 
 FChannelGroupModel::FChannelGroupModel(FName InChannelName, const FText& InDisplayText)
-	: ChannelName(InChannelName)
+	: ChannelsSerialNumber(0)
+	, ChannelName(InChannelName)
 	, DisplayText(InDisplayText)
 {
 }
@@ -252,6 +262,22 @@ void FChannelGroupModel::AddChannel(TWeakViewModelPtr<FChannelModel> InChannel)
 	if (!Channels.Contains(InChannel))
 	{
 		Channels.Add(InChannel);
+		++ChannelsSerialNumber;
+	}
+}
+
+uint32 FChannelGroupModel::GetChannelsSerialNumber() const
+{
+	const_cast<FChannelGroupModel*>(this)->CleanupChannels();
+	return ChannelsSerialNumber;
+}
+
+void FChannelGroupModel::CleanupChannels()
+{
+	const int32 NumRemoved = Channels.RemoveAll([](TWeakViewModelPtr<FChannelModel> Item) { return !Item.IsValid(); });
+	if (NumRemoved > 0)
+	{
+		++ChannelsSerialNumber;
 	}
 }
 
@@ -271,6 +297,11 @@ TSharedPtr<IKeyArea> FChannelGroupModel::GetKeyArea(const UMovieSceneSection* In
 	return Channel ? Channel->GetKeyArea() : nullptr;
 }
 
+TSharedPtr<FChannelModel> FChannelGroupModel::GetChannel(int32 Index) const
+{
+	return Channels[Index].Pin();
+}
+
 TSharedPtr<FChannelModel> FChannelGroupModel::GetChannel(TSharedPtr<FSectionModel> InOwnerSection) const
 {
 	return GetChannel(InOwnerSection->GetSection());
@@ -360,6 +391,263 @@ bool FChannelGroupModel::HasCurves() const
 	return false;
 }
 
+void FChannelGroupModel::BuildChannelOverrideMenu(FMenuBuilder& MenuBuilder)
+{
+	// Gather up information:
+	// - Candidates for overriding each of the channel types found in this group.
+	// - Already overriden channels.
+	int32 NumOverridableChannels = 0;
+	TMap<UClass*, TArray<UMovieSceneChannelOverrideContainer*>> ExistingOverrides;
+	TArray<UMovieSceneChannelOverrideContainer::FOverrideCandidates> AllCandidateOverrides;
+	for (int32 Index = 0; Index < GetChannels().Num(); ++Index)
+	{
+		TSharedPtr<FChannelModel> Channel = GetChannel(Index);
+		if (!Channel)
+		{
+			continue;
+		}
+
+		UMovieSceneSection* Section = Channel->GetSection();
+		if (!Section)
+		{
+			continue;
+		}
+
+		IMovieSceneChannelOverrideProvider* OverrideProvider = Cast<IMovieSceneChannelOverrideProvider>(Section);
+		if (!OverrideProvider)
+		{
+			continue;
+		}
+
+		++NumOverridableChannels;
+
+		UMovieSceneSectionChannelOverrideRegistry* OverrideRegistry = OverrideProvider->GetChannelOverrideRegistry(false);
+		UE::MovieScene::FChannelOverrideProviderTraitsHandle OverrideTraits = OverrideProvider->GetChannelOverrideProviderTraits();
+
+		UMovieSceneChannelOverrideContainer* OverrideContainer = OverrideRegistry ? 
+			OverrideRegistry->GetChannel(GetChannelName()) : nullptr;
+		if (OverrideContainer)
+		{
+			ExistingOverrides.FindOrAdd(OverrideContainer->GetClass()).Add(OverrideContainer);
+		}
+		
+		const FName DefaultChannelTypeName = OverrideTraits->GetDefaultChannelTypeName(GetChannelName());
+		if (!DefaultChannelTypeName.IsNone())
+		{
+			UMovieSceneChannelOverrideContainer::FOverrideCandidates OverrideCandidates;
+			UMovieSceneChannelOverrideContainer::GetOverrideCandidates(DefaultChannelTypeName, OverrideCandidates);
+
+			AllCandidateOverrides.Add(MoveTemp(OverrideCandidates));
+		}
+	}
+
+	// Show channel parameters.
+	if (ExistingOverrides.Num() > 0)
+	{
+		MenuBuilder.BeginSection("ChannelParameters", LOCTEXT("ChannelGroupParameterMenuSectionName", "Channel Parameters"));
+
+		for (const TTuple<UClass*, TArray<UMovieSceneChannelOverrideContainer*>>& ExistingOverride : ExistingOverrides)
+		{
+			const FText ChannelOverrideClassName = ExistingOverride.Key->GetDisplayNameText();
+			MenuBuilder.AddSubMenu(
+				FText::Format(LOCTEXT("OverrideChannelParametersMenuLabel", "{0} Parameters"), ChannelOverrideClassName),
+				FText::Format(
+					LOCTEXT("OverrideChannelParametersMenuTooltip", "Edit parameters for all channels of type {0}"),
+					ChannelOverrideClassName),
+				FNewMenuDelegate::CreateSP(this, &FChannelGroupOutlinerModel::BuildChannelOverrideParametersMenu, ExistingOverride.Value)
+			);
+		}
+
+		MenuBuilder.EndSection();
+	}
+
+	// Add menu entries for overriding channels, with only types that are common among all sections.
+	TMap<TSubclassOf<UMovieSceneChannelOverrideContainer>, int32> CandidateCounts;
+	for (const UMovieSceneChannelOverrideContainer::FOverrideCandidates& CandidateOverrides : AllCandidateOverrides)
+	{
+		for (TSubclassOf<UMovieSceneChannelOverrideContainer> CandidateOverride : CandidateOverrides)
+		{
+			++CandidateCounts.FindOrAdd(CandidateOverride);
+		}
+	}
+
+	UMovieSceneChannelOverrideContainer::FOverrideCandidates CommonCandidateOverrides;
+	for (const TTuple<TSubclassOf<UMovieSceneChannelOverrideContainer>, int32> CandidateCount : CandidateCounts)
+	{
+		if (CandidateCount.Value == NumOverridableChannels)
+		{
+			CommonCandidateOverrides.Add(CandidateCount.Key);
+		}
+	}
+
+	if (CommonCandidateOverrides.Num() > 0 || ExistingOverrides.Num() > 0)
+	{
+		MenuBuilder.BeginSection("ChannelOverrides", LOCTEXT("ChannelGroupOverrideMenuSectionName", "Channel Overrides"));
+
+		if (CommonCandidateOverrides.Num() == 1)
+		{
+			BuildChannelOverrideMenu(MenuBuilder, CommonCandidateOverrides);
+		}
+		else if (CommonCandidateOverrides.Num() > 1)
+		{
+			MenuBuilder.AddSubMenu(
+				LOCTEXT("OverrideChannelChoiceLabel", "Override with..."),
+				LOCTEXT("OverrideChannelChoiceTooltip", "Overrides all channels with the specified channel"),
+				FNewMenuDelegate::CreateSP(this, &FChannelGroupOutlinerModel::BuildChannelOverrideMenu, CommonCandidateOverrides));
+		}
+
+		if (ExistingOverrides.Num() > 0)
+		{
+			MenuBuilder.AddMenuEntry(
+				LOCTEXT("RemoveOverrideChannelLabel", "Remove override"),
+				LOCTEXT("RemoveOverrideChannelTooltip", "Removes the channel override and revert back to the default channel type"),
+				FSlateIcon(),
+				FUIAction(
+					FExecuteAction::CreateSP(this, &FChannelGroupOutlinerModel::RemoveChannelOverrides),
+					FCanExecuteAction()));
+		}
+
+		MenuBuilder.EndSection();
+	}
+}
+
+void FChannelGroupModel::BuildChannelOverrideMenu(FMenuBuilder& MenuBuilder, UMovieSceneChannelOverrideContainer::FOverrideCandidates OverrideCandidates)
+{
+	for (TSubclassOf<UMovieSceneChannelOverrideContainer> OverrideCandidate : OverrideCandidates)
+	{
+		MenuBuilder.AddMenuEntry(
+			FText::Format(LOCTEXT("OverrideChannelLabel", "Override with {0}"),
+				OverrideCandidate->GetDisplayNameText()),
+			FText::Format(LOCTEXT("OverrideChannelTooltip", "Overrides all channels with {0}"),
+				OverrideCandidate->GetDisplayNameText()),
+			FSlateIcon(),
+			FUIAction(
+				FExecuteAction::CreateSP(this, &FChannelGroupModel::OverrideChannels, OverrideCandidate),
+				FCanExecuteAction()));
+	}
+}
+
+void FChannelGroupModel::OverrideChannels(TSubclassOf<UMovieSceneChannelOverrideContainer> OverrideClass)
+{
+	const FScopedTransaction Transaction(LOCTEXT("OverrideChannels", "Override Channels"));
+
+	for (TWeakViewModelPtr<FChannelModel> WeakChannel : GetChannels())
+	{
+		TSharedPtr<FChannelModel> Channel = WeakChannel.Pin();
+		if (!Channel || !Channel->GetSection())
+		{
+			continue;
+		}
+
+		UMovieSceneSection* Section = Channel->GetSection();
+		IMovieSceneChannelOverrideProvider* OverrideProvider = Cast<IMovieSceneChannelOverrideProvider>(Section);
+		if (!OverrideProvider)
+		{
+			continue;
+		}
+
+		if (!Section->TryModify())
+		{
+			continue;
+		}
+
+		UMovieSceneSectionChannelOverrideRegistry* OverrideRegistry = OverrideProvider->GetChannelOverrideRegistry(true);
+		check(OverrideRegistry);
+
+		UMovieSceneChannelOverrideContainer* OverrideChannelContainer = NewObject<UMovieSceneChannelOverrideContainer>(Section, OverrideClass);
+		OverrideRegistry->AddChannel(GetChannelName(), OverrideChannelContainer);	
+		
+		OverrideProvider->OnChannelOverridesChanged();
+
+		++ChannelsSerialNumber;
+	}
+
+	TSharedPtr<FSequenceModel> SequenceModel = FindAncestorOfType<FSequenceModel>();
+	if (!SequenceModel)
+	{
+		return;
+	}
+
+	TSharedPtr<ISequencer> Sequencer = SequenceModel->GetSequencer();
+	if (!Sequencer)
+	{
+		return;
+	}
+
+	Sequencer->NotifyMovieSceneDataChanged(EMovieSceneDataChangeType::MovieSceneStructureItemsChanged);
+}
+
+void FChannelGroupModel::RemoveChannelOverrides()
+{
+	const FScopedTransaction Transaction(LOCTEXT("RemoveChannelOverrides", "Remove Override Channels"));
+
+	for (TWeakViewModelPtr<FChannelModel> WeakChannel : GetChannels())
+	{
+		TSharedPtr<FChannelModel> Channel = WeakChannel.Pin();
+		if (!Channel || !Channel->GetSection())
+		{
+			continue;
+		}
+
+		IMovieSceneChannelOverrideProvider* OverrideProvider = Cast<IMovieSceneChannelOverrideProvider>(Channel->GetSection());
+		if (!OverrideProvider)
+		{
+			continue;
+		}
+
+		UMovieSceneSectionChannelOverrideRegistry* OverrideRegistry = OverrideProvider->GetChannelOverrideRegistry(false);
+		if (!OverrideRegistry)
+		{
+			continue;
+		}
+
+		if (!Channel->GetSection()->TryModify())
+		{
+			continue;
+		}
+
+		OverrideRegistry->RemoveChannel(GetChannelName());
+
+		OverrideProvider->OnChannelOverridesChanged();
+
+		++ChannelsSerialNumber;
+	}
+
+	TSharedPtr<FSequenceModel> SequenceModel = FindAncestorOfType<FSequenceModel>();
+	if (!SequenceModel)
+	{
+		return;
+	}
+
+	TSharedPtr<ISequencer> Sequencer = SequenceModel->GetSequencer();
+	if (!Sequencer)
+	{
+		return;
+	}
+
+	Sequencer->NotifyMovieSceneDataChanged(EMovieSceneDataChangeType::MovieSceneStructureItemsChanged);
+}
+
+void FChannelGroupModel::BuildChannelOverrideParametersMenu(FMenuBuilder& MenuBuilder, TArray<UMovieSceneChannelOverrideContainer*> ChannelParameters)
+{
+	FPropertyEditorModule& EditModule = FModuleManager::Get().GetModuleChecked<FPropertyEditorModule>("PropertyEditor");
+
+	FDetailsViewArgs DetailsViewArgs;
+	DetailsViewArgs.bAllowSearch = false;
+	DetailsViewArgs.NameAreaSettings = FDetailsViewArgs::HideNameArea;
+	DetailsViewArgs.bHideSelectionTip = true;
+	DetailsViewArgs.bShowOptions = false;
+	DetailsViewArgs.bShowScrollBar = false;
+
+	TSharedRef<IDetailsView> DetailsView = EditModule.CreateDetailView(DetailsViewArgs);
+
+	TArray<UObject*> Objects;
+	Objects.Append(ChannelParameters);
+	DetailsView->SetObjects(Objects, true);
+
+	MenuBuilder.AddWidget(DetailsView, FText(), true, false);
+}
+
 FChannelGroupOutlinerModel::FChannelGroupOutlinerModel(FName InChannelName, const FText& InDisplayText)
 	: TOutlinerModelMixin<FChannelGroupModel>(InChannelName, InDisplayText)
 {
@@ -479,6 +767,13 @@ void FChannelGroupOutlinerModel::CreateCurveModels(TArray<TUniquePtr<FCurveModel
 	FChannelGroupModel::CreateCurveModels(OutCurveModels);
 }
 
+void FChannelGroupOutlinerModel::BuildContextMenu(FMenuBuilder& MenuBuilder)
+{
+	FOutlinerItemModelMixin::BuildContextMenu(MenuBuilder);
+
+	BuildChannelOverrideMenu(MenuBuilder);
+}
+
 } // namespace Sequencer
 } // namespace UE
 
diff --git a/Engine/Source/Editor/Sequencer/Private/MVVM/ViewModels/TrackModel.cpp b/Engine/Source/Editor/Sequencer/Private/MVVM/ViewModels/TrackModel.cpp
--- a/Engine/Source/Editor/Sequencer/Private/MVVM/ViewModels/TrackModel.cpp
+++ b/Engine/Source/Editor/Sequencer/Private/MVVM/ViewModels/TrackModel.cpp
@@ -322,6 +322,11 @@ void FTrackModel::ForceUpdate()
 			//       we fixup row indices
 		}
 	}
+
+	for (TSharedPtr<FChannelGroupModel> ChannelModel : GetDescendantsOfType<FChannelGroupModel>(false, EViewModelListType::Outliner))
+	{
+
+	}
 }
 
 FOutlinerSizing FTrackModel::GetOutlinerSizing() const
@@ -645,6 +650,12 @@ void FTrackModel::BuildContextMenu(FMenuBuilder& MenuBuilder)
 		);
 	}
 
+	TViewModelPtr<FChannelGroupModel> ChannelGroup = TopLevelChannelList.GetHead().ImplicitCast();
+	if (ChannelGroup)
+	{
+		ChannelGroup->BuildChannelOverrideMenu(MenuBuilder);
+	}
+
 	FOutlinerItemModel::BuildContextMenu(MenuBuilder);
 }
 
diff --git a/Engine/Source/Editor/Sequencer/Private/MVVM/ViewModels/TrackRowModel.cpp b/Engine/Source/Editor/Sequencer/Private/MVVM/ViewModels/TrackRowModel.cpp
--- a/Engine/Source/Editor/Sequencer/Private/MVVM/ViewModels/TrackRowModel.cpp
+++ b/Engine/Source/Editor/Sequencer/Private/MVVM/ViewModels/TrackRowModel.cpp
@@ -302,6 +302,14 @@ void FTrackRowModel::BuildContextMenu(FMenuBuilder& MenuBuilder)
 			})
 		);
 	}
+
+	TViewModelPtr<FChannelGroupModel> ChannelGroup = TopLevelChannelList.GetHead().ImplicitCast();
+	if (ChannelGroup)
+	{
+		ChannelGroup->BuildChannelOverrideMenu(MenuBuilder);
+	}
+
+	FOutlinerItemModel::BuildContextMenu(MenuBuilder);
 }
 
 bool FTrackRowModel::CanDelete(FText* OutErrorMessage) const
diff --git a/Engine/Source/Editor/Sequencer/Private/SKeyAreaEditorSwitcher.cpp b/Engine/Source/Editor/Sequencer/Private/SKeyAreaEditorSwitcher.cpp
--- a/Engine/Source/Editor/Sequencer/Private/SKeyAreaEditorSwitcher.cpp
+++ b/Engine/Source/Editor/Sequencer/Private/SKeyAreaEditorSwitcher.cpp
@@ -20,6 +20,7 @@ void SKeyAreaEditorSwitcher::Construct(const FArguments& InArgs, TSharedPtr<FCha
 {
 	WeakModel = InModel;
 	WeakSequencer = InWeakSequencer;
+	CachedChannelsSerialNumber = 0;
 }
 
 int32 SKeyAreaEditorSwitcher::GetWidgetIndex() const
@@ -109,15 +110,16 @@ void SKeyAreaEditorSwitcher::Tick( const FGeometry& AllottedGeometry, const doub
 	}
 	else
 	{
-		TArray<TSharedRef<IKeyArea>> NewCachedKeyAreas = Model->GetAllKeyAreas();
-		if (CachedKeyAreas != NewCachedKeyAreas)
+		const uint32 NewChannelsSerialNumber = Model->GetChannelsSerialNumber();
+		const bool bRebuild = (NewChannelsSerialNumber != CachedChannelsSerialNumber);
+		if (bRebuild)
 		{
-			Swap(CachedKeyAreas, NewCachedKeyAreas);
-
-			// Node is valid but now has a different set of key areas. Must rebuild the widgets.
+			CachedChannelsSerialNumber = NewChannelsSerialNumber;
+			CachedKeyAreas = Model->GetAllKeyAreas();
 			Rebuild();
 		}
 
+		// Figure out which widget is active based on which section the current time is at.
 		TArray<UMovieSceneSection*> AllSections;
 		for (const TSharedRef<IKeyArea>& KeyArea : CachedKeyAreas)
 		{
diff --git a/Engine/Source/Editor/Sequencer/Private/SKeyAreaEditorSwitcher.h b/Engine/Source/Editor/Sequencer/Private/SKeyAreaEditorSwitcher.h
--- a/Engine/Source/Editor/Sequencer/Private/SKeyAreaEditorSwitcher.h
+++ b/Engine/Source/Editor/Sequencer/Private/SKeyAreaEditorSwitcher.h
@@ -47,8 +47,10 @@ private:
 	TWeakPtr<ISequencer> WeakSequencer;
 	/** Key areas cached from the node */
 	TArray<TSharedRef<IKeyArea>> CachedKeyAreas;
+	/** Serial cached from the node */
+	uint32 CachedChannelsSerialNumber;
 };
 
 
 } // namespace Sequencer
-} // namespace UE
\ No newline at end of file
+} // namespace UE
diff --git a/Engine/Source/Editor/Sequencer/Public/MVVM/ViewModels/ChannelModel.h b/Engine/Source/Editor/Sequencer/Public/MVVM/ViewModels/ChannelModel.h
--- a/Engine/Source/Editor/Sequencer/Public/MVVM/ViewModels/ChannelModel.h
+++ b/Engine/Source/Editor/Sequencer/Public/MVVM/ViewModels/ChannelModel.h
@@ -11,7 +11,7 @@
 #include "MVVM/Extensions/IDeletableExtension.h"
 #include "MVVM/Extensions/ITrackLaneExtension.h"
 #include "MVVM/Extensions/IKeyExtension.h"
-
+#include "Channels/MovieSceneChannelOverrideContainer.h"
 #include "EventHandlers/ISignedObjectEventHandler.h"
 #include "Tree/ICurveEditorTreeItem.h"
 
@@ -24,7 +24,6 @@ struct FMovieSceneChannelHandle;
 class IKeyArea;
 class ISequencerSection;
 class UMovieSceneSection;
-class FSequencerSectionKeyAreaNode;
 
 namespace UE
 {
@@ -131,6 +130,9 @@ public:
 	/** Get the key area of the channel associated with the given section */
 	TSharedPtr<IKeyArea> GetKeyArea(const UMovieSceneSection* InOwnerSection) const;
 
+	/** Get the channel model at the given index in the list of channels */
+	TSharedPtr<FChannelModel> GetChannel(int32 Index) const;
+
 	/** Get the channel model of the channel associated with the given section */
 	TSharedPtr<FChannelModel> GetChannel(TSharedPtr<FSectionModel> InOwnerSection) const;
 
@@ -140,6 +142,9 @@ public:
 	/** Get the key areas of all channels */
 	TArray<TSharedRef<IKeyArea>> GetAllKeyAreas() const;
 
+	/** Gets a serial number representing if the list of channels has changed */
+	uint32 GetChannelsSerialNumber() const;
+
 public:
 
 	/*~ ITrackAreaExtension */
@@ -149,9 +154,22 @@ public:
 	void CreateCurveModels(TArray<TUniquePtr<FCurveModel>>& OutCurveModels);
 	bool HasCurves() const;
 
+	void BuildChannelOverrideMenu(FMenuBuilder& MenuBuilder);
+
+	void CleanupChannels();
+
+private:
+
+	void BuildChannelOverrideMenu(FMenuBuilder& MenuBuilder, UMovieSceneChannelOverrideContainer::FOverrideCandidates OverrideCandidates);
+	void OverrideChannels(TSubclassOf<UMovieSceneChannelOverrideContainer> OverrideClass);
+	void RemoveChannelOverrides();
+	
+	void BuildChannelOverrideParametersMenu(FMenuBuilder& MenuBuilder, TArray<UMovieSceneChannelOverrideContainer*> ChannelParameters);
+
 protected:
 
 	TArray<TWeakViewModelPtr<FChannelModel>> Channels;
+	uint32 ChannelsSerialNumber;
 	FName ChannelName;
 	FText DisplayText;
 };
@@ -193,14 +211,17 @@ public:
 
 	/*~ ICurveEditorTreeItemExtension */
 	bool HasCurves() const override;
+	void BuildContextMenu(FMenuBuilder& MenuBuilder) override;
 
 private:
 
 	EVisibility GetKeyEditorVisibility() const;
+	
+private:
+
 	FOutlinerSizing ComputedSizing;
 };
 
-
 } // namespace Sequencer
 } // namespace UE
 
diff --git a/Engine/Source/Runtime/MovieScene/Private/Channels/MovieSceneChannelOverrideContainer.cpp b/Engine/Source/Runtime/MovieScene/Private/Channels/MovieSceneChannelOverrideContainer.cpp
new file mode 100644
--- /dev/null
+++ b/Engine/Source/Runtime/MovieScene/Private/Channels/MovieSceneChannelOverrideContainer.cpp
@@ -0,0 +1,22 @@
+// Copyright Epic Games, Inc. All Rights Reserved.
+
+#include "Channels/MovieSceneChannelOverrideContainer.h"
+#include "Channels/MovieSceneChannel.h"
+#include "Templates/SubclassOf.h"
+#include "UObject/UObjectIterator.h"
+
+void UMovieSceneChannelOverrideContainer::GetOverrideCandidates(FName InDefaultChannelTypeName, FOverrideCandidates& OutCandidates)
+{
+	for (TObjectIterator<UClass> ClassIterator; ClassIterator; ++ClassIterator)
+	{
+		if (ClassIterator->IsChildOf(UMovieSceneChannelOverrideContainer::StaticClass()) && 
+				!ClassIterator->HasAnyClassFlags(CLASS_Abstract | CLASS_Deprecated | CLASS_NewerVersionExists))
+		{
+			const UMovieSceneChannelOverrideContainer* OverrideContainer = ClassIterator->GetDefaultObject<UMovieSceneChannelOverrideContainer>();
+			if (OverrideContainer && OverrideContainer->SupportsOverride(InDefaultChannelTypeName))
+			{
+				OutCandidates.Add(TSubclassOf<UMovieSceneChannelOverrideContainer>(*ClassIterator));
+			}
+		}
+	}
+}
diff --git a/Engine/Source/Runtime/MovieScene/Private/Channels/MovieSceneSectionChannelOverrideRegistry.cpp b/Engine/Source/Runtime/MovieScene/Private/Channels/MovieSceneSectionChannelOverrideRegistry.cpp
--- a/Engine/Source/Runtime/MovieScene/Private/Channels/MovieSceneSectionChannelOverrideRegistry.cpp
+++ b/Engine/Source/Runtime/MovieScene/Private/Channels/MovieSceneSectionChannelOverrideRegistry.cpp
@@ -3,52 +3,58 @@
 #include "Channels/MovieSceneSectionChannelOverrideRegistry.h"
 #include "Channels/MovieSceneChannelProxy.h"
 #include "Channels/MovieSceneChannelOverrideContainer.h"
+#include "Channels/IMovieSceneChannelOverrideProvider.h"
 #include "EntitySystem/MovieSceneEntityBuilder.h"
 
 UMovieSceneSectionChannelOverrideRegistry::UMovieSceneSectionChannelOverrideRegistry()
 {
 }
 
-EMovieSceneChannelProxyType UMovieSceneSectionChannelOverrideRegistry::CacheChannelProxy()
+void UMovieSceneSectionChannelOverrideRegistry::AddChannel(FName ChannelName, UMovieSceneChannelOverrideContainer* ChannelContainer)
 {
-	// TODO in Phase 2
-	return EMovieSceneChannelProxyType();
+	Overrides.Emplace(ChannelName, ChannelContainer);
 }
 
-void UMovieSceneSectionChannelOverrideRegistry::AddChannel(int32 ChannelToOverride, UMovieSceneChannelOverrideContainer* ChannelContainer)
+bool UMovieSceneSectionChannelOverrideRegistry::ContainsChannel(FName ChannelName) const
 {
-	Overrides.Emplace(ChannelToOverride, ChannelContainer);
+	return Overrides.Contains(ChannelName);
 }
 
-void UMovieSceneSectionChannelOverrideRegistry::RemoveChannel(int32 ChannelToRemove)
+UMovieSceneChannelOverrideContainer* UMovieSceneSectionChannelOverrideRegistry::GetChannel(FName ChannelName) const
 {
-	Overrides.Remove(ChannelToRemove);
+	return Overrides.FindRef(ChannelName);
 }
 
-bool UMovieSceneSectionChannelOverrideRegistry::IsOverriden(int32 ChannelIndex) const
+void UMovieSceneSectionChannelOverrideRegistry::RemoveChannel(FName ChannelName)
 {
-	return Overrides.Contains(ChannelIndex);
+	Overrides.Remove(ChannelName);
 }
 
-void UMovieSceneSectionChannelOverrideRegistry::ImportEntityImpl(const UE::MovieScene::FEntityImportParams& Params, UE::MovieScene::FImportedEntity* OutImportedEntity)
+void UMovieSceneSectionChannelOverrideRegistry::ImportEntityImpl(const FMovieSceneChannelOverrideEntityImportParamsHandle& OverrideParams, const UE::MovieScene::FEntityImportParams& ImportParams, UE::MovieScene::FImportedEntity* OutImportedEntity)
 {
 	using namespace UE::MovieScene;
-;
-	int32 ChannelIndex = UMovieSceneSectionChannelOverrideRegistry::ToChannelIndex(Params.EntityID);
-	check(IsOverriden(ChannelIndex));	
-	Overrides[ChannelIndex]->ImportEntityImpl(ChannelIndex, Params, OutImportedEntity);
+
+	check(OverrideParams.IsValid());
+	FName ChannelName = OverrideParams.GetPtr()->ChannelName;
+	check(!ChannelName.IsNone());
+	Overrides[ChannelName]->ImportEntityImpl(OverrideParams, ImportParams, OutImportedEntity);
 }
 
-bool UMovieSceneSectionChannelOverrideRegistry::PopulateEvaluationFieldImpl(const TRange<FFrameNumber>& EffectiveRange, const FMovieSceneEvaluationFieldEntityMetaData& InMetaData, FMovieSceneEntityComponentFieldBuilder* OutFieldBuilder, UMovieSceneSection& OwnerSection)
+void UMovieSceneSectionChannelOverrideRegistry::PopulateEvaluationFieldImpl(const TRange<FFrameNumber>& EffectiveRange, const FMovieSceneEvaluationFieldEntityMetaData& InMetaData, FMovieSceneEntityComponentFieldBuilder* OutFieldBuilder, UMovieSceneSection& OwnerSection)
 {
 	using namespace UE::MovieScene;
 
-	for (const TPair<int32, TObjectPtr<UMovieSceneChannelOverrideContainer>>& OverrideChannelPair : Overrides)
+	IMovieSceneChannelOverrideProvider* OverrideProvider = Cast<IMovieSceneChannelOverrideProvider>(GetOuter());
+	if (ensure(OverrideProvider))
 	{
-		const int32 EntityIndex = OutFieldBuilder->FindOrAddEntity(&OwnerSection, ToEntityID(OverrideChannelPair.Key));
-		const int32 MetaDataIndex = OutFieldBuilder->AddMetaData(InMetaData);
-		OutFieldBuilder->AddPersistentEntity(EffectiveRange, EntityIndex, MetaDataIndex);
+		for (const TPair<FName, TObjectPtr<UMovieSceneChannelOverrideContainer>>& Override : Overrides)
+		{
+			FChannelOverrideProviderTraitsHandle ChannelOverrideTraits = OverrideProvider->GetChannelOverrideProviderTraits();
+			check(ChannelOverrideTraits.IsValid());
+			const int32 EntityID = ChannelOverrideTraits->GetChannelOverrideEntityID(Override.Key);
+			const int32 EntityIndex = OutFieldBuilder->FindOrAddEntity(&OwnerSection, EntityID);
+			const int32 MetaDataIndex = OutFieldBuilder->AddMetaData(InMetaData);
+			OutFieldBuilder->AddPersistentEntity(EffectiveRange, EntityIndex, MetaDataIndex);
+		}
 	}
-
-	return false;
 }
diff --git a/Engine/Source/Runtime/MovieScene/Public/Channels/IMovieSceneChannelOverrideProvider.h b/Engine/Source/Runtime/MovieScene/Public/Channels/IMovieSceneChannelOverrideProvider.h
new file mode 100644
--- /dev/null
+++ b/Engine/Source/Runtime/MovieScene/Public/Channels/IMovieSceneChannelOverrideProvider.h
@@ -0,0 +1,172 @@
+// Copyright Epic Games, Inc. All Rights Reserved.
+
+#pragma once
+
+#include "Misc/InlineValue.h"
+#include "UObject/Interface.h"
+#include "Templates/SubclassOf.h"
+#include "IMovieSceneChannelOverrideProvider.generated.h"
+
+class UMovieSceneSectionChannelOverrideRegistry;
+
+namespace UE
+{
+namespace MovieScene
+{
+
+/** Utility class to convert channel names to indices */
+struct FChannelOverrideNames
+{
+	FChannelOverrideNames(int32 InIndexOffset, std::initializer_list<FName> InChannelNames)
+		: ChannelNames(InChannelNames)
+		, IndexOffset(InIndexOffset)
+	{
+	}
+
+	int32 GetIndex(const FName& InChannelName) const
+	{
+		int32 Index = ChannelNames.Find(InChannelName);
+		if (Index != INDEX_NONE)
+		{
+			Index += IndexOffset;
+		}
+		return Index;
+	}
+
+	FName GetChannelName(int32 InIndex) const
+	{
+		const int32 Index = InIndex - IndexOffset;
+		if (ChannelNames.IsValidIndex(Index))
+		{
+			return ChannelNames[Index];
+		}
+		return NAME_None;
+	}
+
+	TArray<FName> ChannelNames;
+	int32 IndexOffset = 0;
+};
+
+/**
+ *
+ */
+struct FChannelOverrideProviderTraits
+{
+	virtual ~FChannelOverrideProviderTraits() {}
+
+	/**
+	 * Gets the channel ID for a given channel.
+	 */
+	virtual FName GetDefaultChannelTypeName(FName ChannelName) const = 0;
+
+	/**
+	 * Convert a channel name to an override entity ID.
+	 */
+	virtual int32 GetChannelOverrideEntityID(FName ChannelName) const = 0;
+
+	/**
+	 * Convert an override entity ID to a channel name.
+	 */
+	virtual FName GetChannelOverrideName(int32 EntityID) const = 0;
+};
+
+using FChannelOverrideProviderTraitsHandle = TInlineValue<FChannelOverrideProviderTraits, 64>;
+
+/**
+ *
+ */
+template<typename DefaultChannelType, int OverrideEntityID = 10>
+struct TSingleChannelOverrideProviderTraits : FChannelOverrideProviderTraits
+{
+	FName SingleChannelName;
+
+	TSingleChannelOverrideProviderTraits(FName InChannelName)
+		: SingleChannelName(InChannelName)
+	{
+	}
+
+	virtual FName GetDefaultChannelTypeName(FName ChannelName) const override
+	{
+		return DefaultChannelType::StaticStruct()->GetFName();
+	}
+
+	virtual int32 GetChannelOverrideEntityID(FName ChannelName) const override
+	{
+		ensure(ChannelName == SingleChannelName);
+		return OverrideEntityID;
+	}
+
+	virtual FName GetChannelOverrideName(int32 EntityID) const override
+	{
+		ensure(EntityID == OverrideEntityID);
+		return SingleChannelName;
+	}
+};
+
+template<typename... OverrideChannelTypes>
+struct TNamedChannelOverrideProviderTraits : FChannelOverrideProviderTraits
+{
+	const FChannelOverrideNames* ChannelNames;
+	FName ChannelTypes[sizeof...(OverrideChannelTypes)];
+
+	TNamedChannelOverrideProviderTraits(const FChannelOverrideNames* InChannelNames)
+		: ChannelNames(InChannelNames)
+		, ChannelTypes{ OverrideChannelTypes::StaticStruct()->GetFName()... }
+	{}
+
+	virtual FName GetDefaultChannelTypeName(FName ChannelName) const override
+	{
+		const int32 ChannelIndex = ChannelNames->GetIndex(ChannelName);
+		check(ChannelIndex >= 0 && ChannelIndex < sizeof(ChannelTypes));
+		return ChannelTypes[ChannelIndex];
+	}
+
+	virtual int32 GetChannelOverrideEntityID(FName ChannelName) const override
+	{
+		return ChannelNames->GetIndex(ChannelName);
+	}
+
+	virtual FName GetChannelOverrideName(int32 EntityID) const override
+	{
+		return ChannelNames->GetChannelName(EntityID);
+	}
+};
+
+}  // namespace MovieScene
+}  // namespace UE
+
+UINTERFACE()
+class MOVIESCENE_API UMovieSceneChannelOverrideProvider : public UInterface
+{
+public:
+
+	GENERATED_BODY()
+};
+
+/**
+ * Interface to be added to UMovieSceneSection types when they contain entity data
+ */
+class MOVIESCENE_API IMovieSceneChannelOverrideProvider
+{
+public:
+
+	using FChannelOverrideProviderTraitsHandle = UE::MovieScene::FChannelOverrideProviderTraitsHandle;
+
+	GENERATED_BODY()
+
+	/**
+	 * Gets the channel override container.
+	 */
+	virtual UMovieSceneSectionChannelOverrideRegistry* GetChannelOverrideRegistry(bool bCreateIfMissing) = 0;
+
+	/**
+	 * Gets the naming/indexing traits for this provider.
+	 */
+	virtual FChannelOverrideProviderTraitsHandle GetChannelOverrideProviderTraits() const = 0;
+
+	/**
+	 * Called when channel overrides have been added or removed. Should invalidate the channel proxy.
+	 */
+	virtual void OnChannelOverridesChanged() = 0;
+};
+
diff --git a/Engine/Source/Runtime/MovieScene/Public/Channels/MovieSceneChannelOverrideContainer.h b/Engine/Source/Runtime/MovieScene/Public/Channels/MovieSceneChannelOverrideContainer.h
--- a/Engine/Source/Runtime/MovieScene/Public/Channels/MovieSceneChannelOverrideContainer.h
+++ b/Engine/Source/Runtime/MovieScene/Public/Channels/MovieSceneChannelOverrideContainer.h
@@ -1,31 +1,115 @@
 // Copyright Epic Games, Inc. All Rights Reserved.
 
 #pragma once
+
 #include "CoreMinimal.h"
 #include "UObject/ObjectMacros.h"
 #include "UObject/Object.h"
+#include "Channels/MovieSceneChannelHandle.h"
 #include "Evaluation/MovieSceneEvaluationField.h"
 #include "EntitySystem/IMovieSceneEntityProvider.h"
+#include "EntitySystem/MovieSceneEntityIDs.h"
 #include "MovieSceneSection.h"
+#include "MovieSceneSignedObject.h"
+#include "Misc/InlineValue.h"
+
 #include "MovieSceneChannelOverrideContainer.generated.h"
 
 class UMovieSceneEntitySystemLinker;
-struct FMovieSceneEntityComponentFieldBuilder;
 struct FMovieSceneChannel;
+struct FMovieSceneChannelMetaData;
+struct FMovieSceneChannelProxyData;
+struct FMovieSceneEntityComponentFieldBuilder;
+
+/**
+ * Base structure for entity import parameters of a channel override.
+ */
+struct IMovieSceneChannelOverrideEntityImportParams
+{
+	FName ChannelName;
+};
+
+/*
+ * Handle for storing channel override import parameters.
+ *
+ * These parameters carry type-specific information, such as the result component in which the channel
+ * override should throw output values.
+ */
+struct FMovieSceneChannelOverrideEntityImportParamsHandle : 
+	TInlineValue<IMovieSceneChannelOverrideEntityImportParams, sizeof(FName) + 16>
+{
+	using Super = TInlineValue<IMovieSceneChannelOverrideEntityImportParams, sizeof(FName) + 16>;
+
+	FMovieSceneChannelOverrideEntityImportParamsHandle()
+		: Super()
+	{
+	}
+
+	template<typename T>
+	FMovieSceneChannelOverrideEntityImportParamsHandle(T&& In)
+		: Super(In)
+	{
+	}
+
+	template<typename T>
+	const T& CastThis() const
+	{
+		return *static_cast<const T*>(this->GetPtr());
+	}
+};
+
+/**
+ * Entity import parameters for channel overrides that just need an untyped result component.
+ */
+struct FMovieSceneChannelOverrideResultComponentEntityImportParams : IMovieSceneChannelOverrideEntityImportParams
+{
+	UE::MovieScene::FComponentTypeID ResultComponent;
+};
 
 /**
-* A wrapper to implement polymorphism for FMovieSceneChannel. Meant to be override by children classes
-*/
+ * Entity import parameters for channel overrides that need a typed result component.
+ */
+template<typename ComponentType>
+struct TMovieSceneChannelOverrideResultComponentEntityImportParams : IMovieSceneChannelOverrideEntityImportParams
+{
+	UE::MovieScene::TComponentTypeID<ComponentType> ResultComponent;
+};
+
+/**
+ * A wrapper to implement polymorphism for FMovieSceneChannel.
+ */
 UCLASS(Abstract)
-class MOVIESCENE_API UMovieSceneChannelOverrideContainer : public UObject
+class MOVIESCENE_API UMovieSceneChannelOverrideContainer : public UMovieSceneSignedObject
 {
 	GENERATED_BODY()
 
 public:
-	/** Virtual methods to be overriden by children containers */
-	virtual void ImportEntityImpl(int32 ChannelIndex, const UE::MovieScene::FEntityImportParams& Params, UE::MovieScene::FImportedEntity* OutImportedEntity) {};
-	virtual EMovieSceneChannelProxyType CacheChannelProxy() { return EMovieSceneChannelProxyType(); }
 
+	/** Returns whether this container's underlying channel can be used as an override to the given channel type */
+	virtual bool SupportsOverride(FName DefaultChannelTypeName) const { return false; }
+
+	/** Imports the entity for this channel */
+	virtual void ImportEntityImpl(
+			const FMovieSceneChannelOverrideEntityImportParamsHandle& OverrideParams, 
+			const UE::MovieScene::FEntityImportParams& ImportParams, UE::MovieScene::FImportedEntity* OutImportedEntity) {};
+
+	/** Gets the underlying channel */
 	virtual const FMovieSceneChannel* GetChannel() const { return nullptr; }
+	/** Gets the underlying channel */
 	virtual FMovieSceneChannel* GetChannel() { return nullptr; }
-};
\ No newline at end of file
+
+#if WITH_EDITOR
+	/** Caches the channel proxy for this channel */
+	virtual FMovieSceneChannelHandle AddChannelProxy(FName ChannelName, FMovieSceneChannelProxyData& ProxyData, const FMovieSceneChannelMetaData& MetaData) { return FMovieSceneChannelHandle(); }
+#else
+	virtual void AddChannelProxy(FName ChannelName, FMovieSceneChannelProxyData& ProxyData) {}
+#endif
+
+public:
+
+	using FOverrideCandidates = TArray<TSubclassOf<UMovieSceneChannelOverrideContainer>, TInlineAllocator<8>>;
+
+	/** Get a list of channel overrides that can work in the place of the given channel type */
+	static void GetOverrideCandidates(FName InDefaultChannelTypeName, FOverrideCandidates& OutCandidates);
+};
+
diff --git a/Engine/Source/Runtime/MovieScene/Public/Channels/MovieSceneChannelProxy.h b/Engine/Source/Runtime/MovieScene/Public/Channels/MovieSceneChannelProxy.h
--- a/Engine/Source/Runtime/MovieScene/Public/Channels/MovieSceneChannelProxy.h
+++ b/Engine/Source/Runtime/MovieScene/Public/Channels/MovieSceneChannelProxy.h
@@ -120,7 +120,7 @@ struct FMovieSceneChannelProxyData
 		// Add the channel
 		const int32 ChannelTypeIndex = AddInternal(InChannel);
 		// Add a default editor data at the same index, hopefully the caller will set it afterwards
-		Entries[ChannelTypeIndex].AddMetaData<ChannelType>(InMetaData, TMovieSceneChannelTraits<ChannelType>::ExtendedEditorDataType());
+		Entries[ChannelTypeIndex].AddMetaData<ChannelType>(InMetaData, typename TMovieSceneChannelTraits<ChannelType>::ExtendedEditorDataType());
 		// Return index usable for SetExtendedEditorData
 		const FName ChannelTypeName = ChannelType::StaticStruct()->GetFName();
 		return FMovieSceneChannelHandle(nullptr, ChannelTypeName, Entries[ChannelTypeIndex].GetChannels().Num() - 1);
diff --git a/Engine/Source/Runtime/MovieScene/Public/Channels/MovieSceneSectionChannelOverrideRegistry.h b/Engine/Source/Runtime/MovieScene/Public/Channels/MovieSceneSectionChannelOverrideRegistry.h
--- a/Engine/Source/Runtime/MovieScene/Public/Channels/MovieSceneSectionChannelOverrideRegistry.h
+++ b/Engine/Source/Runtime/MovieScene/Public/Channels/MovieSceneSectionChannelOverrideRegistry.h
@@ -3,6 +3,10 @@
 #pragma once
 #include "Containers/Map.h"
 #include "CoreMinimal.h"
+
+#include "Channels/MovieSceneChannel.h"
+#include "Channels/MovieSceneChannelOverrideContainer.h"
+#include "Channels/MovieSceneChannelProxy.h"
 #include "EntitySystem/IMovieSceneEntityProvider.h"
 #include "Evaluation/MovieSceneEvaluationField.h"
 #include "HAL/Platform.h"
@@ -40,101 +44,136 @@ class MOVIESCENE_API UMovieSceneSectionChannelOverrideRegistry : public UObject
 {
 	GENERATED_BODY()
 
-	static constexpr int32 ImportEntityIDOffset = 10;
-
-private:
-	/** Map of channel overrides. */
-	UPROPERTY()
-	TMap<int32, TObjectPtr<UMovieSceneChannelOverrideContainer>> Overrides;
-
 public:
 	UMovieSceneSectionChannelOverrideRegistry();
 
-	/** 
-	* A way for all kinds of sections inform the sequencer editor what channels they have
-	*/
-	EMovieSceneChannelProxyType CacheChannelProxy();
-
 	/** 
 	* Add channel to the registry 
-	* @param ChannelToOverride		The key indicates which channel to override
-	* @param Data					The container that owns a overriden channel instanse
+	* @param ChannelName		The name of the channel to override
+	* @param Data				The container that owns a overriden channel instanse
 	*/
-	void AddChannel(int32 ChannelToOverride, UMovieSceneChannelOverrideContainer* ChannelContainer);
+	void AddChannel(FName ChannelName, UMovieSceneChannelOverrideContainer* ChannelContainer);
 
 	/**
-	* Removes a channel from the registry
-	* @param ChannelToOverride		The key indicates which channel to override
+	* Returns if the channel is overriden
+	* @param ChannelName	Name of the channel
+	* @return	Whether this channel is overriden
 	*/
-	void RemoveChannel(int32 ChannelToRemove);
+	bool ContainsChannel(FName ChannelName) const;
 
 	/**
-	* Forward ImportEntityImpl calls to overriden channels
-	*/
-	void ImportEntityImpl(const UE::MovieScene::FEntityImportParams& Params, UE::MovieScene::FImportedEntity* OutImportedEntity);
+	 * Returns the channel override for a given name, or nullptr if not found
+	 * @param ChannelName	Name of the channel
+	 * @return	The channel override container
+	 */
+	UMovieSceneChannelOverrideContainer* GetChannel(FName ChannelName) const;
 
 	/**
-	* Called when overridden channels should populate evaluation field
-	*/
-	bool PopulateEvaluationFieldImpl(const TRange<FFrameNumber>& EffectiveRange, const FMovieSceneEvaluationFieldEntityMetaData& InMetaData, FMovieSceneEntityComponentFieldBuilder* OutFieldBuilder, UMovieSceneSection& OwnerSection);
+	 *
+	 */
+	template<typename ChannelContainerType>
+	void GetChannels(TArray<ChannelContainerType*>& OutChannels) const
+	{
+		for (TPair<FName, TObjectPtr<UMovieSceneChannelOverrideContainer>> Pair : Overrides)
+		{
+			ChannelContainerType* TypedContainer = Cast<ChannelContainerType>(Pair.Value);
+			if (TypedContainer)
+			{
+				OutChannels.Add(TypedContainer);
+			}
+		}
+	}
 
 	/**
-	* From Channel Index to EntityID
-	* @param ChannelIndex	Represents the channel index in Overrides. Should be from 0 to 9
+	* Removes a channel from the registry
+	* @param ChannelName		The name of the channel whose override to remove
 	*/
-	static int32 ToEntityID(int32 ChannelIndex) { return ChannelIndex + ImportEntityIDOffset; }
+	void RemoveChannel(FName ChannelName);
 
 	/**
-	* From EntityID to channel index
-	* @param EntityID	Represents the EntityID.
+	* Forward ImportEntityImpl calls to an overriden channel
 	*/
-	static int32 ToChannelIndex(int32 EntityID) { return EntityID - ImportEntityIDOffset; }
+	void ImportEntityImpl(const FMovieSceneChannelOverrideEntityImportParamsHandle& OverrideParams, const UE::MovieScene::FEntityImportParams& ImportParams, UE::MovieScene::FImportedEntity* OutImportedEntity);
 
 	/**
-	* Returns if the channel is overriden
-	* @param ChannelIndex	Index for the channel.
-	* @return	If this channel is overriden
+	* Called when overridden channels should populate evaluation field
 	*/
-	bool IsOverriden(int32 ChannelIndex) const;
+	void PopulateEvaluationFieldImpl(const TRange<FFrameNumber>& EffectiveRange, const FMovieSceneEvaluationFieldEntityMetaData& InMetaData, FMovieSceneEntityComponentFieldBuilder* OutFieldBuilder, UMovieSceneSection& OwnerSection);
 
-	/**
-	* Returns the overriden map
-	*/
-	const TMap<int32, TObjectPtr<UMovieSceneChannelOverrideContainer>>& GetOverrides() const { return Overrides; }
+private:
+
+	/** Map of channel overrides. */
+	UPROPERTY()
+	TMap<FName, TObjectPtr<UMovieSceneChannelOverrideContainer>> Overrides;
 };
 
 namespace UE
 {
 namespace MovieScene 
 {
-inline bool IsChannelOverriden(const UMovieSceneSectionChannelOverrideRegistry* OverrideRegistry, int32 ChannelIndex)
+
+/** Utility function to return whether a channel is overriden */
+inline bool IsChannelOverriden(const UMovieSceneSectionChannelOverrideRegistry* OverrideRegistry, FName ChannelName)
 {
-	return OverrideRegistry && OverrideRegistry->IsOverriden(ChannelIndex);
+	return OverrideRegistry && OverrideRegistry->ContainsChannel(ChannelName);
 }
 
-template<typename ChannelType>
-inline bool HasAnyDataImpl(const UMovieSceneSectionChannelOverrideRegistry* OverrideRegistry, int32 Index, const ChannelType& InChannel)
+/** Utility function to return a channel override */
+inline UMovieSceneChannelOverrideContainer* GetChannelOverride(const UMovieSceneSectionChannelOverrideRegistry* OverrideRegistry, FName ChannelName)
 {
-	return InChannel.HasAnyData() || IsChannelOverriden(OverrideRegistry, Index);
+	return OverrideRegistry ? OverrideRegistry->GetChannel(ChannelName) : nullptr;
 }
 
-template<typename HeadChannelType, typename ...TailChannelTypes>
-inline bool HasAnyDataImpl(const UMovieSceneSectionChannelOverrideRegistry* OverrideRegistry, int32 HeadIndex, const HeadChannelType& InHeadChannel, TailChannelTypes ...InTailChannels)
+#if WITH_EDITOR
+
+/**
+ * Utility function for adding a possibly-overriden channel into a channel proxy, with some specific extended editor
+ * data.
+ */
+template<typename ChannelType, typename ExtendedEditorDataType>
+void AddChannelProxy(
+		FMovieSceneChannelProxyData& ProxyData, const UMovieSceneSectionChannelOverrideRegistry* OverrideRegistry,
+		FName ChannelName, ChannelType& DefaultChannel, 
+		const FMovieSceneChannelMetaData& InMetaData, ExtendedEditorDataType&& InExtendedEditorData)
 {
-	if (InHeadChannel.HasAnyData() || IsChannelOverriden(OverrideRegistry, HeadIndex))
+	UMovieSceneChannelOverrideContainer* OverrideContainer = GetChannelOverride(OverrideRegistry, ChannelName);
+	if (OverrideContainer)
 	{
-		return true;
+		// Because the extended editor data is a strongly-typed templated parameter, we can't pass it through the
+		// virtual methods of the channel containers. So channel containers will add an entry with default extended
+		// editor data, and we set the correct data immediately after.
+		FMovieSceneChannelHandle ChannelHandle = OverrideContainer->AddChannelProxy(ChannelName, ProxyData, InMetaData);
+		ProxyData.SetExtendedEditorData<ChannelType>(ChannelHandle, InExtendedEditorData);
+	}
+	else
+	{
+		ProxyData.Add(DefaultChannel, InMetaData, InExtendedEditorData);
 	}
-	return HasAnyDataImpl(OverrideRegistry, HeadIndex + 1, InTailChannels...);
 }
 
+#else
+
 /**
-* Returns if at least one channel in the InChannels has any data. An example is HasAnyData(OverrideRegistry, 0, RedCurve, GreenCurve, BlueCurve, AlphaCurve)
-*/
-template<typename ...ChannelTypes>
-inline bool HasAnyData(const UMovieSceneSectionChannelOverrideRegistry* OverrideRegistry, int32 HeadIndex, ChannelTypes ...InChannels)
+ * Utility function for adding a possibly-overriden channel into a channel proxy.
+ */
+template<typename ChannelType>
+void AddChannelProxy(
+		FMovieSceneChannelProxyData& ProxyData, const UMovieSceneSectionChannelOverrideRegistry* OverrideRegistry,
+		FName ChannelName, ChannelType& DefaultChannel)
 {
-	return HasAnyDataImpl(OverrideRegistry, 0, InChannels...);
-}
+	UMovieSceneChannelOverrideContainer* OverrideContainer = GetChannelOverride(OverrideRegistry, ChannelName);
+	if (OverrideContainer)
+	{
+		OverrideContainer->AddChannelProxy(ChannelName, ProxyData);
+	}
+	else
+	{
+		ProxyData.Add(DefaultChannel);
+	}
 }
-}
\ No newline at end of file
+
+#endif
+
+}  // namespace MovieScene
+}  // namespace UE
+
diff --git a/Engine/Source/Runtime/MovieSceneTracks/Private/Channels/MovieSceneDoublePerlinNoiseChannel.cpp b/Engine/Source/Runtime/MovieSceneTracks/Private/Channels/MovieSceneDoublePerlinNoiseChannel.cpp
--- a/Engine/Source/Runtime/MovieSceneTracks/Private/Channels/MovieSceneDoublePerlinNoiseChannel.cpp
+++ b/Engine/Source/Runtime/MovieSceneTracks/Private/Channels/MovieSceneDoublePerlinNoiseChannel.cpp
@@ -1,20 +1,39 @@
 // Copyright Epic Games, Inc. All Rights Reserved.
 
 #include "Channels/MovieSceneDoublePerlinNoiseChannel.h"
+#include "MovieScene.h"
+#include "MovieSceneSection.h"
 
 FMovieSceneDoublePerlinNoiseChannel::FMovieSceneDoublePerlinNoiseChannel()
-	: DoublePerlinNoiseParams{}
+	: PerlinNoiseParams{}
 {
 }
 
-FMovieSceneDoublePerlinNoiseChannel::FMovieSceneDoublePerlinNoiseChannel(const FDoublePerlinNoiseParams& InDoublePerlinNoiseParams)
-	: DoublePerlinNoiseParams{ InDoublePerlinNoiseParams }
+FMovieSceneDoublePerlinNoiseChannel::FMovieSceneDoublePerlinNoiseChannel(const FDoublePerlinNoiseParams& InPerlinNoiseParams)
+	: PerlinNoiseParams{ InPerlinNoiseParams }
 {
 }
 
-double FMovieSceneDoublePerlinNoiseChannel::Evaluate(double InSeconds)
+double FMovieSceneDoublePerlinNoiseChannel::Evaluate(double InSeconds) const
 {
-	double Result = static_cast<double>(FMath::PerlinNoise1D(InSeconds * DoublePerlinNoiseParams.Frequency));
-	Result *= DoublePerlinNoiseParams.Amplitude;
+	double Result = static_cast<double>(FMath::PerlinNoise1D(InSeconds * PerlinNoiseParams.Frequency));
+	Result *= PerlinNoiseParams.Amplitude;
 	return Result;
 }
+
+bool FMovieSceneDoublePerlinNoiseChannel::Evaluate(const UMovieSceneSection* InSection, FFrameTime InTime, double& OutValue) const
+{
+	if (ensure(InSection))
+	{
+		UMovieScene* MovieScene = InSection->GetTypedOuter<UMovieScene>();
+		if (MovieScene)
+		{
+			const FFrameRate TickResolution = MovieScene->GetTickResolution();
+			const double Seconds = TickResolution.AsSeconds(InTime);
+			OutValue = Evaluate(Seconds);
+			return true;
+		}
+	}
+	return false;
+}
+
diff --git a/Engine/Source/Runtime/MovieSceneTracks/Private/Channels/MovieSceneDoublePerlinNoiseChannelContainer.cpp b/Engine/Source/Runtime/MovieSceneTracks/Private/Channels/MovieSceneDoublePerlinNoiseChannelContainer.cpp
--- a/Engine/Source/Runtime/MovieSceneTracks/Private/Channels/MovieSceneDoublePerlinNoiseChannelContainer.cpp
+++ b/Engine/Source/Runtime/MovieSceneTracks/Private/Channels/MovieSceneDoublePerlinNoiseChannelContainer.cpp
@@ -1,26 +1,41 @@
 // Copyright Epic Games, Inc. All Rights Reserved.
 
 #include "Channels/MovieSceneDoublePerlinNoiseChannelContainer.h"
+#include "Channels/MovieSceneDoubleChannel.h"
 #include "Tracks/MovieScenePropertyTrack.h"
 #include "MovieSceneTracksComponentTypes.h"
 #include "Misc/Guid.h"
 
-void UMovieSceneDoublePerlinNoiseChannelContainer::ImportEntityImpl(int32 ChannelIndex, const UE::MovieScene::FEntityImportParams& Params, UE::MovieScene::FImportedEntity* OutImportedEntity)
+bool UMovieSceneDoublePerlinNoiseChannelContainer::SupportsOverride(FName DefaultChannelTypeName) const
+{
+	return FMovieSceneDoubleChannel::StaticStruct()->GetFName() == DefaultChannelTypeName;
+}
+
+void UMovieSceneDoublePerlinNoiseChannelContainer::ImportEntityImpl(const FMovieSceneChannelOverrideEntityImportParamsHandle& OverrideParams, const UE::MovieScene::FEntityImportParams& ImportParams, UE::MovieScene::FImportedEntity* OutImportedEntity)
 {
 	using namespace UE::MovieScene;
 
 	const FMovieSceneTracksComponentTypes* TracksComponents = FMovieSceneTracksComponentTypes::Get();
 	const FBuiltInComponentTypes* BuiltInComponents = FBuiltInComponentTypes::Get();
 
+	const auto& DoubleResultParams = OverrideParams.CastThis<TMovieSceneChannelOverrideResultComponentEntityImportParams<double>>();
+
 	OutImportedEntity->AddBuilder(
 		FEntityBuilder()
-		.Add(TracksComponents->DoublePerlinNoiseChannel, DoublePerlinNoiseChannel.GetParam())
-		.Add(BuiltInComponents->DoubleResult[ChannelIndex], TNumericLimits<double>::Max())
+		.Add(TracksComponents->DoublePerlinNoiseChannel, PerlinNoiseChannel.GetParams())
+		.Add(DoubleResultParams.ResultComponent, TNumericLimits<double>::Max())
 	);
 }
 
-EMovieSceneChannelProxyType UMovieSceneDoublePerlinNoiseChannelContainer::CacheChannelProxy()
+#if WITH_EDITOR
+FMovieSceneChannelHandle UMovieSceneDoublePerlinNoiseChannelContainer::AddChannelProxy(FName ChannelName, FMovieSceneChannelProxyData& ProxyData, const FMovieSceneChannelMetaData& MetaData)
 {
-	//TODO in Phase 2
-	return EMovieSceneChannelProxyType();
+	return ProxyData.AddWithDefaultEditorData(PerlinNoiseChannel, MetaData);
 }
+#else
+void UMovieSceneDoublePerlinNoiseChannelContainer::AddChannelProxy(FName ChannelName, FMovieSceneChannelProxyData& ProxyData)
+{
+	ProxyData.Add(PerlinNoiseChannel);
+}
+#endif
+
diff --git a/Engine/Source/Runtime/MovieSceneTracks/Private/Channels/MovieSceneFloatPerlinNoiseChannel.cpp b/Engine/Source/Runtime/MovieSceneTracks/Private/Channels/MovieSceneFloatPerlinNoiseChannel.cpp
--- a/Engine/Source/Runtime/MovieSceneTracks/Private/Channels/MovieSceneFloatPerlinNoiseChannel.cpp
+++ b/Engine/Source/Runtime/MovieSceneTracks/Private/Channels/MovieSceneFloatPerlinNoiseChannel.cpp
@@ -1,20 +1,39 @@
 // Copyright Epic Games, Inc. All Rights Reserved.
 
 #include "Channels/MovieSceneFloatPerlinNoiseChannel.h"
+#include "MovieScene.h"
+#include "MovieSceneSection.h"
 
 FMovieSceneFloatPerlinNoiseChannel::FMovieSceneFloatPerlinNoiseChannel()
-	: FloatPerlinNoiseParams{}
+	: PerlinNoiseParams{}
 {
 }
 
 FMovieSceneFloatPerlinNoiseChannel::FMovieSceneFloatPerlinNoiseChannel(const FFloatPerlinNoiseParams& InFloatPerlinNoiseParams)
-	: FloatPerlinNoiseParams{ InFloatPerlinNoiseParams }
+	: PerlinNoiseParams{ InFloatPerlinNoiseParams }
 {
 }
 
-float FMovieSceneFloatPerlinNoiseChannel::Evaluate(double InSeconds)
+float FMovieSceneFloatPerlinNoiseChannel::Evaluate(double InSeconds) const
 {
-	float Result = FMath::PerlinNoise1D(InSeconds * FloatPerlinNoiseParams.Frequency);
-	Result *= FloatPerlinNoiseParams.Amplitude;
+	float Result = FMath::PerlinNoise1D(InSeconds * PerlinNoiseParams.Frequency);
+	Result *= PerlinNoiseParams.Amplitude;
 	return Result;
-}
\ No newline at end of file
+}
+
+bool FMovieSceneFloatPerlinNoiseChannel::Evaluate(const UMovieSceneSection* InSection, FFrameTime InTime, float& OutValue) const
+{
+	if (ensure(InSection))
+	{
+		UMovieScene* MovieScene = InSection->GetTypedOuter<UMovieScene>();
+		if (MovieScene)
+		{
+			const FFrameRate TickResolution = MovieScene->GetTickResolution();
+			const double Seconds = TickResolution.AsSeconds(InTime);
+			OutValue = Evaluate(Seconds);
+			return true;
+		}
+	}
+	return false;
+}
+
diff --git a/Engine/Source/Runtime/MovieSceneTracks/Private/Channels/MovieSceneFloatPerlinNoiseChannelContainer.cpp b/Engine/Source/Runtime/MovieSceneTracks/Private/Channels/MovieSceneFloatPerlinNoiseChannelContainer.cpp
--- a/Engine/Source/Runtime/MovieSceneTracks/Private/Channels/MovieSceneFloatPerlinNoiseChannelContainer.cpp
+++ b/Engine/Source/Runtime/MovieSceneTracks/Private/Channels/MovieSceneFloatPerlinNoiseChannelContainer.cpp
@@ -1,26 +1,42 @@
 // Copyright Epic Games, Inc. All Rights Reserved.
 
 #include "Channels/MovieSceneFloatPerlinNoiseChannelContainer.h"
+#include "Channels/MovieSceneFloatChannel.h"
+#include "Channels/MovieSceneChannelProxy.h"
 #include "Tracks/MovieScenePropertyTrack.h"
 #include "MovieSceneTracksComponentTypes.h"
 #include "Misc/Guid.h"
 
-void UMovieSceneFloatPerlinNoiseChannelContainer::ImportEntityImpl(int32 ChannelIndex, const UE::MovieScene::FEntityImportParams& Params, UE::MovieScene::FImportedEntity* OutImportedEntity)
+bool UMovieSceneFloatPerlinNoiseChannelContainer::SupportsOverride(FName DefaultChannelTypeName) const
+{
+	return FMovieSceneFloatChannel::StaticStruct()->GetFName() == DefaultChannelTypeName;
+}
+
+void UMovieSceneFloatPerlinNoiseChannelContainer::ImportEntityImpl(const FMovieSceneChannelOverrideEntityImportParamsHandle& OverrideParams, const UE::MovieScene::FEntityImportParams& ImportParams, UE::MovieScene::FImportedEntity* OutImportedEntity)
 {
 	using namespace UE::MovieScene;
 
 	const FMovieSceneTracksComponentTypes* TracksComponents = FMovieSceneTracksComponentTypes::Get();
 	const FBuiltInComponentTypes* BuiltInComponents = FBuiltInComponentTypes::Get();
 
+	const auto& FloatResultParams = OverrideParams.CastThis<TMovieSceneChannelOverrideResultComponentEntityImportParams<float>>();
+
 	OutImportedEntity->AddBuilder(
 		FEntityBuilder()
-		.Add(TracksComponents->FloatPerlinNoiseChannel, FloatPerlinNoiseChannel.GetParam())
-		.Add(BuiltInComponents->DoubleResult[ChannelIndex], TNumericLimits<float>::Max())
+		.Add(TracksComponents->FloatPerlinNoiseChannel, PerlinNoiseChannel.GetParams())
+		.Add(FloatResultParams.ResultComponent, TNumericLimits<float>::Max())
 	);
 }
 
-EMovieSceneChannelProxyType UMovieSceneFloatPerlinNoiseChannelContainer::CacheChannelProxy()
+#if WITH_EDITOR
+FMovieSceneChannelHandle UMovieSceneFloatPerlinNoiseChannelContainer::AddChannelProxy(FName ChannelName, FMovieSceneChannelProxyData& ProxyData, const FMovieSceneChannelMetaData& MetaData)
 {
-	//TODO in Phase 2
-	return EMovieSceneChannelProxyType();
+	return ProxyData.AddWithDefaultEditorData(PerlinNoiseChannel, MetaData);
 }
+#else
+void UMovieSceneFloatPerlinNoiseChannelContainer::AddChannelProxy(FName ChannelName, FMovieSceneChannelProxyData& ProxyData)
+{
+	ProxyData.Add(PerlinNoiseChannel);
+}
+#endif
+
diff --git a/Engine/Source/Runtime/MovieSceneTracks/Private/Sections/MovieScene3DTransformSection.cpp b/Engine/Source/Runtime/MovieSceneTracks/Private/Sections/MovieScene3DTransformSection.cpp
--- a/Engine/Source/Runtime/MovieSceneTracks/Private/Sections/MovieScene3DTransformSection.cpp
+++ b/Engine/Source/Runtime/MovieSceneTracks/Private/Sections/MovieScene3DTransformSection.cpp
@@ -1,10 +1,12 @@
 // Copyright Epic Games, Inc. All Rights Reserved.
 
 #include "Sections/MovieScene3DTransformSection.h"
+#include "Channels/IMovieSceneChannelOverrideProvider.h"
 #include "UObject/StructOnScope.h"
 #include "UObject/SequencerObjectVersion.h"
 #include "Algo/AnyOf.h"
 #include "Channels/MovieSceneChannelProxy.h"
+#include "Channels/MovieSceneSectionChannelOverrideRegistry.h"
 #include "GameFramework/Actor.h"
 #include "EulerTransform.h"
 #include "Systems/MovieScene3DTransformPropertySystem.h"
@@ -307,6 +309,13 @@ namespace UE::MovieScene
 /* UMovieScene3DTransformSection interface
  *****************************************************************************/
 
+UE::MovieScene::FChannelOverrideNames UMovieScene3DTransformSection::ChannelOverrideNames(10, {
+		TEXT("Location.X"), TEXT("Location.Y"), TEXT("Location.Z"),
+		TEXT("Rotation.X"), TEXT("Rotation.Y"), TEXT("Rotation.Z"),
+		TEXT("Scale.X"), TEXT("Scale.Y"), TEXT("Scale.Z"),
+		TEXT("Weight")
+		});
+
 UMovieScene3DTransformSection::UMovieScene3DTransformSection(const FObjectInitializer& ObjectInitializer)
 	: Super(ObjectInitializer)
 	, bUseQuaternionInterpolation(false)
@@ -396,31 +405,54 @@ void UMovieScene3DTransformSection::BuildEntity(BaseBuilderType& InBaseBuilder,
 		RotationChannel[2] = TrackComponents->QuaternionRotationChannel[2];
 	}
 
-	// Let override registry handles overriden channel here
-	int32 ChannelIndex = UMovieSceneSectionChannelOverrideRegistry::ToChannelIndex(Params.EntityID);
+	// Let the override registry handle overriden channels
+	const FName ChannelOverrideName = ChannelOverrideNames.GetChannelName(Params.EntityID);
 
-	// If we are building the entity of an overriden channel, we pass the partial builder we have at this point (the one with the bindings) and add the tag to it.
+	// If we are building the entity of an overriden channel, we pass the partial builder we have at this point 
+	// (the one with the bindings) and add the tag to it.
 	// Then call into the override repository so it can add a second builder with the override stuff.
-	if (IsChannelOverriden(OverrideRegistry, ChannelIndex))
+	if (ChannelOverrideName != NAME_None)
 	{
+		check(OverrideRegistry && OverrideRegistry->ContainsChannel(ChannelOverrideName));
 		OutImportedEntity->AddBuilder(InBaseBuilder.AddTag(PropertyTag));
-		OverrideRegistry->ImportEntityImpl(Params, OutImportedEntity);
+
+		const int32 ChannelOverrideIndex = (Params.EntityID - ChannelOverrideNames.IndexOffset);
+		if (ChannelOverrideIndex == 9)
+		{
+			if (EnumHasAnyFlags(Channels, EMovieSceneTransformChannel::Weight))
+			{
+				TMovieSceneChannelOverrideResultComponentEntityImportParams<double> OverrideParams;
+				OverrideParams.ChannelName = ChannelOverrideName;
+				OverrideParams.ResultComponent = BuiltInComponentTypes->WeightResult;
+				OverrideRegistry->ImportEntityImpl(OverrideParams, Params, OutImportedEntity);
+			}
+		}
+		else
+		{
+			if (ActiveChannelsMask[ChannelOverrideIndex])
+			{
+				TMovieSceneChannelOverrideResultComponentEntityImportParams<double> OverrideParams;
+				OverrideParams.ChannelName = ChannelOverrideName;
+				OverrideParams.ResultComponent = BuiltInComponentTypes->DoubleResult[ChannelOverrideIndex];
+				OverrideRegistry->ImportEntityImpl(OverrideParams, Params, OutImportedEntity);
+			}
+		}
 		return;
 	}	
 	
 	// Here proceed with the all the normal channels
 	OutImportedEntity->AddBuilder(
 		InBaseBuilder
-		.AddConditional(BuiltInComponentTypes->DoubleChannel[0], &Translation[0], ActiveChannelsMask[0] && !IsChannelOverriden(OverrideRegistry, 0))
-		.AddConditional(BuiltInComponentTypes->DoubleChannel[1], &Translation[1], ActiveChannelsMask[1] && !IsChannelOverriden(OverrideRegistry, 1))
-		.AddConditional(BuiltInComponentTypes->DoubleChannel[2], &Translation[2], ActiveChannelsMask[2] && !IsChannelOverriden(OverrideRegistry, 2))
-		.AddConditional(RotationChannel[0],                      &Rotation[0],    ActiveChannelsMask[3] && !IsChannelOverriden(OverrideRegistry, 3))
-		.AddConditional(RotationChannel[1],                      &Rotation[1],    ActiveChannelsMask[4] && !IsChannelOverriden(OverrideRegistry, 4))
-		.AddConditional(RotationChannel[2],                       &Rotation[2],   ActiveChannelsMask[5] && !IsChannelOverriden(OverrideRegistry, 5))
-		.AddConditional(BuiltInComponentTypes->DoubleChannel[6], &Scale[0],       ActiveChannelsMask[6] && !IsChannelOverriden(OverrideRegistry, 6))
-		.AddConditional(BuiltInComponentTypes->DoubleChannel[7], &Scale[1],       ActiveChannelsMask[7] && !IsChannelOverriden(OverrideRegistry, 7))
-		.AddConditional(BuiltInComponentTypes->DoubleChannel[8], &Scale[2],       ActiveChannelsMask[8] && !IsChannelOverriden(OverrideRegistry, 8))
-		.AddConditional(BuiltInComponentTypes->WeightChannel,    &ManualWeight,   EnumHasAnyFlags(Channels, EMovieSceneTransformChannel::Weight) && ManualWeight.HasAnyData() && !IsChannelOverriden(OverrideRegistry, 9))
+		.AddConditional(BuiltInComponentTypes->DoubleChannel[0], &Translation[0], ActiveChannelsMask[0] && !IsChannelOverriden(OverrideRegistry, TEXT("Location.X")))
+		.AddConditional(BuiltInComponentTypes->DoubleChannel[1], &Translation[1], ActiveChannelsMask[1] && !IsChannelOverriden(OverrideRegistry, TEXT("Location.Y")))
+		.AddConditional(BuiltInComponentTypes->DoubleChannel[2], &Translation[2], ActiveChannelsMask[2] && !IsChannelOverriden(OverrideRegistry, TEXT("Location.Z")))
+		.AddConditional(RotationChannel[0],                      &Rotation[0],    ActiveChannelsMask[3] && !IsChannelOverriden(OverrideRegistry, TEXT("Rotation.X")))
+		.AddConditional(RotationChannel[1],                      &Rotation[1],    ActiveChannelsMask[4] && !IsChannelOverriden(OverrideRegistry, TEXT("Rotation.Y")))
+		.AddConditional(RotationChannel[2],                      &Rotation[2],    ActiveChannelsMask[5] && !IsChannelOverriden(OverrideRegistry, TEXT("Rotation.Z")))
+		.AddConditional(BuiltInComponentTypes->DoubleChannel[6], &Scale[0],       ActiveChannelsMask[6] && !IsChannelOverriden(OverrideRegistry, TEXT("Scale.X")))
+		.AddConditional(BuiltInComponentTypes->DoubleChannel[7], &Scale[1],       ActiveChannelsMask[7] && !IsChannelOverriden(OverrideRegistry, TEXT("Scale.Y")))
+		.AddConditional(BuiltInComponentTypes->DoubleChannel[8], &Scale[2],       ActiveChannelsMask[8] && !IsChannelOverriden(OverrideRegistry, TEXT("Scale.Z")))
+		.AddConditional(BuiltInComponentTypes->WeightChannel,    &ManualWeight,   EnumHasAnyFlags(Channels, EMovieSceneTransformChannel::Weight) && ManualWeight.HasAnyData() && !IsChannelOverriden(OverrideRegistry, TEXT("Weight")))
 		.AddTag(PropertyTag)
 	);
 }
@@ -453,9 +485,12 @@ bool UMovieScene3DTransformSection::PopulateEvaluationFieldImpl(const TRange<FFr
 	}
 	if (OverrideRegistry)
 	{
-		return OverrideRegistry->PopulateEvaluationFieldImpl(EffectiveRange, InMetaData, OutFieldBuilder, *this);
+		// Add evaluation field entries for each channel that runs with a different logic.
+		OverrideRegistry->PopulateEvaluationFieldImpl(EffectiveRange, InMetaData, OutFieldBuilder, *this);
 	}
 
+	// Return false even if we have an override registry, so that we also add an evaluation field entry for the default
+	// stuff.
 	return false;
 }
 
@@ -598,6 +633,8 @@ FMovieSceneTransformMask UMovieScene3DTransformSection::GetMaskByName(const FNam
 
 EMovieSceneChannelProxyType UMovieScene3DTransformSection::CacheChannelProxy()
 {
+	using namespace UE::MovieScene;
+
 	FMovieSceneChannelProxyData Channels;
 
 	//Constraints go on top
@@ -641,29 +678,29 @@ EMovieSceneChannelProxyType UMovieScene3DTransformSection::CacheChannelProxy()
 
 	F3DTransformChannelEditorData EditorData(TransformMask.GetChannels(), NumConstraints);
 
-	Channels.Add(Translation[0], EditorData.MetaData[0], EditorData.ExternalValues[0]);
-	Channels.Add(Translation[1], EditorData.MetaData[1], EditorData.ExternalValues[1]);
-	Channels.Add(Translation[2], EditorData.MetaData[2], EditorData.ExternalValues[2]);
-	Channels.Add(Rotation[0],    EditorData.MetaData[3], EditorData.ExternalValues[3]);
-	Channels.Add(Rotation[1],    EditorData.MetaData[4], EditorData.ExternalValues[4]);
-	Channels.Add(Rotation[2],    EditorData.MetaData[5], EditorData.ExternalValues[5]);
-	Channels.Add(Scale[0],       EditorData.MetaData[6], EditorData.ExternalValues[6]);
-	Channels.Add(Scale[1],       EditorData.MetaData[7], EditorData.ExternalValues[7]);
-	Channels.Add(Scale[2],       EditorData.MetaData[8], EditorData.ExternalValues[8]);
-	Channels.Add(ManualWeight,   EditorData.MetaData[9], EditorData.WeightExternalValue);
+	AddChannelProxy(Channels, OverrideRegistry, TEXT("Location.X"), Translation[0], EditorData.MetaData[0], EditorData.ExternalValues[0]);
+	AddChannelProxy(Channels, OverrideRegistry, TEXT("Location.Y"), Translation[1], EditorData.MetaData[1], EditorData.ExternalValues[1]);
+	AddChannelProxy(Channels, OverrideRegistry, TEXT("Location.Z"), Translation[2], EditorData.MetaData[2], EditorData.ExternalValues[2]);
+	AddChannelProxy(Channels, OverrideRegistry, TEXT("Rotation.X"), Rotation[0],    EditorData.MetaData[3], EditorData.ExternalValues[3]);
+	AddChannelProxy(Channels, OverrideRegistry, TEXT("Rotation.Y"), Rotation[1],    EditorData.MetaData[4], EditorData.ExternalValues[4]);
+	AddChannelProxy(Channels, OverrideRegistry, TEXT("Rotation.Z"), Rotation[2],    EditorData.MetaData[5], EditorData.ExternalValues[5]);
+	AddChannelProxy(Channels, OverrideRegistry, TEXT("Scale.X"),    Scale[0],       EditorData.MetaData[6], EditorData.ExternalValues[6]);
+	AddChannelProxy(Channels, OverrideRegistry, TEXT("Scale.Y"),    Scale[1],       EditorData.MetaData[7], EditorData.ExternalValues[7]);
+	AddChannelProxy(Channels, OverrideRegistry, TEXT("Scale.Z"),    Scale[2],       EditorData.MetaData[8], EditorData.ExternalValues[8]);
+	AddChannelProxy(Channels, OverrideRegistry, TEXT("Weight"),     ManualWeight,   EditorData.MetaData[9], EditorData.WeightExternalValue);
 
 #else
 
-	Channels.Add(Translation[0]);
-	Channels.Add(Translation[1]);
-	Channels.Add(Translation[2]);
-	Channels.Add(Rotation[0]);
-	Channels.Add(Rotation[1]);
-	Channels.Add(Rotation[2]);
-	Channels.Add(Scale[0]);
-	Channels.Add(Scale[1]);
-	Channels.Add(Scale[2]);
-	Channels.Add(ManualWeight);
+	AddChannelProxy(Channels, OverrideRegistry, TEXT("Location.X"), Translation[0]);
+	AddChannelProxy(Channels, OverrideRegistry, TEXT("Location.Y"), Translation[1]);
+	AddChannelProxy(Channels, OverrideRegistry, TEXT("Location.Z"), Translation[2]);
+	AddChannelProxy(Channels, OverrideRegistry, TEXT("Rotation.X"), Rotation[0]);
+	AddChannelProxy(Channels, OverrideRegistry, TEXT("Rotation.Y"), Rotation[1]);
+	AddChannelProxy(Channels, OverrideRegistry, TEXT("Rotation.Z"), Rotation[2]);
+	AddChannelProxy(Channels, OverrideRegistry, TEXT("Scale.X"),    Scale[0]);
+	AddChannelProxy(Channels, OverrideRegistry, TEXT("Scale.Y"),    Scale[1]);
+	AddChannelProxy(Channels, OverrideRegistry, TEXT("Scale.Z"),    Scale[2]);
+	AddChannelProxy(Channels, OverrideRegistry, TEXT("Weight"),     ManualWeight);
 
 #endif
 
@@ -902,3 +939,49 @@ void UMovieScene3DTransformSection::RemoveConstraintChannel(const FName& InConst
 		}
 	}
 }
+
+UMovieSceneSectionChannelOverrideRegistry* UMovieScene3DTransformSection::GetChannelOverrideRegistry(bool bCreateIfMissing)
+{
+	if (bCreateIfMissing && OverrideRegistry == nullptr)
+	{
+		OverrideRegistry = NewObject<UMovieSceneSectionChannelOverrideRegistry>(this);
+	}
+	return OverrideRegistry;
+}
+
+UE::MovieScene::FChannelOverrideProviderTraitsHandle UMovieScene3DTransformSection::GetChannelOverrideProviderTraits() const
+{
+	struct F3DTransformChannelOverrideProviderTraits : UE::MovieScene::FChannelOverrideProviderTraits
+	{
+		FName GetDefaultChannelTypeName(FName ChannelName) const override
+		{
+			if (ChannelName == TEXT("Weight"))
+			{
+				return FMovieSceneFloatChannel::StaticStruct()->GetFName();
+			}
+			else
+			{
+				return FMovieSceneDoubleChannel::StaticStruct()->GetFName();
+			}
+		}
+
+		int32 GetChannelOverrideEntityID(FName ChannelName) const override
+		{
+			return UMovieScene3DTransformSection::ChannelOverrideNames.GetIndex(ChannelName);
+		}
+
+		FName GetChannelOverrideName(int32 EntityID) const override
+		{
+			return UMovieScene3DTransformSection::ChannelOverrideNames.GetChannelName(EntityID);
+		}
+	};
+
+	F3DTransformChannelOverrideProviderTraits Traits;
+	return UE::MovieScene::FChannelOverrideProviderTraitsHandle(Traits);
+}
+
+void UMovieScene3DTransformSection::OnChannelOverridesChanged()
+{
+	ChannelProxy = nullptr;
+}
+
diff --git a/Engine/Source/Runtime/MovieSceneTracks/Private/Sections/MovieSceneFloatSection.cpp b/Engine/Source/Runtime/MovieSceneTracks/Private/Sections/MovieSceneFloatSection.cpp
--- a/Engine/Source/Runtime/MovieSceneTracks/Private/Sections/MovieSceneFloatSection.cpp
+++ b/Engine/Source/Runtime/MovieSceneTracks/Private/Sections/MovieSceneFloatSection.cpp
@@ -2,6 +2,7 @@
 
 #include "Sections/MovieSceneFloatSection.h"
 #include "Channels/MovieSceneChannelProxy.h"
+#include "Channels/MovieSceneSectionChannelOverrideRegistry.h"
 #include "EntitySystem/BuiltInComponentTypes.h"
 #include "EntitySystem/MovieSceneEntityBuilder.h"
 #include "EntitySystem/MovieSceneEntityManager.h"
@@ -60,7 +61,28 @@ void UMovieSceneFloatSection::ImportEntityImpl(UMovieSceneEntitySystemLinker* En
 	const FBuiltInComponentTypes* Components = FBuiltInComponentTypes::Get();
 	const FMovieSceneTracksComponentTypes* TracksComponents = FMovieSceneTracksComponentTypes::Get();
 
-	FPropertyTrackEntityImportHelper(TracksComponents->Float)
-		.Add(Components->FloatChannel[0], &FloatCurve)
+	FPropertyTrackWithOverridableChannelsEntityImportHelper(TracksComponents->Float, this)
+		.Add(Components->FloatChannel[0], FName("FloatChannel"), &FloatCurve)
 		.Commit(this, Params, OutImportedEntity);
 }
+
+UMovieSceneSectionChannelOverrideRegistry* UMovieSceneFloatSection::GetChannelOverrideRegistry(bool bCreateIfMissing)
+{
+	if (bCreateIfMissing && OverrideRegistry == nullptr)
+	{
+		OverrideRegistry = NewObject<UMovieSceneSectionChannelOverrideRegistry>(this);
+	}
+	return OverrideRegistry;
+}
+
+UE::MovieScene::FChannelOverrideProviderTraitsHandle UMovieSceneFloatSection::GetChannelOverrideProviderTraits() const
+{
+	UE::MovieScene::TSingleChannelOverrideProviderTraits<FMovieSceneFloatChannel> Traits(FName("FloatChannel"));
+	return UE::MovieScene::FChannelOverrideProviderTraitsHandle(Traits);
+}
+
+void UMovieSceneFloatSection::OnChannelOverridesChanged()
+{
+	ChannelProxy = nullptr;
+}
+
diff --git a/Engine/Source/Runtime/MovieSceneTracks/Private/Tracks/MovieScenePropertyTrack.cpp b/Engine/Source/Runtime/MovieSceneTracks/Private/Tracks/MovieScenePropertyTrack.cpp
--- a/Engine/Source/Runtime/MovieSceneTracks/Private/Tracks/MovieScenePropertyTrack.cpp
+++ b/Engine/Source/Runtime/MovieSceneTracks/Private/Tracks/MovieScenePropertyTrack.cpp
@@ -331,13 +331,17 @@ UMovieSceneSection* UMovieScenePropertyTrack::GetSectionToKey() const
 const int32 FMovieScenePropertyTrackEntityImportHelper::SectionPropertyValueImportingID = 0;
 const int32 FMovieScenePropertyTrackEntityImportHelper::SectionEditConditionToggleImportingID = 1;
 
-void FMovieScenePropertyTrackEntityImportHelper::PopulateEvaluationField(UMovieSceneSection& Section, const TRange<FFrameNumber>& EffectiveRange, const FMovieSceneEvaluationFieldEntityMetaData& InMetaData, FMovieSceneEntityComponentFieldBuilder* OutFieldBuilder, UMovieSceneSectionChannelOverrideRegistry* OverridesRegistry/*= nullptr*/)
+void FMovieScenePropertyTrackEntityImportHelper::PopulateEvaluationField(UMovieSceneSection& Section, const TRange<FFrameNumber>& EffectiveRange, const FMovieSceneEvaluationFieldEntityMetaData& InMetaData, FMovieSceneEntityComponentFieldBuilder* OutFieldBuilder)
 {
 	using namespace UE::MovieScene;
 
-	if (OverridesRegistry)
+	IMovieSceneChannelOverrideProvider* RegistryProvider = Cast<IMovieSceneChannelOverrideProvider>(&Section);
+	if (RegistryProvider)
 	{
-		OverridesRegistry->PopulateEvaluationFieldImpl(EffectiveRange, InMetaData, OutFieldBuilder, Section);
+		if (UMovieSceneSectionChannelOverrideRegistry* OverrideRegistry = RegistryProvider->GetChannelOverrideRegistry(false))
+		{
+			OverrideRegistry->PopulateEvaluationFieldImpl(EffectiveRange, InMetaData, OutFieldBuilder, Section);
+		}
 	}
 
 	// Add the default entity for this section.
diff --git a/Engine/Source/Runtime/MovieSceneTracks/Public/Channels/MovieSceneDoublePerlinNoiseChannel.h b/Engine/Source/Runtime/MovieSceneTracks/Public/Channels/MovieSceneDoublePerlinNoiseChannel.h
--- a/Engine/Source/Runtime/MovieSceneTracks/Public/Channels/MovieSceneDoublePerlinNoiseChannel.h
+++ b/Engine/Source/Runtime/MovieSceneTracks/Public/Channels/MovieSceneDoublePerlinNoiseChannel.h
@@ -23,13 +23,19 @@ struct MOVIESCENETRACKS_API FMovieSceneDoublePerlinNoiseChannel : public FMovieS
 	* @param InSeconds  The Frame second to evaluate at
 	* @return A value to receive the PerlinNoise result
 	*/
-	double Evaluate(double InSeconds);
+	double Evaluate(double InSeconds) const;
 
-	const FDoublePerlinNoiseParams& GetParam() const { return DoublePerlinNoiseParams; }
-	FDoublePerlinNoiseParams& GetParam() { return DoublePerlinNoiseParams; }
+	bool Evaluate(const UMovieSceneSection* InSection, FFrameTime InTime, double& OutValue) const;
+
+	const FDoublePerlinNoiseParams& GetParams() const { return PerlinNoiseParams; }
+	FDoublePerlinNoiseParams& GetParams() { return PerlinNoiseParams; }
 
 private:
-	FDoublePerlinNoiseParams DoublePerlinNoiseParams;
+
+	UPROPERTY(EditAnywhere, Category = "Perlin Noise")
+	FDoublePerlinNoiseParams PerlinNoiseParams;
+
+	friend class FMovieSceneDoublePerlinNoiseChannelDetailsCustomization;
 };
 
 template<>
@@ -37,8 +43,14 @@ struct TMovieSceneChannelTraits<FMovieSceneDoublePerlinNoiseChannel> : TMovieSce
 {
 #if WITH_EDITOR
 
-	/** FMovieSceneFloatPerlinNoiseChannel channels can have external values (ie, they can get their values from external objects for UI purposes) */
+	/** Perlin noise channels can have external values (ie, they can get their values from external objects for UI purposes) */
 	typedef TMovieSceneExternalValue<double> ExtendedEditorDataType;
 
 #endif
 };
+
+inline bool EvaluateChannel(const UMovieSceneSection* InSection, const FMovieSceneDoublePerlinNoiseChannel* InChannel, FFrameTime InTime, double& OutValue)
+{
+	return InChannel->Evaluate(InSection, InTime, OutValue);
+}
+
diff --git a/Engine/Source/Runtime/MovieSceneTracks/Public/Channels/MovieSceneDoublePerlinNoiseChannelContainer.h b/Engine/Source/Runtime/MovieSceneTracks/Public/Channels/MovieSceneDoublePerlinNoiseChannelContainer.h
--- a/Engine/Source/Runtime/MovieSceneTracks/Public/Channels/MovieSceneDoublePerlinNoiseChannelContainer.h
+++ b/Engine/Source/Runtime/MovieSceneTracks/Public/Channels/MovieSceneDoublePerlinNoiseChannelContainer.h
@@ -1,26 +1,42 @@
 // Copyright Epic Games, Inc. All Rights Reserved.
 
 #pragma once
+
 #include "Channels/MovieSceneChannelOverrideContainer.h"
 #include "Channels/MovieSceneDoublePerlinNoiseChannel.h"
 #include "MovieSceneDoublePerlinNoiseChannelContainer.generated.h"
 
 /**
-* Double perlin noise channel overriden container
-*/
-UCLASS(meta = (DisplayName = "Double Perlin Noise", ToolTip = "Override a channel to use double perlin noise"))
+ * Double perlin noise channel overriden container
+ */
+UCLASS(meta=(DisplayName="Double Perlin Noise", ToolTip="Override a channel to use double perlin noise"))
 class  MOVIESCENETRACKS_API UMovieSceneDoublePerlinNoiseChannelContainer : public UMovieSceneChannelOverrideContainer
 {
 	GENERATED_BODY()
 
+public:
+
+	using ChannelType = FMovieSceneDoublePerlinNoiseChannel;
+
+	bool SupportsOverride(FName DefaultChannelTypeName) const override;
+	void ImportEntityImpl(
+			const FMovieSceneChannelOverrideEntityImportParamsHandle& OverrideParams, 
+			const UE::MovieScene::FEntityImportParams& ImportParams, UE::MovieScene::FImportedEntity* OutImportedEntity) override;
+
+	const FMovieSceneChannel* GetChannel() const override { return &PerlinNoiseChannel; }
+	FMovieSceneChannel* GetChannel() override { return &PerlinNoiseChannel; }
+
+#if WITH_EDITOR
+	FMovieSceneChannelHandle AddChannelProxy(FName ChannelName, FMovieSceneChannelProxyData& ProxyData, const FMovieSceneChannelMetaData& MetaData) override;
+#else
+	void AddChannelProxy(FName ChannelName, FMovieSceneChannelProxyData& ProxyData) override;
+#endif
+
 private:
-	UPROPERTY()
-	FMovieSceneDoublePerlinNoiseChannel DoublePerlinNoiseChannel;
 
-public:
-	virtual void ImportEntityImpl(int32 ChannelIndex, const UE::MovieScene::FEntityImportParams& Params, UE::MovieScene::FImportedEntity* OutImportedEntity) override final;
-	virtual EMovieSceneChannelProxyType CacheChannelProxy() override final;
+	UPROPERTY(EditAnywhere, Category = "Perlin Noise")
+	FMovieSceneDoublePerlinNoiseChannel PerlinNoiseChannel;
 
-	virtual const FMovieSceneChannel* GetChannel() const override final { return &DoublePerlinNoiseChannel; }
-	virtual FMovieSceneChannel* GetChannel() override final { return &DoublePerlinNoiseChannel; }
+	friend class FMovieSceneDoublePerlinNoiseChannelDetailsCustomization;
 };
+
diff --git a/Engine/Source/Runtime/MovieSceneTracks/Public/Channels/MovieSceneFloatPerlinNoiseChannel.h b/Engine/Source/Runtime/MovieSceneTracks/Public/Channels/MovieSceneFloatPerlinNoiseChannel.h
--- a/Engine/Source/Runtime/MovieSceneTracks/Public/Channels/MovieSceneFloatPerlinNoiseChannel.h
+++ b/Engine/Source/Runtime/MovieSceneTracks/Public/Channels/MovieSceneFloatPerlinNoiseChannel.h
@@ -23,13 +23,19 @@ struct MOVIESCENETRACKS_API FMovieSceneFloatPerlinNoiseChannel : public FMovieSc
 	* @param InSeconds  The Frame second to evaluate at
 	* @return A value to receive the PerlinNoise result
 	*/
-	float Evaluate(double InSeconds);
+	float Evaluate(double InSeconds) const;
 
-	FFloatPerlinNoiseParams& GetParam() { return FloatPerlinNoiseParams; }
-	const FFloatPerlinNoiseParams& GetParam() const { return FloatPerlinNoiseParams; }
+	bool Evaluate(const UMovieSceneSection* InSection, FFrameTime InTime, float& OutValue) const;
+
+	FFloatPerlinNoiseParams& GetParams() { return PerlinNoiseParams; }
+	const FFloatPerlinNoiseParams& GetParams() const { return PerlinNoiseParams; }
 
 private:
-	FFloatPerlinNoiseParams FloatPerlinNoiseParams;
+
+	UPROPERTY(EditAnywhere, Category="Perlin Noise")
+	FFloatPerlinNoiseParams PerlinNoiseParams;
+
+	friend class FMovieSceneFloatPerlinNoiseChannelDetailsCustomization;
 };
 
 template<>
@@ -37,8 +43,14 @@ struct TMovieSceneChannelTraits<FMovieSceneFloatPerlinNoiseChannel> : TMovieScen
 {
 #if WITH_EDITOR
 
-	/** FMovieSceneFloatPerlinNoiseChannel channels can have external values (ie, they can get their values from external objects for UI purposes) */
+	/** Perlin noise channels can have external values (ie, they can get their values from external objects for UI purposes) */
 	typedef TMovieSceneExternalValue<float> ExtendedEditorDataType;
 
 #endif
 };
+
+inline bool EvaluateChannel(const UMovieSceneSection* InSection, const FMovieSceneFloatPerlinNoiseChannel* InChannel, FFrameTime InTime, float& OutValue)
+{
+	return InChannel->Evaluate(InSection, InTime, OutValue);
+}
+
diff --git a/Engine/Source/Runtime/MovieSceneTracks/Public/Channels/MovieSceneFloatPerlinNoiseChannelContainer.h b/Engine/Source/Runtime/MovieSceneTracks/Public/Channels/MovieSceneFloatPerlinNoiseChannelContainer.h
--- a/Engine/Source/Runtime/MovieSceneTracks/Public/Channels/MovieSceneFloatPerlinNoiseChannelContainer.h
+++ b/Engine/Source/Runtime/MovieSceneTracks/Public/Channels/MovieSceneFloatPerlinNoiseChannelContainer.h
@@ -1,26 +1,43 @@
 // Copyright Epic Games, Inc. All Rights Reserved.
 
 #pragma once
+
 #include "Channels/MovieSceneChannelOverrideContainer.h"
 #include "Channels/MovieSceneFloatPerlinNoiseChannel.h"
+#include "EntitySystem/IMovieSceneEntityProvider.h"
 #include "MovieSceneFloatPerlinNoiseChannelContainer.generated.h"
 
 /**
-* Float perlin noise channel overriden container
-*/
-UCLASS(meta = (DisplayName = "Float Perlin Noise", ToolTip = "Override a channel to use float perlin noise"))
+ * Float perlin noise channel overriden container
+ */
+UCLASS(meta=(DisplayName="Float Perlin Noise", ToolTip="Override a channel to use float perlin noise"))
 class MOVIESCENETRACKS_API UMovieSceneFloatPerlinNoiseChannelContainer : public UMovieSceneChannelOverrideContainer
 {
 	GENERATED_BODY()
 
+public:
+
+	using ChannelType = FMovieSceneFloatPerlinNoiseChannel;
+
+	bool SupportsOverride(FName DefaultChannelTypeName) const override;
+	void ImportEntityImpl(
+			const FMovieSceneChannelOverrideEntityImportParamsHandle& OverrideParams, 
+			const UE::MovieScene::FEntityImportParams& ImportParams, UE::MovieScene::FImportedEntity* OutImportedEntity) override;
+
+	const FMovieSceneChannel* GetChannel() const override { return &PerlinNoiseChannel; }
+	FMovieSceneChannel* GetChannel() override { return &PerlinNoiseChannel; }
+
+#if WITH_EDITOR
+	FMovieSceneChannelHandle AddChannelProxy(FName ChannelName, FMovieSceneChannelProxyData& ProxyData, const FMovieSceneChannelMetaData& MetaData) override;
+#else
+	void AddChannelProxy(FName ChannelName, FMovieSceneChannelProxyData& ProxyData) override;
+#endif
+
 private:
-	UPROPERTY()
-	FMovieSceneFloatPerlinNoiseChannel FloatPerlinNoiseChannel;
 
-public:
-	virtual void ImportEntityImpl(int32 ChannelIndex, const UE::MovieScene::FEntityImportParams& Params, UE::MovieScene::FImportedEntity* OutImportedEntity) override final;
-	virtual EMovieSceneChannelProxyType CacheChannelProxy() override final;
+	UPROPERTY(EditAnywhere, Category = "Perlin Noise")
+	FMovieSceneFloatPerlinNoiseChannel PerlinNoiseChannel;
 
-	virtual const FMovieSceneChannel* GetChannel() const override final { return &FloatPerlinNoiseChannel; }
-	virtual FMovieSceneChannel* GetChannel() override final { return &FloatPerlinNoiseChannel; }
+	friend class FMovieSceneFloatPerlinNoiseChannelDetailsCustomization;
 };
+
diff --git a/Engine/Source/Runtime/MovieSceneTracks/Public/MovieSceneTracksComponentTypes.h b/Engine/Source/Runtime/MovieSceneTracks/Public/MovieSceneTracksComponentTypes.h
--- a/Engine/Source/Runtime/MovieSceneTracks/Public/MovieSceneTracksComponentTypes.h
+++ b/Engine/Source/Runtime/MovieSceneTracks/Public/MovieSceneTracksComponentTypes.h
@@ -27,10 +27,12 @@ struct MOVIESCENETRACKS_API FFloatPerlinNoiseParams
 {
 	GENERATED_BODY()
 
-	UPROPERTY()
+	/** The frequency of the noise, i.e. how many times per second does the noise peak */
+	UPROPERTY(EditAnywhere, Category="Perlin Noise")
 	float Frequency;
 
-	UPROPERTY()
+	/** The amplitude of the noise, which will vary between [-Amplitude, +Amplitude] */
+	UPROPERTY(EditAnywhere, Category = "Perlin Noise")
 	float Amplitude;
 
 	FFloatPerlinNoiseParams();
@@ -43,10 +45,12 @@ struct MOVIESCENETRACKS_API FDoublePerlinNoiseParams
 {
 	GENERATED_BODY()
 
-	UPROPERTY()
+	/** The frequency of the noise, i.e. how many times per second does the noise peak */
+	UPROPERTY(EditAnywhere, Category = "Perlin Noise")
 	double Frequency;
 
-	UPROPERTY()
+	/** The amplitude of the noise, which will vary between [-Amplitude, +Amplitude] */
+	UPROPERTY(EditAnywhere, Category = "Perlin Noise")
 	double Amplitude;
 
 	FDoublePerlinNoiseParams();
diff --git a/Engine/Source/Runtime/MovieSceneTracks/Public/Sections/MovieScene3DTransformSection.h b/Engine/Source/Runtime/MovieSceneTracks/Public/Sections/MovieScene3DTransformSection.h
--- a/Engine/Source/Runtime/MovieSceneTracks/Public/Sections/MovieScene3DTransformSection.h
+++ b/Engine/Source/Runtime/MovieSceneTracks/Public/Sections/MovieScene3DTransformSection.h
@@ -12,6 +12,7 @@
 #include "Channels/MovieSceneDoubleChannel.h"
 #include "Channels/MovieSceneFloatChannel.h"
 #include "Channels/MovieSceneSectionChannelOverrideRegistry.h"
+#include "Channels/IMovieSceneChannelOverrideProvider.h"
 #include "EntitySystem/IMovieSceneEntityProvider.h"
 #include "TransformData.h"
 #include "Misc/LargeWorldCoordinates.h"
@@ -211,6 +212,7 @@ class UMovieScene3DTransformSection
 	: public UMovieSceneSection
 	, public IMovieSceneConstrainedSection
 	, public IMovieSceneEntityProvider
+	, public IMovieSceneChannelOverrideProvider
 {
 	GENERATED_UCLASS_BODY()
 
@@ -258,7 +260,7 @@ protected:
 	MOVIESCENETRACKS_API virtual EMovieSceneChannelProxyType CacheChannelProxy() override;
 
 private:
-	
+
 	MOVIESCENETRACKS_API virtual bool PopulateEvaluationFieldImpl(const TRange<FFrameNumber>& EffectiveRange, const FMovieSceneEvaluationFieldEntityMetaData& InMetaData, FMovieSceneEntityComponentFieldBuilder* OutFieldBuilder) override;
 	MOVIESCENETRACKS_API virtual void ImportEntityImpl(UMovieSceneEntitySystemLinker* EntityLinker, const FEntityImportParams& Params, FImportedEntity* OutImportedEntity) override;
 	MOVIESCENETRACKS_API virtual void InterrogateEntityImpl(UMovieSceneEntitySystemLinker* EntityLinker, const FEntityImportParams& Params, FImportedEntity* OutImportedEntity) override;
@@ -268,6 +270,14 @@ private:
 	void PopulateConstraintEntities(const TRange<FFrameNumber>& EffectiveRange, const FMovieSceneEvaluationFieldEntityMetaData& InMetaData, FMovieSceneEntityComponentFieldBuilder* OutFieldBuilder);
 	void ImportConstraintEntity(UMovieSceneEntitySystemLinker* EntityLinker, const FEntityImportParams& Params, FImportedEntity* OutImportedEntity);
 
+private:
+
+	static UE::MovieScene::FChannelOverrideNames ChannelOverrideNames;
+
+	UMovieSceneSectionChannelOverrideRegistry* GetChannelOverrideRegistry(bool bCreateIfMissing) override;
+	UE::MovieScene::FChannelOverrideProviderTraitsHandle GetChannelOverrideProviderTraits() const override;
+	void OnChannelOverridesChanged() override;
+
 private:
 
 	UPROPERTY()
diff --git a/Engine/Source/Runtime/MovieSceneTracks/Public/Sections/MovieSceneFloatSection.h b/Engine/Source/Runtime/MovieSceneTracks/Public/Sections/MovieSceneFloatSection.h
--- a/Engine/Source/Runtime/MovieSceneTracks/Public/Sections/MovieSceneFloatSection.h
+++ b/Engine/Source/Runtime/MovieSceneTracks/Public/Sections/MovieSceneFloatSection.h
@@ -8,6 +8,7 @@
 #include "Curves/RichCurve.h"
 #include "MovieSceneSection.h"
 #include "Channels/MovieSceneFloatChannel.h"
+#include "Channels/IMovieSceneChannelOverrideProvider.h"
 #include "EntitySystem/IMovieSceneEntityProvider.h"
 #include "MovieSceneFloatSection.generated.h"
 
@@ -19,6 +20,7 @@ UCLASS( MinimalAPI )
 class UMovieSceneFloatSection
 	: public UMovieSceneSection
 	, public IMovieSceneEntityProvider
+	, public IMovieSceneChannelOverrideProvider
 {
 	GENERATED_BODY()
 
@@ -31,12 +33,22 @@ public:
 	 */
 	const FMovieSceneFloatChannel& GetChannel() const { return FloatCurve; }
 
+protected:
+
+	virtual UMovieSceneSectionChannelOverrideRegistry* GetChannelOverrideRegistry(bool bCreateIfMissing) override;
+	virtual UE::MovieScene::FChannelOverrideProviderTraitsHandle GetChannelOverrideProviderTraits() const override;
+	virtual void OnChannelOverridesChanged() override;
+
 protected:
 
 	/** Float data */
 	UPROPERTY()
 	FMovieSceneFloatChannel FloatCurve;
 
+	/** Channel overrides */
+	UPROPERTY()
+	TObjectPtr<UMovieSceneSectionChannelOverrideRegistry> OverrideRegistry;
+
 private:
 
 	virtual void ImportEntityImpl(UMovieSceneEntitySystemLinker* EntityLinker, const FEntityImportParams& Params, FImportedEntity* OutImportedEntity) override;
@@ -44,3 +56,4 @@ private:
 
 	virtual EMovieSceneChannelProxyType CacheChannelProxy() override;
 };
+
diff --git a/Engine/Source/Runtime/MovieSceneTracks/Public/Tracks/MovieScenePropertyTrack.h b/Engine/Source/Runtime/MovieSceneTracks/Public/Tracks/MovieScenePropertyTrack.h
--- a/Engine/Source/Runtime/MovieSceneTracks/Public/Tracks/MovieScenePropertyTrack.h
+++ b/Engine/Source/Runtime/MovieSceneTracks/Public/Tracks/MovieScenePropertyTrack.h
@@ -3,16 +3,17 @@
 #pragma once
 
 #include "CoreMinimal.h"
+#include "Channels/IMovieSceneChannelOverrideProvider.h"
+#include "Channels/MovieSceneSectionChannelOverrideRegistry.h"
 #include "EntitySystem/BuiltInComponentTypes.h"
 #include "EntitySystem/IMovieSceneEntityProvider.h"
 #include "EntitySystem/MovieSceneEntityBuilder.h"
+#include "EntitySystem/MovieScenePropertyBinding.h"
 #include "MovieSceneCommonHelpers.h"
 #include "MovieSceneNameableTrack.h"
 #include "MovieSceneTracksComponentTypes.h"
 #include "Systems/MovieScenePiecewiseBoolBlenderSystem.h"
-#include "EntitySystem/MovieScenePropertyBinding.h"
 #include "UObject/ObjectMacros.h"
-#include "Channels/MovieSceneSectionChannelOverrideRegistry.h"
 #include "MovieScenePropertyTrack.generated.h"
 
 /**
@@ -155,7 +156,7 @@ struct MOVIESCENETRACKS_API FMovieScenePropertyTrackEntityImportHelper
 	static const int32 SectionPropertyValueImportingID;
 	static const int32 SectionEditConditionToggleImportingID;
 
-	static void PopulateEvaluationField(UMovieSceneSection& Section, const TRange<FFrameNumber>& EffectiveRange, const FMovieSceneEvaluationFieldEntityMetaData& InMetaData, FMovieSceneEntityComponentFieldBuilder* OutFieldBuilder, UMovieSceneSectionChannelOverrideRegistry* OverrideRegistry = nullptr);
+	static void PopulateEvaluationField(UMovieSceneSection& Section, const TRange<FFrameNumber>& EffectiveRange, const FMovieSceneEvaluationFieldEntityMetaData& InMetaData, FMovieSceneEntityComponentFieldBuilder* OutFieldBuilder);
 
 	static bool IsPropertyValueID(const UE::MovieScene::FEntityImportParams& Params);
 	static bool IsEditConditionToggleID(const UE::MovieScene::FEntityImportParams& Params);
@@ -177,9 +178,8 @@ namespace MovieScene
 template<typename... T>
 struct TPropertyTrackEntityImportHelperImpl
 {
-	TPropertyTrackEntityImportHelperImpl(TEntityBuilder<T...>&& InBuilder, UMovieSceneSectionChannelOverrideRegistry* InOverrideRegistry = nullptr, FComponentTypeID InPropertyTag = {})
+	TPropertyTrackEntityImportHelperImpl(TEntityBuilder<T...>&& InBuilder, FComponentTypeID InPropertyTag = {})
 		: Builder(MoveTemp(InBuilder))
-		, OverrideRegistry(InOverrideRegistry)
 		, PropertyTag(InPropertyTag)
 	{
 	}
@@ -187,19 +187,13 @@ struct TPropertyTrackEntityImportHelperImpl
 	template<typename U, typename PayloadType>
 	TPropertyTrackEntityImportHelperImpl<T..., TAdd<U>> Add(TComponentTypeID<U> ComponentType, PayloadType&& InPayload)
 	{
-		return TPropertyTrackEntityImportHelperImpl<T..., TAdd<U>>(Builder.Add(ComponentType, InPayload), OverrideRegistry, PropertyTag);
+		return TPropertyTrackEntityImportHelperImpl<T..., TAdd<U>>(Builder.Add(ComponentType, InPayload), PropertyTag);
 	}
 
 	template<typename U, typename PayloadType>
 	TPropertyTrackEntityImportHelperImpl<T..., TAddConditional<U>> AddConditional(TComponentTypeID<U> ComponentType, PayloadType&& InPayload, bool bCondition)
 	{
-		return TPropertyTrackEntityImportHelperImpl<T..., TAddConditional<U>>(Builder.AddConditional(ComponentType, InPayload, bCondition), OverrideRegistry, PropertyTag);
-	}
-
-	template<typename U, typename PayloadType>
-	TPropertyTrackEntityImportHelperImpl<T..., TAddConditional<U>> AddConditional(TComponentTypeID<U> ComponentType, PayloadType&& InPayload, bool bCondition, int32 ChannelIndex)
-	{
-		return TPropertyTrackEntityImportHelperImpl<T..., TAddConditional<U>>(Builder.AddConditional(ComponentType, InPayload, bCondition && !IsChannelOverriden(OverrideRegistry, ChannelIndex)), OverrideRegistry, PropertyTag);
+		return TPropertyTrackEntityImportHelperImpl<T..., TAddConditional<U>>(Builder.AddConditional(ComponentType, InPayload, bCondition), PropertyTag);
 	}
 
 	void Commit(const UMovieSceneSection* InSection, const UE::MovieScene::FEntityImportParams& Params, UE::MovieScene::FImportedEntity* OutImportedEntity)
@@ -217,23 +211,12 @@ struct TPropertyTrackEntityImportHelperImpl
 					.Add(BuiltInComponents->PropertyBinding, PropertyTrack->GetPropertyBinding())
 					.AddConditional(BuiltInComponents->GenericObjectBinding, ObjectBindingID, ObjectBindingID.IsValid()));
 			}
-			else if (FMovieScenePropertyTrackEntityImportHelper::IsEditConditionToggleID(Params))
+			else if (ensure(FMovieScenePropertyTrackEntityImportHelper::IsEditConditionToggleID(Params)))
 			{
 				// We effectively discard the builder we've been setting up, because we just
 				// need to import the edit condition toggle entity.
 				FMovieScenePropertyTrackEntityImportHelper::ImportEditConditionToggleEntity(Params, OutImportedEntity);
 			}
-			// Overriden
-			else if (OverrideRegistry)
-			{
-				auto BaseBuilder = FEntityBuilder()
-					.Add(BuiltInComponents->PropertyBinding, PropertyTrack->GetPropertyBinding())
-					.AddConditional(BuiltInComponents->GenericObjectBinding, ObjectBindingID, ObjectBindingID.IsValid())
-					.AddTag(PropertyTag);
-			
-				OutImportedEntity->AddBuilder(BaseBuilder);
-				OverrideRegistry->ImportEntityImpl(Params, OutImportedEntity);
-			}
 		}
 		else
 		{
@@ -246,7 +229,6 @@ struct TPropertyTrackEntityImportHelperImpl
 protected:
 
 	TEntityBuilder<T...> Builder;
-	TObjectPtr<UMovieSceneSectionChannelOverrideRegistry> OverrideRegistry;
 	FComponentTypeID PropertyTag;
 };
 
@@ -257,9 +239,8 @@ protected:
 template<>
 struct TPropertyTrackEntityImportHelperImpl<>
 {
-	TPropertyTrackEntityImportHelperImpl(FComponentTypeID InPropertyTag, UMovieSceneSectionChannelOverrideRegistry* InOverrideRegistry = nullptr)
+	TPropertyTrackEntityImportHelperImpl(FComponentTypeID InPropertyTag)
 		: Builder(FEntityBuilder().AddTag(InPropertyTag))
-		, OverrideRegistry(InOverrideRegistry)
 		, PropertyTag(InPropertyTag)
 	{
 	}
@@ -267,33 +248,205 @@ struct TPropertyTrackEntityImportHelperImpl<>
 	template<typename U, typename PayloadType>
 	TPropertyTrackEntityImportHelperImpl<FAdd, TAdd<U>> Add(TComponentTypeID<U> ComponentType, PayloadType&& InPayload)
 	{
-		return TPropertyTrackEntityImportHelperImpl<FAdd, TAdd<U>>(Builder.Add(ComponentType, InPayload), OverrideRegistry, PropertyTag);
+		return TPropertyTrackEntityImportHelperImpl<FAdd, TAdd<U>>(Builder.Add(ComponentType, InPayload), PropertyTag);
 	}
 
 	template<typename U, typename PayloadType>
 	TPropertyTrackEntityImportHelperImpl<FAdd, TAddConditional<U>> AddConditional(TComponentTypeID<U> ComponentType, PayloadType&& InPayload, bool bCondition)
 	{
-		return TPropertyTrackEntityImportHelperImpl<FAdd, TAddConditional<U>>(Builder.AddConditional(ComponentType, InPayload, bCondition), OverrideRegistry, PropertyTag);
+		return TPropertyTrackEntityImportHelperImpl<FAdd, TAddConditional<U>>(Builder.AddConditional(ComponentType, InPayload, bCondition), PropertyTag);
+	}
+
+protected:
+
+	TEntityBuilder<FAdd> Builder;
+	FComponentTypeID PropertyTag;
+};
+
+struct MOVIESCENETRACKS_API FPropertyTrackEntityImportHelper : TPropertyTrackEntityImportHelperImpl<>
+{
+	template<typename PropertyTraits>
+	FPropertyTrackEntityImportHelper(const TPropertyComponents<PropertyTraits>& PropertyComponents)
+		: TPropertyTrackEntityImportHelperImpl<>(PropertyComponents.PropertyTag)
+	{}
+};
+
+struct FPropertyTrackEntityImportHelperParamsImpl
+{
+	FPropertyTrackEntityImportHelperParamsImpl(TScriptInterface<IMovieSceneChannelOverrideProvider> InRegistryProvider)
+		: RegistryProvider(InRegistryProvider)
+	{
+		if (InRegistryProvider)
+		{
+			Registry = InRegistryProvider->GetChannelOverrideRegistry(false);
+		}
+	}
+	TScriptInterface<IMovieSceneChannelOverrideProvider> RegistryProvider;
+	TObjectPtr<UMovieSceneSectionChannelOverrideRegistry> Registry;
+	TArray<FMovieSceneChannelOverrideResultComponentEntityImportParams> ImportParams;
+};
+
+/**
+ * Same as TPropertyTrackEntityImportHelperImpl but with support for overridable channels.
+ */
+template<typename... T>
+struct TPropertyTrackWithOverridableChannelsEntityImportHelperImpl
+{
+	TPropertyTrackWithOverridableChannelsEntityImportHelperImpl(TEntityBuilder<T...>&& InBuilder, FPropertyTrackEntityImportHelperParamsImpl& InOverrideInfo, FComponentTypeID InPropertyTag = {})
+		: Builder(MoveTemp(InBuilder))
+		, OverrideInfo(InOverrideInfo)
+		, PropertyTag(InPropertyTag)
+	{
 	}
 
 	template<typename U, typename PayloadType>
-	TPropertyTrackEntityImportHelperImpl<FAdd, TAddConditional<U>> AddConditional(TComponentTypeID<U> ComponentType, PayloadType&& InPayload, bool bCondition, int32 ChannelIndex)
+	TPropertyTrackEntityImportHelperImpl<T..., TAddConditional<U>> Add(TComponentTypeID<U> ComponentType, FName ChannelName, PayloadType&& InPayload)
 	{
-		return TPropertyTrackEntityImportHelperImpl<FAdd, TAddConditional<U>>(Builder.AddConditional(ComponentType, InPayload, bCondition && !IsChannelOverriden(OverrideRegistry, ChannelIndex)), OverrideRegistry, PropertyTag);
+		check(!ChannelName.IsNone());
+		const bool bOverriden = (OverrideInfo.Registry && OverrideInfo.Registry->ContainsChannel(ChannelName));
+		if (bOverriden)
+		{
+			FMovieSceneChannelOverrideResultComponentEntityImportParams OverrideParams;
+			OverrideParams.ChannelName = ChannelName;
+			OverrideParams.ResultComponent = ComponentType;
+			OverrideInfo.ImportParams.Add(OverrideParams);
+		}
+		return TPropertyTrackEntityImportHelperImpl<T..., TAddConditional<U>>(Builder.Add(ComponentType, InPayload, !bOverriden), OverrideInfo, PropertyTag);
+	}
+
+	template<typename U, typename PayloadType>
+	TPropertyTrackEntityImportHelperImpl<T..., TAddConditional<U>> AddConditional(TComponentTypeID<U> ComponentType, FName ChannelName, PayloadType&& InPayload, bool bCondition)
+	{
+		check(!ChannelName.IsNone());
+		const bool bOverriden = (OverrideInfo.Registry && OverrideInfo.Registry->ContainsChannel(ChannelName));
+		if (bOverriden)
+		{
+			FMovieSceneChannelOverrideResultComponentEntityImportParams OverrideParams;
+			OverrideParams.ChannelName = ChannelName;
+			OverrideParams.ResultComponent = ComponentType;
+			OverrideInfo.ImportParams.Add(OverrideParams);
+		}
+		return TPropertyTrackEntityImportHelperImpl<T..., TAddConditional<U>>(Builder.AddConditional(ComponentType, InPayload, bCondition && !bOverriden), OverrideInfo, PropertyTag);
+	}
+
+	void Commit(const UMovieSceneSection* InSection, const UE::MovieScene::FEntityImportParams& Params, UE::MovieScene::FImportedEntity* OutImportedEntity)
+	{
+		const FGuid ObjectBindingID = Params.GetObjectBindingID();
+		const FBuiltInComponentTypes* BuiltInComponents = FBuiltInComponentTypes::Get();
+		const FMovieSceneTracksComponentTypes* TracksComponents = FMovieSceneTracksComponentTypes::Get();
+
+		if (UMovieScenePropertyTrack* PropertyTrack = Cast<UMovieScenePropertyTrack>(InSection->GetOuter()))
+		{
+			if (FMovieScenePropertyTrackEntityImportHelper::IsPropertyValueID(Params))
+			{
+				OutImportedEntity->AddBuilder(
+					Builder
+					.Add(BuiltInComponents->PropertyBinding, PropertyTrack->GetPropertyBinding())
+					.AddConditional(BuiltInComponents->GenericObjectBinding, ObjectBindingID, ObjectBindingID.IsValid()));
+			}
+			else if (FMovieScenePropertyTrackEntityImportHelper::IsEditConditionToggleID(Params))
+			{
+				// We effectively discard the builder we've been setting up, because we just
+				// need to import the edit condition toggle entity.
+				FMovieScenePropertyTrackEntityImportHelper::ImportEditConditionToggleEntity(Params, OutImportedEntity);
+			}
+			else if (ensure(OverrideInfo.Registry))
+			{
+				// This should be an overriden channel, which goes into a different entity.
+				if (ensure(OverrideInfo.RegistryProvider))
+				{
+					auto BaseBuilder = FEntityBuilder()
+						.Add(BuiltInComponents->PropertyBinding, PropertyTrack->GetPropertyBinding())
+						.AddConditional(BuiltInComponents->GenericObjectBinding, ObjectBindingID, ObjectBindingID.IsValid())
+						.AddTag(PropertyTag);
+					OutImportedEntity->AddBuilder(BaseBuilder);
+				
+					FChannelOverrideProviderTraitsHandle Traits = OverrideInfo.RegistryProvider->GetChannelOverrideProviderTraits();
+					check(Traits.IsValid());
+					const FName ChannelOverrideName = Traits->GetChannelOverrideName(Params.EntityID);
+					auto* ChannelOverrideParams = OverrideInfo.ImportParams.FindByPredicate(
+							[=](const FMovieSceneChannelOverrideResultComponentEntityImportParams& CurParams)
+							{
+								return CurParams.ChannelName == ChannelOverrideName;
+							});
+					if (ensure(ChannelOverrideParams))
+					{
+						OverrideInfo.Registry->ImportEntityImpl(*ChannelOverrideParams, Params, OutImportedEntity);
+					}
+				}
+			}
+		}
+		else
+		{
+			OutImportedEntity->AddBuilder(
+				Builder
+				.AddConditional(BuiltInComponents->GenericObjectBinding, ObjectBindingID, ObjectBindingID.IsValid()));
+		}
+	}
+
+protected:
+
+	TEntityBuilder<T...> Builder;
+	FPropertyTrackEntityImportHelperParamsImpl OverrideInfo;
+	FComponentTypeID PropertyTag;
+};
+
+
+/**
+ * The starting point for TPropertyTrackWithOverridableChannelsEntityImportHelperImpl<...T>
+ */
+template<>
+struct TPropertyTrackWithOverridableChannelsEntityImportHelperImpl<>
+{
+	TPropertyTrackWithOverridableChannelsEntityImportHelperImpl(FComponentTypeID InPropertyTag, TScriptInterface<IMovieSceneChannelOverrideProvider> InRegistryProvider)
+		: Builder(FEntityBuilder().AddTag(InPropertyTag))
+		, OverrideInfo(InRegistryProvider)
+		, PropertyTag(InPropertyTag)
+	{
+	}
+
+	template<typename U, typename PayloadType>
+	TPropertyTrackWithOverridableChannelsEntityImportHelperImpl<FAdd, TAddConditional<U>> Add(TComponentTypeID<U> ComponentType, FName ChannelName, PayloadType&& InPayload)
+	{
+		check(!ChannelName.IsNone());
+		const bool bOverriden = (OverrideInfo.Registry && OverrideInfo.Registry->ContainsChannel(ChannelName));
+		if (bOverriden)
+		{
+			FMovieSceneChannelOverrideResultComponentEntityImportParams OverrideParams;
+			OverrideParams.ChannelName = ChannelName;
+			OverrideParams.ResultComponent = ComponentType;
+			OverrideInfo.ImportParams.Add(OverrideParams);
+		}
+		return TPropertyTrackWithOverridableChannelsEntityImportHelperImpl<FAdd, TAddConditional<U>>(Builder.AddConditional(ComponentType, InPayload, !bOverriden), OverrideInfo, PropertyTag);
+	}
+
+	template<typename U, typename PayloadType>
+	TPropertyTrackWithOverridableChannelsEntityImportHelperImpl<FAdd, TAddConditional<U>> AddConditional(TComponentTypeID<U> ComponentType, FName ChannelName, PayloadType&& InPayload, bool bCondition)
+	{
+		check(!ChannelName.IsNone());
+		const bool bOverriden = (OverrideInfo.Registry && OverrideInfo.Registry->ContainsChannel(ChannelName));
+		if (bOverriden)
+		{
+			FMovieSceneChannelOverrideResultComponentEntityImportParams OverrideParams;
+			OverrideParams.ChannelName = ChannelName;
+			OverrideParams.ResultComponent = ComponentType;
+			OverrideInfo.ImportParams.Add(OverrideParams);
+		}
+		return TPropertyTrackWithOverridableChannelsEntityImportHelperImpl<FAdd, TAddConditional<U>>(Builder.AddConditional(ComponentType, InPayload, bCondition && !bOverriden), OverrideInfo, PropertyTag);
 	}
 
 protected:
 
 	TEntityBuilder<FAdd> Builder;
-	TObjectPtr<UMovieSceneSectionChannelOverrideRegistry> OverrideRegistry;
+	FPropertyTrackEntityImportHelperParamsImpl OverrideInfo;
 	FComponentTypeID PropertyTag;
 };
 
-struct MOVIESCENETRACKS_API FPropertyTrackEntityImportHelper : TPropertyTrackEntityImportHelperImpl<>
+struct MOVIESCENETRACKS_API FPropertyTrackWithOverridableChannelsEntityImportHelper : TPropertyTrackWithOverridableChannelsEntityImportHelperImpl<>
 {
 	template<typename PropertyTraits>
-	FPropertyTrackEntityImportHelper(const TPropertyComponents<PropertyTraits>& PropertyComponents, UMovieSceneSectionChannelOverrideRegistry* InOverrideRegistry = nullptr)
-		: TPropertyTrackEntityImportHelperImpl<>(PropertyComponents.PropertyTag, InOverrideRegistry)
+	FPropertyTrackWithOverridableChannelsEntityImportHelper(const TPropertyComponents<PropertyTraits>& PropertyComponents, TScriptInterface<IMovieSceneChannelOverrideProvider> InRegistryProvider)
+		: TPropertyTrackWithOverridableChannelsEntityImportHelperImpl<>(PropertyComponents.PropertyTag, InRegistryProvider)
 	{}
 };
 
