Sequencer: Null pointer check for world

diff --git a/Engine/Source/Runtime/MovieSceneTracks/Private/Evaluation/MovieSceneAudioTemplate.cpp b/Engine/Source/Runtime/MovieSceneTracks/Private/Evaluation/MovieSceneAudioTemplate.cpp
--- a/Engine/Source/Runtime/MovieSceneTracks/Private/Evaluation/MovieSceneAudioTemplate.cpp
+++ b/Engine/Source/Runtime/MovieSceneTracks/Private/Evaluation/MovieSceneAudioTemplate.cpp
@@ -467,25 +467,28 @@ struct FAudioSectionExecutionToken : IMovieSceneExecutionToken
 		bool bSoundNeedsStateChange = !AudioComponent.IsPlaying() || AudioComponent.Sound != Sound;
 		bool bSoundNeedsTimeSync = false;
 
+		UObject* PlaybackContext = Player.GetPlaybackContext();
+		UWorld* World = PlaybackContext ? PlaybackContext->GetWorld() : nullptr;
+
 		// Sync only if there is no time dilation because otherwise the system will constantly resync because audio 
 		// playback is not dilated and will never match the expected playback time.
 		const bool bDoTimeSync = 
-			Player.GetPlaybackContext()->GetWorld() && 
-			(FMath::IsNearlyEqual(Player.GetPlaybackContext()->GetWorld()->GetWorldSettings()->GetEffectiveTimeDilation(), 1.f) ||
+			World && World->GetWorldSettings() &&
+			(FMath::IsNearlyEqual(World->GetWorldSettings()->GetEffectiveTimeDilation(), 1.f) ||
 			 !bIgnoreAudioSyncDuringWorldTimeDilationCVar);
 
 		if (bDoTimeSync)
 		{
 			float CurrentGameTime = 0.0f;
 
-			FAudioDevice* AudioDevice = Player.GetPlaybackContext()->GetWorld()->GetAudioDeviceRaw();
+			FAudioDevice* AudioDevice = World ? World->GetAudioDeviceRaw() : nullptr;
 			if (UseAudioClockForSequencerDesyncCVar && AudioDevice)
 			{
 				CurrentGameTime = AudioDevice->GetAudioClock();
 			}
 			else
 			{
-				CurrentGameTime = Player.GetPlaybackContext()->GetWorld()->GetAudioTimeSeconds();
+				CurrentGameTime = World ? World->GetAudioTimeSeconds() : 0.f;
 			}
 
 			// This tells us how much time has passed in the game world (and thus, reasonably, the audio playback)
@@ -528,8 +531,6 @@ struct FAudioSectionExecutionToken : IMovieSceneExecutionToken
 				AudioComponent.SetSound(Sound);
 			}
 #if WITH_EDITOR
-			UObject* PlaybackContext = Player.GetPlaybackContext();
-			UWorld* World = PlaybackContext ? PlaybackContext->GetWorld() : nullptr;
 			if (GIsEditor && World != nullptr && !World->IsPlayInEditor())
 			{
 				AudioComponent.bIsUISound = true;
@@ -547,21 +548,21 @@ struct FAudioSectionExecutionToken : IMovieSceneExecutionToken
 				AudioComponent.Play(AudioTime);
 
 				// Keep track of when we asked this audio clip to play (in game time) so that we can figure out if there's a significant desync in the future.
-				if (Player.GetPlaybackContext()->GetWorld())
+				if (World)
 				{
 					if (!TrackData.SoundLastPlayedAtTime.Contains(&AudioComponent))
 					{
 						TrackData.SoundLastPlayedAtTime.Add(&AudioComponent);
 					}
 
-					FAudioDevice* AudioDevice = Player.GetPlaybackContext()->GetWorld()->GetAudioDeviceRaw();
+					FAudioDevice* AudioDevice = World->GetAudioDeviceRaw();
 					if (UseAudioClockForSequencerDesyncCVar && AudioDevice)
 					{
 						TrackData.SoundLastPlayedAtTime[&AudioComponent] = AudioDevice->GetAudioClock() - AudioTime;
 					}
 					else
 					{
-						TrackData.SoundLastPlayedAtTime[&AudioComponent] = Player.GetPlaybackContext()->GetWorld()->GetAudioTimeSeconds() - AudioTime;
+						TrackData.SoundLastPlayedAtTime[&AudioComponent] = World->GetAudioTimeSeconds() - AudioTime;
 					}
 					
 				}
