Sequencer: Fixed material parameter bound material tracking

Previously, the BoundMaterial could change any time, which caused the overlapping entity tracker to get confused about what channels were being animated by which materials.
This was compounded by the systems that added the bound material components would sometimes _remove_ the BoundMaterial altogether if it was no longer bound to a valid material. This would ultimately mean that when the entity was destroyed, it wouldn't get cleared out of the overlapping entity tracker because it was explicitly looking for entities with the BoundMaterial component on there.

To address this, we now always run material switchers as instantiation-time compiled ranges rather than evaluating the object channel during evaluation. This ensures that everything can be fixed up during instantiation ok and we can continue to use the built-in blending paths.

Additionally, a new OnPostSpawn event exists that allows the bound materials to re-validate its bound materials every frame and cause instantiation to be run if they need rebinding at all.

diff --git a/Engine/Source/Runtime/MovieScene/Private/EntitySystem/BuiltInComponentTypes.cpp b/Engine/Source/Runtime/MovieScene/Private/EntitySystem/BuiltInComponentTypes.cpp
--- a/Engine/Source/Runtime/MovieScene/Private/EntitySystem/BuiltInComponentTypes.cpp
+++ b/Engine/Source/Runtime/MovieScene/Private/EntitySystem/BuiltInComponentTypes.cpp
@@ -299,6 +299,7 @@ FBuiltInComponentTypes::FBuiltInComponentTypes()
 	{
 		// Object channel components should be duplicated to children
 		ComponentRegistry->Factories.DuplicateChildComponent(ObjectPathChannel);
+		ComponentRegistry->Factories.DuplicateChildComponent(ObjectResult);
 
 		// Object channel components need a time and result to evaluate
 		ComponentRegistry->Factories.DefineMutuallyInclusiveComponent(ObjectPathChannel, EvalTime);
diff --git a/Engine/Source/Runtime/MovieScene/Private/EntitySystem/MovieSceneEntitySystemLinker.cpp b/Engine/Source/Runtime/MovieScene/Private/EntitySystem/MovieSceneEntitySystemLinker.cpp
--- a/Engine/Source/Runtime/MovieScene/Private/EntitySystem/MovieSceneEntitySystemLinker.cpp
+++ b/Engine/Source/Runtime/MovieScene/Private/EntitySystem/MovieSceneEntitySystemLinker.cpp
@@ -161,6 +161,7 @@ void UMovieSceneEntitySystemLinker::SystemUnlinked(UMovieSceneEntitySystem* InSy
 	check(EntitySystemsByGlobalGraphID[GlobalID] == InSystem);
 	EntitySystemsByGlobalGraphID.RemoveAt(GlobalID);
 
+	Events.PostSpawnEvent.RemoveAll(InSystem);
 	Events.TagGarbage.RemoveAll(InSystem);
 	Events.CleanTaggedGarbage.RemoveAll(InSystem);
 	Events.AddReferencedObjects.RemoveAll(InSystem);
diff --git a/Engine/Source/Runtime/MovieScene/Private/EntitySystem/MovieSceneEntitySystemRunner.cpp b/Engine/Source/Runtime/MovieScene/Private/EntitySystem/MovieSceneEntitySystemRunner.cpp
--- a/Engine/Source/Runtime/MovieScene/Private/EntitySystem/MovieSceneEntitySystemRunner.cpp
+++ b/Engine/Source/Runtime/MovieScene/Private/EntitySystem/MovieSceneEntitySystemRunner.cpp
@@ -484,6 +484,9 @@ bool FMovieSceneEntitySystemRunner::GameThread_SpawnPhase(UMovieSceneEntitySyste
 		}
 	}
 
+
+	Linker->Events.PostSpawnEvent.Broadcast(Linker);
+
 	// --------------------------------------------------------------------------------------------------------------------------------------------
 	// Only run the instantiation phase if there is anything to instantiate. This must come after the spawn phase because new instantiations may
 	// be created during the spawn phase
diff --git a/Engine/Source/Runtime/MovieScene/Public/EntitySystem/MovieSceneEntityMutations.h b/Engine/Source/Runtime/MovieScene/Public/EntitySystem/MovieSceneEntityMutations.h
--- a/Engine/Source/Runtime/MovieScene/Public/EntitySystem/MovieSceneEntityMutations.h
+++ b/Engine/Source/Runtime/MovieScene/Public/EntitySystem/MovieSceneEntityMutations.h
@@ -28,7 +28,7 @@ struct FAddSingleMutation : IMovieSceneEntityMutation
 		: ComponentToAdd(InType)
 	{}
 
-	virtual void CreateMutation(FEntityManager* EntityManager, FComponentMask* InOutEntityComponentTypes) const override;
+	MOVIESCENE_API virtual void CreateMutation(FEntityManager* EntityManager, FComponentMask* InOutEntityComponentTypes) const override;
 };
 
 struct FRemoveSingleMutation : IMovieSceneEntityMutation
@@ -39,14 +39,14 @@ struct FRemoveSingleMutation : IMovieSceneEntityMutation
 		: ComponentToRemove(InType)
 	{}
 
-	virtual void CreateMutation(FEntityManager* EntityManager, FComponentMask* InOutEntityComponentTypes) const override;
+	MOVIESCENE_API virtual void CreateMutation(FEntityManager* EntityManager, FComponentMask* InOutEntityComponentTypes) const override;
 };
 
 struct FAddMultipleMutation : IMovieSceneEntityMutation
 {
 	FComponentMask MaskToAdd;
 
-	virtual void CreateMutation(FEntityManager* EntityManager, FComponentMask* InOutEntityComponentTypes) const override;
+	MOVIESCENE_API virtual void CreateMutation(FEntityManager* EntityManager, FComponentMask* InOutEntityComponentTypes) const override;
 };
 
 struct FRemoveMultipleMutation : IMovieSceneEntityMutation
@@ -54,7 +54,7 @@ struct FRemoveMultipleMutation : IMovieSceneEntityMutation
 
 	void RemoveComponent(FComponentTypeID InComponentType);
 
-	virtual void CreateMutation(FEntityManager* EntityManager, FComponentMask* InOutEntityComponentTypes) const override;
+	MOVIESCENE_API virtual void CreateMutation(FEntityManager* EntityManager, FComponentMask* InOutEntityComponentTypes) const override;
 
 private:
 
diff --git a/Engine/Source/Runtime/MovieScene/Public/EntitySystem/MovieSceneEntitySystemLinker.h b/Engine/Source/Runtime/MovieScene/Public/EntitySystem/MovieSceneEntitySystemLinker.h
--- a/Engine/Source/Runtime/MovieScene/Public/EntitySystem/MovieSceneEntitySystemLinker.h
+++ b/Engine/Source/Runtime/MovieScene/Public/EntitySystem/MovieSceneEntitySystemLinker.h
@@ -62,6 +62,7 @@ namespace MovieScene
 DECLARE_MULTICAST_DELEGATE_OneParam(FMovieSceneEntitySystemLinkerEvent, UMovieSceneEntitySystemLinker*);
 DECLARE_MULTICAST_DELEGATE_TwoParams(FMovieSceneEntitySystemLinkerAROEvent, UMovieSceneEntitySystemLinker*, FReferenceCollector&);
 DECLARE_MULTICAST_DELEGATE_TwoParams(FMovieSceneEntitySystemLinkerWorldEvent, UMovieSceneEntitySystemLinker*, UWorld*);
+DECLARE_MULTICAST_DELEGATE_OneParam(FMovieSceneEntitySystemLinkerPostSpawnEvent, UMovieSceneEntitySystemLinker*);
 
 UCLASS()
 class MOVIESCENE_API UMovieSceneEntitySystemLinker
@@ -320,11 +321,12 @@ public:
 
 	struct
 	{
-		FMovieSceneEntitySystemLinkerEvent      TagGarbage;
-		FMovieSceneEntitySystemLinkerEvent      CleanTaggedGarbage;
-		FMovieSceneEntitySystemLinkerAROEvent   AddReferencedObjects;
-		FMovieSceneEntitySystemLinkerEvent      AbandonLinker;
-		FMovieSceneEntitySystemLinkerWorldEvent CleanUpWorld;
+		FMovieSceneEntitySystemLinkerPostSpawnEvent PostSpawnEvent;
+		FMovieSceneEntitySystemLinkerEvent          TagGarbage;
+		FMovieSceneEntitySystemLinkerEvent          CleanTaggedGarbage;
+		FMovieSceneEntitySystemLinkerAROEvent       AddReferencedObjects;
+		FMovieSceneEntitySystemLinkerEvent          AbandonLinker;
+		FMovieSceneEntitySystemLinkerWorldEvent     CleanUpWorld;
 	} Events;
 
 private:
diff --git a/Engine/Source/Runtime/MovieScene/Public/EntitySystem/MovieSceneOverlappingEntityTracker.h b/Engine/Source/Runtime/MovieScene/Public/EntitySystem/MovieSceneOverlappingEntityTracker.h
--- a/Engine/Source/Runtime/MovieScene/Public/EntitySystem/MovieSceneOverlappingEntityTracker.h
+++ b/Engine/Source/Runtime/MovieScene/Public/EntitySystem/MovieSceneOverlappingEntityTracker.h
@@ -165,9 +165,16 @@ struct TOverlappingEntityTrackerImpl
 	/**
 	 * Update this tracker by (re)linking the specified allocation
 	 */
-	void VisitLinkedAllocation(const FEntityAllocation* Allocation, TRead<InputKeyTypes>... ReadKeys)
+	void VisitActiveAllocation(const FEntityAllocation* Allocation, TComponentPtr<const InputKeyTypes>... ReadKeys)
 	{
-		VisitLinkedAllocationImpl(Allocation, ReadKeys...);
+		VisitActiveAllocationImpl(Allocation, ReadKeys...);
+	}
+	/**
+	 * Update this tracker by (re)linking the specified allocation
+	 */
+	void VisitLinkedAllocation(const FEntityAllocation* Allocation, TComponentPtr<const InputKeyTypes>... ReadKeys)
+	{
+		VisitActiveAllocationImpl(Allocation, ReadKeys...);
 	}
 
 	/**
@@ -336,7 +343,7 @@ struct TOverlappingEntityTrackerImpl
 
 protected:
 
-	void VisitLinkedAllocationImpl(const FEntityAllocation* Allocation, TRead<InputKeyTypes>... Keys)
+	void VisitActiveAllocationImpl(const FEntityAllocation* Allocation, TComponentPtr<const InputKeyTypes>... Keys)
 	{
 		check(bIsInitialized);
 
@@ -344,11 +351,12 @@ protected:
 
 		const FMovieSceneEntityID* EntityIDs = Allocation->GetRawEntityIDs();
 
+		const bool bNeedsLink = Allocation->HasComponent(FBuiltInComponentTypes::Get()->Tags.NeedsLink);
 		if (Allocation->HasComponent(FBuiltInComponentTypes::Get()->Tags.RestoreState))
 		{
 			for (int32 Index = 0; Index < Num; ++Index)
 			{
-				const uint16 OutputIndex = MakeOutput(EntityIDs[Index], TOverlappingEntityInput<InputKeyTypes...>(Keys[Index]...));
+				const uint16 OutputIndex = MakeOutput(EntityIDs[Index], TOverlappingEntityInput<InputKeyTypes...>(Keys[Index]...), bNeedsLink);
 				Outputs[OutputIndex].Aggregate.bNeedsRestoration = true;
 			}
 		}
@@ -356,7 +364,7 @@ protected:
 		{
 			for (int32 Index = 0; Index < Num; ++Index)
 			{
-				MakeOutput(EntityIDs[Index], TOverlappingEntityInput<InputKeyTypes...>(Keys[Index]...));
+				MakeOutput(EntityIDs[Index], TOverlappingEntityInput<InputKeyTypes...>(Keys[Index]...), bNeedsLink);
 			}
 		}
 	}
@@ -374,17 +382,42 @@ protected:
 		}
 	}
 
-	uint16 MakeOutput(FMovieSceneEntityID EntityID, ParamType InKey)
+	uint16 MakeOutput(FMovieSceneEntityID EntityID, ParamType InKey, bool bAlwaysInvalidate)
 	{
-		// If this entity was already assigned an output, clear it
-		ClearOutputByEntity(EntityID);
+		const uint16 PreviousOutputIndex = FindOutputByEntity(EntityID);
+		const uint16 DesiredOutputIndex  = CreateOutputByKey(InKey);
 
-		const uint16 Output = CreateOutputByKey(InKey);
+		if (PreviousOutputIndex == DesiredOutputIndex)
+		{
+			if (bAlwaysInvalidate)
+			{
+				// Previous output is now invalidated since we're removing this entity
+				InvalidatedOutputs.PadToNum(DesiredOutputIndex + 1, false);
+				InvalidatedOutputs[DesiredOutputIndex] = true;
+			}
 
-		EntityToOutput.Add(EntityID, Output);
-		OutputToEntity.Add(Output, EntityID);
+			return DesiredOutputIndex;
+		}
 
-		return Output;
+		if (PreviousOutputIndex != NO_OUTPUT)
+		{
+			// Previous output is now invalidated since we're removing this entity
+			InvalidatedOutputs.PadToNum(PreviousOutputIndex + 1, false);
+			InvalidatedOutputs[PreviousOutputIndex] = true;
+
+			// Remove the entitiy's contribution from the previous output
+			OutputToEntity.Remove(PreviousOutputIndex, EntityID);
+			EntityToOutput.Remove(EntityID);
+		}
+
+		// Invalidate the new output
+		InvalidatedOutputs.PadToNum(DesiredOutputIndex + 1, false);
+		InvalidatedOutputs[DesiredOutputIndex] = true;
+
+		EntityToOutput.Add(EntityID, DesiredOutputIndex);
+		OutputToEntity.Add(DesiredOutputIndex, EntityID);
+
+		return DesiredOutputIndex;
 	}
 
 	uint16 CreateOutputByKey(ParamType Key)
@@ -392,8 +425,6 @@ protected:
 		const uint16 ExistingOutput = FindOutputByKey(Key);
 		if (ExistingOutput != NO_OUTPUT)
 		{
-			InvalidatedOutputs.PadToNum(ExistingOutput + 1, false);
-			InvalidatedOutputs[ExistingOutput] = true;
 			return ExistingOutput;
 		}
 
@@ -405,9 +436,6 @@ protected:
 		NewOutputs.PadToNum(NewOutput + 1, false);
 		NewOutputs[NewOutput] = true;
 
-		InvalidatedOutputs.PadToNum(NewOutput + 1, false);
-		InvalidatedOutputs[NewOutput] = true;
-
 		KeyToOutput.Add(Key, NewOutput);
 		return NewOutput;
 	}
diff --git a/Engine/Source/Runtime/MovieSceneTracks/Private/Sections/MovieScenePrimitiveMaterialSection.cpp b/Engine/Source/Runtime/MovieSceneTracks/Private/Sections/MovieScenePrimitiveMaterialSection.cpp
--- a/Engine/Source/Runtime/MovieSceneTracks/Private/Sections/MovieScenePrimitiveMaterialSection.cpp
+++ b/Engine/Source/Runtime/MovieSceneTracks/Private/Sections/MovieScenePrimitiveMaterialSection.cpp
@@ -8,6 +8,10 @@
 #include "EntitySystem/BuiltInComponentTypes.h"
 #include "MovieSceneTracksComponentTypes.h"
 
+namespace UE::MovieScene
+{
+	uint32 PrimitiveMaterialParameterSectionDefaultValueID = uint32(-1);
+}
 
 UMovieScenePrimitiveMaterialSection::UMovieScenePrimitiveMaterialSection(const FObjectInitializer& ObjInit)
 	: Super(ObjInit)
@@ -25,6 +29,60 @@ UMovieScenePrimitiveMaterialSection::UMovieScenePrimitiveMaterialSection(const F
 #endif
 }
 
+bool UMovieScenePrimitiveMaterialSection::PopulateEvaluationFieldImpl(const TRange<FFrameNumber>& EffectiveRange, const FMovieSceneEvaluationFieldEntityMetaData& InMetaData, FMovieSceneEntityComponentFieldBuilder* OutFieldBuilder)
+{
+	using namespace UE::MovieScene;
+
+	TMovieSceneChannelData<const FMovieSceneObjectPathChannelKeyValue> ChannelData = MaterialChannel.GetData();
+
+	TArrayView<const FFrameNumber> Times = ChannelData.GetTimes();
+	TArrayView<const FMovieSceneObjectPathChannelKeyValue> Values = ChannelData.GetValues();
+
+	const int32 MetaDataIndex = OutFieldBuilder->AddMetaData(InMetaData);
+
+	if (Times.Num() == 0)
+	{
+		const int32 EntityIndex = OutFieldBuilder->FindOrAddEntity(this, PrimitiveMaterialParameterSectionDefaultValueID);
+		OutFieldBuilder->AddPersistentEntity(EffectiveRange, EntityIndex, MetaDataIndex);
+		return true;
+	}
+
+	// Add entities for each key range in the channel
+	TRangeBound<FFrameNumber> StartBound = EffectiveRange.GetLowerBound();
+
+	// Find the first effective key
+	int32 ValueIndex = FMath::Min(StartBound.IsOpen() ? 0 : Algo::UpperBound(Times, UE::MovieScene::DiscreteInclusiveLower(StartBound)), Times.Num()-1);
+	if (!Times.IsValidIndex(ValueIndex))
+	{
+		return true;
+	}
+
+	int32 LastEntityIndex = OutFieldBuilder->FindOrAddEntity(this, ValueIndex);
+
+	for ( ; ValueIndex < Times.Num(); ++ValueIndex)
+	{
+		if (!EffectiveRange.Contains(Times[ValueIndex]))
+		{
+			break;
+		}
+
+		// Add the last range to the tree
+		TRange<FFrameNumber> Range(StartBound, TRangeBound<FFrameNumber>::Exclusive(Times[ValueIndex]));
+		OutFieldBuilder->AddPersistentEntity(Range, LastEntityIndex, MetaDataIndex);
+
+		StartBound = TRangeBound<FFrameNumber>::Inclusive(Times[ValueIndex]);
+		LastEntityIndex = OutFieldBuilder->FindOrAddEntity(this, ValueIndex);
+	}
+
+	TRange<FFrameNumber> TailRange(StartBound, EffectiveRange.GetUpperBound());
+	if (!TailRange.IsEmpty())
+	{
+		OutFieldBuilder->AddPersistentEntity(TailRange, LastEntityIndex, MetaDataIndex);
+	}
+
+	return true;
+}
+
 void UMovieScenePrimitiveMaterialSection::ImportEntityImpl(UMovieSceneEntitySystemLinker* Linker, const FEntityImportParams& ImportParams, FImportedEntity* OutImportedEntity)
 {
 	using namespace UE::MovieScene;
@@ -37,12 +95,25 @@ void UMovieScenePrimitiveMaterialSection::ImportEntityImpl(UMovieSceneEntitySyst
 	FBuiltInComponentTypes* BuiltInComponentTypes = FBuiltInComponentTypes::Get();
 	FMovieSceneTracksComponentTypes* TracksComponentTypes = FMovieSceneTracksComponentTypes::Get();
 
+	FMovieSceneObjectPathChannelKeyValue ValueToImport;
+
+	if (ImportParams.EntityID == PrimitiveMaterialParameterSectionDefaultValueID)
+	{
+		// Default value
+		ValueToImport = MaterialChannel.GetDefault();
+	}
+	else
+	{
+		const int32 ChannelValueIndex = static_cast<int32>(ImportParams.EntityID);
+		ValueToImport = MaterialChannel.GetData().GetValues()[ChannelValueIndex];
+	}
+
 	FGuid ObjectBindingID = ImportParams.GetObjectBindingID();
 
 	OutImportedEntity->AddBuilder(
 		FEntityBuilder()
 		.AddConditional(BuiltInComponentTypes->GenericObjectBinding, ObjectBindingID, ObjectBindingID.IsValid())
-		.Add(BuiltInComponentTypes->ObjectPathChannel, &MaterialChannel)
+		.Add(BuiltInComponentTypes->ObjectResult, ValueToImport.Get())
 		.Add(TracksComponentTypes->ComponentMaterialIndex, MaterialIndex)
 	);
 }
\ No newline at end of file
diff --git a/Engine/Source/Runtime/MovieSceneTracks/Private/Systems/MovieSceneComponentMaterialSystem.cpp b/Engine/Source/Runtime/MovieSceneTracks/Private/Systems/MovieSceneComponentMaterialSystem.cpp
--- a/Engine/Source/Runtime/MovieSceneTracks/Private/Systems/MovieSceneComponentMaterialSystem.cpp
+++ b/Engine/Source/Runtime/MovieSceneTracks/Private/Systems/MovieSceneComponentMaterialSystem.cpp
@@ -114,6 +114,11 @@ void UMovieSceneComponentMaterialSystem::OnLink()
 	SystemImpl.OnLink(Linker, FBuiltInComponentTypes::Get()->BoundObject, FMovieSceneTracksComponentTypes::Get()->ComponentMaterialIndex);
 }
 
+void UMovieSceneComponentMaterialSystem::OnUnlink()
+{
+	SystemImpl.OnUnlink(Linker);
+}
+
 void UMovieSceneComponentMaterialSystem::OnRun(FSystemTaskPrerequisites& InPrerequisites, FSystemSubsequentTasks& Subsequents)
 {
 	using namespace UE::MovieScene;
diff --git a/Engine/Source/Runtime/MovieSceneTracks/Private/Systems/MovieSceneMaterialParameterCollectionSystem.cpp b/Engine/Source/Runtime/MovieSceneTracks/Private/Systems/MovieSceneMaterialParameterCollectionSystem.cpp
--- a/Engine/Source/Runtime/MovieSceneTracks/Private/Systems/MovieSceneMaterialParameterCollectionSystem.cpp
+++ b/Engine/Source/Runtime/MovieSceneTracks/Private/Systems/MovieSceneMaterialParameterCollectionSystem.cpp
@@ -192,57 +192,41 @@ void UMovieSceneMaterialParameterCollectionSystem::OnRun(FSystemTaskPrerequisite
 					UMaterialParameterCollectionInstance* Instance = World->GetParameterCollectionInstance(Collection);
 					OutBoundMaterials[Index] = Instance;
 
-					if (!ensureAlwaysMsgf(Instance != nullptr,
+#if !UE_BUILD_SHIPPING && !UE_BUILD_TEST
+					if (ensureAlwaysMsgf(Instance != nullptr,
 						TEXT("Unable to create MPC instance for %s with World %s. Material parameter collection tracks will not function."),
 						*Collection->GetName(), *World->GetName()))
-					{
-						// Remove the entity if it is null to avoid having to check for null downstream
-						InvalidEntities.Add(Allocation->GetRawEntityIDs()[Index]);
-					}
-					else 
 					{
 						if (ScalarParameterNames)
 						{
 							FName Name = ScalarParameterNames[Index];
 							if (Collection->GetScalarParameterByName(Name) == nullptr)
 							{
-								InvalidEntities.Add(Allocation->GetRawEntityIDs()[Index]);
-
-#if !UE_BUILD_SHIPPING && !UE_BUILD_TEST
 								if (!Instance->bLoggedMissingParameterWarning)
 								{
 									MissingParameters.FindOrAdd(MakeTuple(Instance, Player)).Add(Name.ToString());
 								}
-#endif
 							}
 						}
-						else 
+						else if (VectorParameterNames || ColorParameterNames)
 						{
 							FName Name = VectorParameterNames ? VectorParameterNames[Index] : ColorParameterNames[Index];
 							if (Collection->GetVectorParameterByName(Name) == nullptr)
 							{
-								InvalidEntities.Add(Allocation->GetRawEntityIDs()[Index]);
-
-#if !UE_BUILD_SHIPPING && !UE_BUILD_TEST
 								if (!Instance->bLoggedMissingParameterWarning)
 								{
 									MissingParameters.FindOrAdd(MakeTuple(Instance, Player)).Add(Name.ToString());
 								}
-#endif
 							}
 						}
 					}
+#endif
 				}
 			}
 		}
 
 		void Cleanup(UMovieSceneEntitySystemLinker* InLinker)
 		{
-			for (FMovieSceneEntityID InvalidEntity : InvalidEntities)
-			{
-				InLinker->EntityManager.RemoveComponent(InvalidEntity, TracksComponents->BoundMaterial);
-			}
-
 #if !UE_BUILD_SHIPPING && !UE_BUILD_TEST
 
 			for (TPair<TTuple<UMaterialParameterCollectionInstance*, IMovieScenePlayer*>, TArray<FString>>& Pair : MissingParameters)
@@ -267,15 +251,12 @@ void UMovieSceneMaterialParameterCollectionSystem::OnRun(FSystemTaskPrerequisite
 		FMovieSceneTracksComponentTypes* TracksComponents;
 
 		mutable TMap<TTuple<UMaterialParameterCollectionInstance*, IMovieScenePlayer*>, TArray<FString>> MissingParameters;
-
-		/** Track entities that we added a BoundMaterial component to but were unable to resolve so we can remove the component. This should happen rarely */
-		mutable TArray<FMovieSceneEntityID> InvalidEntities;
 	};
 
 
 	// Only mutate things that are tagged as requiring linking
 	FEntityComponentFilter Filter;
-	Filter.All({TracksComponents->MPC, BuiltInComponents->InstanceHandle, BuiltInComponents->Tags.NeedsLink});
+	Filter.All({ TracksComponents->MPC, BuiltInComponents->InstanceHandle, BuiltInComponents->Tags.NeedsLink });
 
 	// Initialize bound dynamic materials (for material collection parameters)
 	FAddMPCMutation BindMaterialsMutation(Linker);
diff --git a/Engine/Source/Runtime/MovieSceneTracks/Private/Systems/MovieSceneMaterialParameterSystem.cpp b/Engine/Source/Runtime/MovieSceneTracks/Private/Systems/MovieSceneMaterialParameterSystem.cpp
--- a/Engine/Source/Runtime/MovieSceneTracks/Private/Systems/MovieSceneMaterialParameterSystem.cpp
+++ b/Engine/Source/Runtime/MovieSceneTracks/Private/Systems/MovieSceneMaterialParameterSystem.cpp
@@ -14,11 +14,11 @@
 namespace UE::MovieScene
 {
 
-bool GMaterialParameterBlending = false;
+bool GMaterialParameterBlending = true;
 FAutoConsoleVariableRef CVarMaterialParameterBlending(
 	TEXT("Sequencer.MaterialParameterBlending"),
 	GMaterialParameterBlending,
-	TEXT("(Default: false) Defines whether material parameter blending should be enabled or not.\n"),
+	TEXT("(Default: true) Defines whether material parameter blending should be enabled or not.\n"),
 	ECVF_Default
 );
 
@@ -101,11 +101,6 @@ struct FOverlappingMaterialParameterHandler
 
 	void UpdateOutput(UObject* BoundMaterial, FName ParameterName, TArrayView<const FMovieSceneEntityID> Inputs, FAnimatedMaterialParameterInfo* Output, FEntityOutputAggregate Aggregate)
 	{
-		if (BoundMaterial == nullptr)
-		{
-			return;
-		}
-
 		using namespace UE::MovieScene;
 
 		FBuiltInComponentTypes* BuiltInComponents = FBuiltInComponentTypes::Get();
@@ -208,19 +203,21 @@ void UMovieSceneMaterialParameterSystem::OnLink()
 {
 	using namespace UE::MovieScene;
 
+	ScalarParameterTracker.Initialize(this);
+	VectorParameterTracker.Initialize(this);
+}
+
+void UMovieSceneMaterialParameterSystem::OnUnlink()
+{
+	using namespace UE::MovieScene;
+
 	// Always reset the float blender system on link to ensure that recycled systems are correctly initialized.
 	DoubleBlenderSystem = nullptr;
 
-	if (GMaterialParameterBlending)
-	{
-		FOverlappingMaterialParameterHandler Handler(this);
-
-		ScalarParameterTracker.Destroy(Handler);
-		VectorParameterTracker.Destroy(Handler);
+	FOverlappingMaterialParameterHandler Handler(this);
 
-		ScalarParameterTracker.Initialize(this);
-		VectorParameterTracker.Initialize(this);
-	}
+	ScalarParameterTracker.Destroy(Handler);
+	VectorParameterTracker.Destroy(Handler);
 }
 
 void UMovieSceneMaterialParameterSystem::OnRun(FSystemTaskPrerequisites& InPrerequisites, FSystemSubsequentTasks& Subsequents)
diff --git a/Engine/Source/Runtime/MovieSceneTracks/Private/Tracks/MovieSceneMaterialParameterCollectionTrack.cpp b/Engine/Source/Runtime/MovieSceneTracks/Private/Tracks/MovieSceneMaterialParameterCollectionTrack.cpp
--- a/Engine/Source/Runtime/MovieSceneTracks/Private/Tracks/MovieSceneMaterialParameterCollectionTrack.cpp
+++ b/Engine/Source/Runtime/MovieSceneTracks/Private/Tracks/MovieSceneMaterialParameterCollectionTrack.cpp
@@ -14,6 +14,8 @@ UMovieSceneMaterialParameterCollectionTrack::UMovieSceneMaterialParameterCollect
 {
 	BuiltInTreePopulationMode = ETreePopulationMode::Blended;
 
+	SupportedBlendTypes.Add(EMovieSceneBlendType::Absolute);
+
 #if WITH_EDITORONLY_DATA
 	TrackTint = FColor(64,192,64,65);
 #endif
diff --git a/Engine/Source/Runtime/MovieSceneTracks/Public/Sections/MovieScenePrimitiveMaterialSection.h b/Engine/Source/Runtime/MovieSceneTracks/Public/Sections/MovieScenePrimitiveMaterialSection.h
--- a/Engine/Source/Runtime/MovieSceneTracks/Public/Sections/MovieScenePrimitiveMaterialSection.h
+++ b/Engine/Source/Runtime/MovieSceneTracks/Public/Sections/MovieScenePrimitiveMaterialSection.h
@@ -18,6 +18,7 @@ public:
 
 	UMovieScenePrimitiveMaterialSection(const FObjectInitializer& ObjInit);
 
+	virtual bool PopulateEvaluationFieldImpl(const TRange<FFrameNumber>& EffectiveRange, const FMovieSceneEvaluationFieldEntityMetaData& InMetaData, FMovieSceneEntityComponentFieldBuilder* OutFieldBuilder) override;
 	virtual void ImportEntityImpl(UMovieSceneEntitySystemLinker* Linker, const FEntityImportParams& ImportParams, FImportedEntity* OutImportedEntity) override;
 
 	UPROPERTY()
diff --git a/Engine/Source/Runtime/MovieSceneTracks/Public/Systems/MovieSceneComponentMaterialSystem.h b/Engine/Source/Runtime/MovieSceneTracks/Public/Systems/MovieSceneComponentMaterialSystem.h
--- a/Engine/Source/Runtime/MovieSceneTracks/Public/Systems/MovieSceneComponentMaterialSystem.h
+++ b/Engine/Source/Runtime/MovieSceneTracks/Public/Systems/MovieSceneComponentMaterialSystem.h
@@ -78,6 +78,7 @@ public:
 private:
 
 	virtual void OnLink() override;
+	virtual void OnUnlink() override;
 	virtual void OnRun(FSystemTaskPrerequisites& InPrerequisites, FSystemSubsequentTasks& Subsequents) override;
 
 	virtual void SavePreAnimatedState(const FPreAnimationParameters& InParameters) override;
diff --git a/Engine/Source/Runtime/MovieSceneTracks/Public/Systems/MovieSceneMaterialParameterSystem.h b/Engine/Source/Runtime/MovieSceneTracks/Public/Systems/MovieSceneMaterialParameterSystem.h
--- a/Engine/Source/Runtime/MovieSceneTracks/Public/Systems/MovieSceneMaterialParameterSystem.h
+++ b/Engine/Source/Runtime/MovieSceneTracks/Public/Systems/MovieSceneMaterialParameterSystem.h
@@ -48,6 +48,7 @@ public:
 private:
 
 	virtual void OnLink() override;
+	virtual void OnUnlink() override;
 	virtual void OnRun(FSystemTaskPrerequisites& InPrerequisites, FSystemSubsequentTasks& Subsequents) override;
 
 	void OnInstantiation();
diff --git a/Engine/Source/Runtime/MovieSceneTracks/Public/Systems/MovieSceneMaterialSystem.h b/Engine/Source/Runtime/MovieSceneTracks/Public/Systems/MovieSceneMaterialSystem.h
--- a/Engine/Source/Runtime/MovieSceneTracks/Public/Systems/MovieSceneMaterialSystem.h
+++ b/Engine/Source/Runtime/MovieSceneTracks/Public/Systems/MovieSceneMaterialSystem.h
@@ -108,10 +108,13 @@ public:
 	TSharedPtr<MaterialParameterStorageType> MaterialParameterStorage;
 
 	void OnLink(UMovieSceneEntitySystemLinker* Linker, TComponentTypeID<RequiredComponents>... InRequiredComponents);
+	void OnUnlink(UMovieSceneEntitySystemLinker* Linker);
 	void OnRun(UMovieSceneEntitySystemLinker* Linker, TComponentTypeID<RequiredComponents>... InRequiredComponents, FSystemTaskPrerequisites& InPrerequisites, FSystemSubsequentTasks& Subsequents);
 
 	void SavePreAnimatedState(UMovieSceneEntitySystemLinker* Linker, TComponentTypeID<RequiredComponents>... InRequiredComponents, const IMovieScenePreAnimatedStateSystemInterface::FPreAnimationParameters& InParameters);
 
+	void OnPostSpawn(UMovieSceneEntitySystemLinker* InLinker, TComponentTypeID<RequiredComponents>... InRequiredComponents);
+
 protected:
 
 	UE::MovieScene::FEntityComponentFilter MaterialSwitcherFilter;
@@ -145,7 +148,7 @@ struct TApplyMaterialSwitchers
 template<typename AccessorType, typename... RequiredComponents>
 struct TInitializeBoundMaterials
 {
-	static void ForEachEntity(typename TCallTraits<RequiredComponents>::ParamType... Inputs, UObject*& OutDynamicMaterial)
+	static bool InitializeBoundMaterial(typename TCallTraits<RequiredComponents>::ParamType... Inputs, UObject*& OutDynamicMaterial)
 	{
 		AccessorType Accessor(Inputs...);
 
@@ -153,23 +156,65 @@ struct TInitializeBoundMaterials
 
 		if (!ExistingMaterial)
 		{
-			// The object no longer has a valid material assigned.
-			// Rather than null the bound material and cause all downstream systems to have to check for null
-			// We just leave it assigned to the previous MID, even if that won't have an effect any more
-			return;
+			OutDynamicMaterial = nullptr;
+			return true;
 		}
-		else if (UMaterialInstanceDynamic* MID = Cast<UMaterialInstanceDynamic>(ExistingMaterial))
+
+		if (UMaterialInstanceDynamic* MID = Cast<UMaterialInstanceDynamic>(ExistingMaterial))
 		{
-			OutDynamicMaterial = MID;
-			return;
+			if (OutDynamicMaterial != MID)
+			{
+				OutDynamicMaterial = MID;
+				return true;
+			}
+			return false;
 		}
 
-		if (OutDynamicMaterial && OutDynamicMaterial->IsA<UMaterialInstanceDynamic>())
+		UMaterialInstanceDynamic* CurrentMID = Cast<UMaterialInstanceDynamic>(OutDynamicMaterial);
+		if (CurrentMID && CurrentMID->Parent == ExistingMaterial)
 		{
-			return;
+			Accessor.SetMaterial(CurrentMID);
+			return false;
 		}
 		
-		OutDynamicMaterial = Accessor.CreateDynamicMaterial(ExistingMaterial);
+		UMaterialInstanceDynamic* NewMaterial = Accessor.CreateDynamicMaterial(ExistingMaterial);
+		OutDynamicMaterial = NewMaterial;
+		Accessor.SetMaterial(NewMaterial);
+		return true;
+	}
+
+	static void ForEachEntity(typename TCallTraits<RequiredComponents>::ParamType... Inputs, UObject*& OutDynamicMaterial)
+	{
+		InitializeBoundMaterial(Inputs..., OutDynamicMaterial);
+	}
+};
+
+
+template<typename AccessorType, typename... RequiredComponents>
+struct TReinitializeBoundMaterials
+{
+	UMovieSceneEntitySystemLinker* Linker;
+	TArray<FMovieSceneEntityID> ReboundMaterials;
+
+	TReinitializeBoundMaterials(UMovieSceneEntitySystemLinker* InLinker)
+		: Linker(InLinker)
+	{}
+
+	void ForEachEntity(FMovieSceneEntityID EntityID, typename TCallTraits<RequiredComponents>::ParamType... Inputs, UObject*& OutDynamicMaterial)
+	{
+		if (TInitializeBoundMaterials<AccessorType, RequiredComponents...>::InitializeBoundMaterial(Inputs..., OutDynamicMaterial))
+		{
+			ReboundMaterials.Add(EntityID);
+		}
+	}
+
+	void PostTask()
+	{
+		FBuiltInComponentTypes* BuiltInComponents = FBuiltInComponentTypes::Get();
+		for (FMovieSceneEntityID EntityID : ReboundMaterials)
+		{
+			Linker->EntityManager.AddComponent(EntityID, BuiltInComponents->Tags.NeedsLink);
+		}
 	}
 };
 
@@ -201,41 +246,13 @@ struct TAddBoundMaterialMutationImpl<AccessorType, TIntegerSequence<int, Indices
 		for (int32 Index = 0; Index < Num; ++Index)
 		{
 			OutBoundMaterials[Index] = nullptr;
-			TInitializeBoundMaterials<AccessorType, RequiredComponents...>::ForEachEntity(InRequiredComponents[Index]..., OutBoundMaterials[Index]);
-
-			// @todo: We could remove the entity if it is null to avoid having to check for null downstream
-			//        but that is a problem for some widgets that play animations in PreConstruct _before_
-			//        the material is able to resolve correctly.
-			// if (!ensureAlwaysMsgf(OutBoundMaterials[Index] != nullptr,
-			//	TEXT("Unable to resolve material for %s. Material parameter tracks on this object will not function until it is resolved."),
-			//	*AccessorType(InRequiredComponents[Index]...).ToString()))
-			//{
-			//	
-			//	InvalidEntities.Add(Allocation->GetRawEntityIDs()[Index]);
-			//}
 		}
 	}
-
-	void RemoveInvalidBoundMaterials(UMovieSceneEntitySystemLinker* InLinker)
-	{
-#if 0
-		// Currently disabled due to the comment in InitializeAllocation. Another approach would be to
-		// Tag these differently, but that would require us mutating the entity manager during evaluation
-		// if the material _does_ resolve in the future, and mutation during eval is not supported
-		for (FMovieSceneEntityID InvalidEntity : InvalidEntities)
-		{
-			InLinker->EntityManager.RemoveComponent(InvalidEntity, TracksComponents->BoundMaterial);
-		}
-#endif
-	}
 private:
 
 	FBuiltInComponentTypes* BuiltInComponents;
 	FMovieSceneTracksComponentTypes* TracksComponents;
 
-	/** Track entities that we added a BoundMaterial component to but were unable to resolve so we can remove the component. This should happen rarely */
-	mutable TArray<FMovieSceneEntityID> InvalidEntities;
-
 	TTuple<TComponentTypeID<RequiredComponents>...> ComponentTypes;
 };
 
@@ -256,6 +273,34 @@ void TMovieSceneMaterialSystem<AccessorType, RequiredComponents...>::OnLink(UMov
 	MaterialParameterFilter.Reset();
 	MaterialParameterFilter.All({ InRequiredComponents... });
 	MaterialParameterFilter.Any({ TracksComponents->ScalarParameterName, TracksComponents->ColorParameterName, TracksComponents->VectorParameterName });
+
+	Linker->Events.PostSpawnEvent.AddRaw(this, &TMovieSceneMaterialSystem<AccessorType, RequiredComponents...>::OnPostSpawn, InRequiredComponents...);
+}
+
+template<typename AccessorType, typename... RequiredComponents>
+void TMovieSceneMaterialSystem<AccessorType, RequiredComponents...>::OnUnlink(UMovieSceneEntitySystemLinker* Linker)
+{
+	Linker->Events.PostSpawnEvent.RemoveAll(this);
+}
+
+template<typename AccessorType, typename... RequiredComponents>
+void TMovieSceneMaterialSystem<AccessorType, RequiredComponents...>::OnPostSpawn(UMovieSceneEntitySystemLinker* InLinker, TComponentTypeID<RequiredComponents>... InRequiredComponents)
+{
+	using namespace UE::MovieScene;
+
+	FBuiltInComponentTypes*          BuiltInComponents = FBuiltInComponentTypes::Get();
+	FMovieSceneTracksComponentTypes* TracksComponents  = FMovieSceneTracksComponentTypes::Get();
+
+	TReinitializeBoundMaterials<AccessorType, RequiredComponents...> ReinitializeBoundMaterialsTask(InLinker);
+
+	// Reinitialize bound dynamic materials, adding NeedsLink during PostTask to any that changed
+	// This will cause the instantiation phase to be re-run for these entities (and any other new or expired ones)
+	FEntityTaskBuilder()
+	.ReadEntityIDs()
+	.ReadAllOf(InRequiredComponents...)
+	.Write(TracksComponents->BoundMaterial)
+	.FilterNone({ BuiltInComponents->Tags.NeedsUnlink })
+	.RunInline_PerEntity(&InLinker->EntityManager, ReinitializeBoundMaterialsTask);
 }
 
 template<typename AccessorType, typename... RequiredComponents>
@@ -275,35 +320,34 @@ void TMovieSceneMaterialSystem<AccessorType, RequiredComponents...>::OnRun(UMovi
 	ESystemPhase CurrentPhase = ActiveRunner->GetCurrentPhase();
 	if (CurrentPhase == ESystemPhase::Instantiation)
 	{
-		// Only mutate things that are tagged as requiring linking
-		FEntityComponentFilter Filter(MaterialParameterFilter);
-		Filter.All({ BuiltInComponents->Tags.NeedsLink});
-
-		// Initialize bound dynamic materials (for material parameters)
-		using MutationType = TAddBoundMaterialMutationImpl<AccessorType, TMakeIntegerSequence<int, sizeof...(RequiredComponents)>, RequiredComponents...>;
-		MutationType BindMaterialsMutation(InRequiredComponents...);
-
-		Linker->EntityManager.MutateAll(Filter, BindMaterialsMutation);
-		BindMaterialsMutation.RemoveInvalidBoundMaterials(Linker);
-	}
-	else if (CurrentPhase == ESystemPhase::Evaluation)
-	{
-		using ApplyMaterialSwitchers = TApplyMaterialSwitchers<AccessorType, RequiredComponents...>;
-		using InitializeBoundMaterials = TInitializeBoundMaterials<AccessorType, RequiredComponents...>;
-
+		// --------------------------------------------------------------------------------------
 		// Apply material switchers
+		TApplyMaterialSwitchers<AccessorType, RequiredComponents...> ApplyMaterialSwitchers;
+
 		FEntityTaskBuilder()
 		.ReadAllOf(InRequiredComponents...)
 		.Read(BuiltInComponents->ObjectResult)
-		.SetDesiredThread(Linker->EntityManager.GetDispatchThread())
-		.template Dispatch_PerEntity<ApplyMaterialSwitchers>(&Linker->EntityManager, InPrerequisites, &Subsequents);
+		.FilterAll({ BuiltInComponents->Tags.NeedsLink })
+		.RunInline_PerEntity(&Linker->EntityManager, ApplyMaterialSwitchers);
+
+		// --------------------------------------------------------------------------------------
+		// Add bound materials for any NeedsLink entities that have material parameters
+		FEntityComponentFilter Filter(MaterialParameterFilter);
+		Filter.All({ BuiltInComponents->Tags.NeedsLink });
+
+		using MutationType = TAddBoundMaterialMutationImpl<AccessorType, TMakeIntegerSequence<int, sizeof...(RequiredComponents)>, RequiredComponents...>;
+		Linker->EntityManager.MutateAll(Filter, MutationType(InRequiredComponents...));
+
+		// --------------------------------------------------------------------------------------
+		// (Re)initialize bound materials for any NeedsLink materials
+		TReinitializeBoundMaterials<AccessorType, RequiredComponents...> ReinitializeBoundMaterialsTask(Linker);
 
-		// Initialize bound dynamic materials
 		FEntityTaskBuilder()
+		.ReadEntityIDs()
 		.ReadAllOf(InRequiredComponents...)
 		.Write(TracksComponents->BoundMaterial)
-		.SetDesiredThread(Linker->EntityManager.GetDispatchThread())
-		.template Dispatch_PerEntity<InitializeBoundMaterials>(&Linker->EntityManager, InPrerequisites, &Subsequents);
+		.FilterNone({ BuiltInComponents->Tags.NeedsUnlink })
+		.RunInline_PerEntity(&Linker->EntityManager, ReinitializeBoundMaterialsTask);
 	}
 }
 
diff --git a/Engine/Source/Runtime/UMG/Private/Animation/MovieSceneWidgetMaterialSystem.cpp b/Engine/Source/Runtime/UMG/Private/Animation/MovieSceneWidgetMaterialSystem.cpp
--- a/Engine/Source/Runtime/UMG/Private/Animation/MovieSceneWidgetMaterialSystem.cpp
+++ b/Engine/Source/Runtime/UMG/Private/Animation/MovieSceneWidgetMaterialSystem.cpp
@@ -110,6 +110,11 @@ void UMovieSceneWidgetMaterialSystem::OnLink()
 	SystemImpl.OnLink(Linker, BuiltInComponents->BoundObject, WidgetComponents->WidgetMaterialPath);
 }
 
+void UMovieSceneWidgetMaterialSystem::OnUnlink()
+{
+	SystemImpl.OnUnlink(Linker);
+}
+
 void UMovieSceneWidgetMaterialSystem::OnRun(FSystemTaskPrerequisites& InPrerequisites, FSystemSubsequentTasks& Subsequents)
 {
 	using namespace UE::MovieScene;
diff --git a/Engine/Source/Runtime/UMG/Public/Animation/MovieSceneWidgetMaterialSystem.h b/Engine/Source/Runtime/UMG/Public/Animation/MovieSceneWidgetMaterialSystem.h
--- a/Engine/Source/Runtime/UMG/Public/Animation/MovieSceneWidgetMaterialSystem.h
+++ b/Engine/Source/Runtime/UMG/Public/Animation/MovieSceneWidgetMaterialSystem.h
@@ -84,6 +84,7 @@ public:
 private:
 
 	virtual void OnLink() override;
+	virtual void OnUnlink() override;
 	virtual void OnRun(FSystemTaskPrerequisites& InPrerequisites, FSystemSubsequentTasks& Subsequents) override;
 
 	virtual void SavePreAnimatedState(const FPreAnimationParameters& InParameters) override;
