Sequencer: Added a new tag for determining when a bound material has been changed, exclude ImportedEntities from processing with the material systems

  - In addition to the new tag, a component mask exists on FBuiltInComponentTypes called RequiresInstantiationMask. Set bits within this mask denote when instantiation should run, rather than being hard coded to NeedsLink. This allows us to keep notions of instantiation compatmentalized, and helps us not overload the NeedsLink terminology. For bound materials this means we can add our own tag while still triggering instantiation. Any components that are set on RequiresInstantiationMask will be removed from any entities at the end of instantiation.

  - ImportedEntities are now excluded from processing in the material systems since these entities' IDs can change when they are modified, and should not have been receiving BoundMaterial components anyway.

diff --git a/Engine/Source/Runtime/MovieScene/Private/EntitySystem/BuiltInComponentTypes.cpp b/Engine/Source/Runtime/MovieScene/Private/EntitySystem/BuiltInComponentTypes.cpp
--- a/Engine/Source/Runtime/MovieScene/Private/EntitySystem/BuiltInComponentTypes.cpp
+++ b/Engine/Source/Runtime/MovieScene/Private/EntitySystem/BuiltInComponentTypes.cpp
@@ -320,6 +320,8 @@ FBuiltInComponentTypes::FBuiltInComponentTypes()
 		ComponentRegistry->Factories.DefineMutuallyInclusiveComponent(EvaluationHook, EvalTime);
 		ComponentRegistry->Factories.DefineMutuallyInclusiveComponent(EvaluationHook, EvaluationHookFlags);
 	}
+
+	RequiresInstantiationMask.Set(Tags.NeedsLink);
 }
 
 FBuiltInComponentTypes::~FBuiltInComponentTypes()
diff --git a/Engine/Source/Runtime/MovieScene/Private/EntitySystem/MovieSceneEntitySystemRunner.cpp b/Engine/Source/Runtime/MovieScene/Private/EntitySystem/MovieSceneEntitySystemRunner.cpp
--- a/Engine/Source/Runtime/MovieScene/Private/EntitySystem/MovieSceneEntitySystemRunner.cpp
+++ b/Engine/Source/Runtime/MovieScene/Private/EntitySystem/MovieSceneEntitySystemRunner.cpp
@@ -490,7 +490,7 @@ bool FMovieSceneEntitySystemRunner::GameThread_SpawnPhase(UMovieSceneEntitySyste
 	// --------------------------------------------------------------------------------------------------------------------------------------------
 	// Only run the instantiation phase if there is anything to instantiate. This must come after the spawn phase because new instantiations may
 	// be created during the spawn phase
-	const bool bAnyPending = Linker->EntityManager.ContainsComponent(FBuiltInComponentTypes::Get()->Tags.NeedsLink) || InstanceRegistry->HasInvalidatedBindings();
+	const bool bAnyPending = Linker->EntityManager.ContainsAnyComponent(FBuiltInComponentTypes::Get()->RequiresInstantiationMask) || InstanceRegistry->HasInvalidatedBindings();
 	if (bInstantiationDirty == false && bAnyPending == false)
 	{
 		FlushState &= ~ERunnerFlushState::Instantiation;
@@ -541,9 +541,12 @@ bool FMovieSceneEntitySystemRunner::GameThread_PostInstantiation(UMovieSceneEnti
 
 	// Nothing needs linking, caching or restoring any more
 	FRemoveMultipleMutation Mutation;
-	Mutation.RemoveComponent(BuiltInComponentTypes->Tags.NeedsLink);
+	Mutation.MaskToRemove = BuiltInComponentTypes->RequiresInstantiationMask;
+	// Inverse the filter because MaskToRemove is applied as a binary AND against all matching allocations.
+	// Therefore, any set bits in RequiresInstantiationMask will be removed from components
+	Mutation.MaskToRemove.BitwiseNOT();
 
-	FEntityComponentFilter Filter = FEntityComponentFilter().Any({ BuiltInComponentTypes->Tags.NeedsLink });
+	FEntityComponentFilter Filter = FEntityComponentFilter().Any(BuiltInComponentTypes->RequiresInstantiationMask);
 	EntityManager.MutateAll(Filter, Mutation);
 
 	// Free anything that has been unlinked
diff --git a/Engine/Source/Runtime/MovieScene/Public/EntitySystem/BuiltInComponentTypes.h b/Engine/Source/Runtime/MovieScene/Public/EntitySystem/BuiltInComponentTypes.h
--- a/Engine/Source/Runtime/MovieScene/Public/EntitySystem/BuiltInComponentTypes.h
+++ b/Engine/Source/Runtime/MovieScene/Public/EntitySystem/BuiltInComponentTypes.h
@@ -230,6 +230,12 @@ public:
 
 public:
 
+	/**
+	 * Component mask where set bits denot component types that should trigger instantiation when present.
+	 * After instantiation, these components will be removed from any entities to prevent instantiation being run constantly
+	 */
+	FComponentMask RequiresInstantiationMask;
+
 	TComponentTypeID<FMovieSceneEntityID> ParentEntity;
 
 	TComponentTypeID<UObject*>            BoundObject;
diff --git a/Engine/Source/Runtime/MovieScene/Public/EntitySystem/MovieSceneEntityManager.h b/Engine/Source/Runtime/MovieScene/Public/EntitySystem/MovieSceneEntityManager.h
--- a/Engine/Source/Runtime/MovieScene/Public/EntitySystem/MovieSceneEntityManager.h
+++ b/Engine/Source/Runtime/MovieScene/Public/EntitySystem/MovieSceneEntityManager.h
@@ -699,6 +699,17 @@ public:
 		return Algo::AnyOf(ComponentTypeIDs, [&Mask](FComponentTypeID In){ return Mask.Contains(In); });
 	}
 
+	/**
+	 * Check whether any entity in this manager has any of the specified components
+	 *
+	 * @param ComponentTypeIDs     The types of the component that are being tested for
+	 * @return true if any one of the specified components exist anywhere in this manager, false otherwise
+	 */
+	bool ContainsAnyComponent(const FComponentMask& ComponentTypeIDs) const
+	{
+		const FComponentMask& Mask = GetAccumulatedMask();
+		return FComponentMask::BitwiseAND(Mask, ComponentTypeIDs, EBitwiseOperatorFlags::MinSize).NumComponents() != 0;
+	}
 
 	/**
 	 * Check whether all of the specified components exist anywhere in this entity manager
diff --git a/Engine/Source/Runtime/MovieScene/Public/EntitySystem/MovieSceneEntityMutations.h b/Engine/Source/Runtime/MovieScene/Public/EntitySystem/MovieSceneEntityMutations.h
--- a/Engine/Source/Runtime/MovieScene/Public/EntitySystem/MovieSceneEntityMutations.h
+++ b/Engine/Source/Runtime/MovieScene/Public/EntitySystem/MovieSceneEntityMutations.h
@@ -51,14 +51,12 @@ struct FAddMultipleMutation : IMovieSceneEntityMutation
 
 struct FRemoveMultipleMutation : IMovieSceneEntityMutation
 {
+	/* Mask that defines components to remove by _unset_ bits. This acts as a bitmask applied as a binary AND for each mutated allocation. */
+	FComponentMask MaskToRemove;
 
 	void RemoveComponent(FComponentTypeID InComponentType);
 
 	MOVIESCENE_API virtual void CreateMutation(FEntityManager* EntityManager, FComponentMask* InOutEntityComponentTypes) const override;
-
-private:
-
-	FComponentMask MaskToRemove;
 };
 
 } // namespace MovieScene
diff --git a/Engine/Source/Runtime/MovieSceneTracks/Private/MovieSceneTracksComponentTypes.cpp b/Engine/Source/Runtime/MovieSceneTracks/Private/MovieSceneTracksComponentTypes.cpp
--- a/Engine/Source/Runtime/MovieSceneTracks/Private/MovieSceneTracksComponentTypes.cpp
+++ b/Engine/Source/Runtime/MovieSceneTracks/Private/MovieSceneTracksComponentTypes.cpp
@@ -559,6 +559,9 @@ FMovieSceneTracksComponentTypes::FMovieSceneTracksComponentTypes()
 	ComponentRegistry->NewComponentType(&ColorParameterName,     TEXT("Color Parameter Name"), EComponentTypeFlags::CopyToChildren | EComponentTypeFlags::CopyToOutput);
 	ComponentRegistry->NewComponentType(&TransformParameterName, TEXT("Transform Parameter Name"), EComponentTypeFlags::CopyToChildren | EComponentTypeFlags::CopyToOutput);
 
+	Tags.BoundMaterialChanged = ComponentRegistry->NewTag(TEXT("Bound Material Changed"));
+	FBuiltInComponentTypes::Get()->RequiresInstantiationMask.Set(Tags.BoundMaterialChanged);
+
 	FBuiltInComponentTypes* BuiltInComponents = FBuiltInComponentTypes::Get();
 
 	// --------------------------------------------------------------------------------------------
diff --git a/Engine/Source/Runtime/MovieSceneTracks/Private/Systems/MovieSceneMaterialParameterCollectionSystem.cpp b/Engine/Source/Runtime/MovieSceneTracks/Private/Systems/MovieSceneMaterialParameterCollectionSystem.cpp
--- a/Engine/Source/Runtime/MovieSceneTracks/Private/Systems/MovieSceneMaterialParameterCollectionSystem.cpp
+++ b/Engine/Source/Runtime/MovieSceneTracks/Private/Systems/MovieSceneMaterialParameterCollectionSystem.cpp
@@ -257,6 +257,7 @@ void UMovieSceneMaterialParameterCollectionSystem::OnRun(FSystemTaskPrerequisite
 	// Only mutate things that are tagged as requiring linking
 	FEntityComponentFilter Filter;
 	Filter.All({ TracksComponents->MPC, BuiltInComponents->InstanceHandle, BuiltInComponents->Tags.NeedsLink });
+	Filter.None({ BuiltInComponents->Tags.ImportedEntity });
 
 	// Initialize bound dynamic materials (for material collection parameters)
 	FAddMPCMutation BindMaterialsMutation(Linker);
diff --git a/Engine/Source/Runtime/MovieSceneTracks/Private/Systems/MovieSceneMaterialParameterSystem.cpp b/Engine/Source/Runtime/MovieSceneTracks/Private/Systems/MovieSceneMaterialParameterSystem.cpp
--- a/Engine/Source/Runtime/MovieSceneTracks/Private/Systems/MovieSceneMaterialParameterSystem.cpp
+++ b/Engine/Source/Runtime/MovieSceneTracks/Private/Systems/MovieSceneMaterialParameterSystem.cpp
@@ -251,33 +251,58 @@ void UMovieSceneMaterialParameterSystem::OnInstantiation()
 
 	if (GMaterialParameterBlending)
 	{
-		const bool bHasScalars = Linker->EntityManager.ContainsComponent(TracksComponents->ScalarParameterName);
-		const bool bHasColors = Linker->EntityManager.ContainsComponent(TracksComponents->ColorParameterName);
-		const bool bHasVectors = Linker->EntityManager.ContainsComponent(TracksComponents->VectorParameterName);
+		auto HandleUnlinkedAllocation = [this](const FEntityAllocation* Allocation, TRead<UObject*> InObjects, TReadOneOf<FName, FName, FName> ScalarVectorOrColorParameterNames)
+		{
+			if (TComponentPtr<const FName> ScalarParameterNames = TComponentPtr<const FName>(ScalarVectorOrColorParameterNames.Get<0>()))
+			{
+				this->ScalarParameterTracker.VisitUnlinkedAllocation(Allocation);
+			}
+			else
+			{
+				this->VectorParameterTracker.VisitUnlinkedAllocation(Allocation);
+			}
+		};
 
-		if (bHasScalars)
+		auto HandleUpdatedAllocation = [this](const FEntityAllocation* Allocation, TRead<UObject*> InObjects, TReadOneOf<FName, FName, FName> ScalarVectorOrColorParameterNames)
 		{
-			ScalarParameterTracker.Update(Linker, TracksComponents->BoundMaterial, TracksComponents->ScalarParameterName, FEntityComponentFilter());
+			if (TComponentPtr<const FName> ScalarParameterNames = TComponentPtr<const FName>(ScalarVectorOrColorParameterNames.Get<0>()))
+			{
+				this->ScalarParameterTracker.VisitActiveAllocation(Allocation, InObjects, ScalarParameterNames);
+			}
+			else if (TComponentPtr<const FName> VectorParameterNames = TComponentPtr<const FName>(ScalarVectorOrColorParameterNames.Get<1>()))
+			{
+				this->VectorParameterTracker.VisitActiveAllocation(Allocation, InObjects, VectorParameterNames);
+			}
+			else if (TComponentPtr<const FName> ColorParameterNames = TComponentPtr<const FName>(ScalarVectorOrColorParameterNames.Get<2>()))
+			{
+				this->VectorParameterTracker.VisitActiveAllocation(Allocation, InObjects, ColorParameterNames);
+			}
+		};
 
-			FOverlappingMaterialParameterHandler Handler(this);
-			Handler.DefaultComponentMask.Set(BuiltInComponents->DoubleResult[0]);
-			ScalarParameterTracker.ProcessInvalidatedOutputs(Linker, Handler);
-		}
+		// Next handle any new or updated bound materials
+		FEntityTaskBuilder()
+		.Read(TracksComponents->BoundMaterial)
+		.ReadOneOf(TracksComponents->ScalarParameterName, TracksComponents->VectorParameterName, TracksComponents->ColorParameterName)
+		.FilterAny({ BuiltInComponents->Tags.NeedsLink, TracksComponents->Tags.BoundMaterialChanged })
+		.FilterNone({ BuiltInComponents->Tags.NeedsUnlink })
+		.Iterate_PerAllocation(&Linker->EntityManager, HandleUpdatedAllocation);
 
-		if (bHasColors)
-		{
-			VectorParameterTracker.Update(Linker, TracksComponents->BoundMaterial, TracksComponents->ColorParameterName, FEntityComponentFilter());
-		}
-		if (bHasVectors)
-		{
-			VectorParameterTracker.Update(Linker, TracksComponents->BoundMaterial, TracksComponents->VectorParameterName, FEntityComponentFilter());
-		}
-		if (bHasColors || bHasVectors)
-		{
-			FOverlappingMaterialParameterHandler Handler(this);
-			Handler.DefaultComponentMask.SetAll({ BuiltInComponents->DoubleResult[0], BuiltInComponents->DoubleResult[1], BuiltInComponents->DoubleResult[2], BuiltInComponents->DoubleResult[3] });
-			VectorParameterTracker.ProcessInvalidatedOutputs(Linker, Handler);
-		}
+		// First step, handle any entities that are going away
+		FEntityTaskBuilder()
+		.Read(TracksComponents->BoundMaterial)
+		.ReadOneOf(TracksComponents->ScalarParameterName, TracksComponents->VectorParameterName, TracksComponents->ColorParameterName)
+		.FilterAll({ BuiltInComponents->Tags.NeedsUnlink })
+		.Iterate_PerAllocation(&Linker->EntityManager, HandleUnlinkedAllocation);
+
+		// Process all blended scalar parameters
+		FOverlappingMaterialParameterHandler ScalarHandler(this);
+		ScalarHandler.DefaultComponentMask.Set(BuiltInComponents->DoubleResult[0]);
+		ScalarParameterTracker.ProcessInvalidatedOutputs(Linker, ScalarHandler);
+
+		// Process all blended vector parameters
+		FOverlappingMaterialParameterHandler VectorHandler(this);
+		VectorHandler.DefaultComponentMask.SetAll({ BuiltInComponents->DoubleResult[0], BuiltInComponents->DoubleResult[1], BuiltInComponents->DoubleResult[2], BuiltInComponents->DoubleResult[3] });
+		VectorParameterTracker.ProcessInvalidatedOutputs(Linker, VectorHandler);
 	}
 }
 
diff --git a/Engine/Source/Runtime/MovieSceneTracks/Public/MovieSceneTracksComponentTypes.h b/Engine/Source/Runtime/MovieSceneTracks/Public/MovieSceneTracksComponentTypes.h
--- a/Engine/Source/Runtime/MovieSceneTracks/Public/MovieSceneTracksComponentTypes.h
+++ b/Engine/Source/Runtime/MovieSceneTracks/Public/MovieSceneTracksComponentTypes.h
@@ -459,6 +459,11 @@ struct MOVIESCENETRACKS_API FMovieSceneTracksComponentTypes
 		TCustomPropertyRegistration<FComponentTransformPropertyTraits, 1> ComponentTransform;
 	} Accessors;
 
+	struct
+	{
+		FComponentTypeID BoundMaterialChanged;
+	} Tags;
+
 	TComponentTypeID<FLevelVisibilityComponentData> LevelVisibility;
 	TComponentTypeID<FMovieSceneDataLayerComponentData> DataLayer;
 
diff --git a/Engine/Source/Runtime/MovieSceneTracks/Public/Systems/MovieSceneMaterialSystem.h b/Engine/Source/Runtime/MovieSceneTracks/Public/Systems/MovieSceneMaterialSystem.h
--- a/Engine/Source/Runtime/MovieSceneTracks/Public/Systems/MovieSceneMaterialSystem.h
+++ b/Engine/Source/Runtime/MovieSceneTracks/Public/Systems/MovieSceneMaterialSystem.h
@@ -210,10 +210,10 @@ struct TReinitializeBoundMaterials
 
 	void PostTask()
 	{
-		FBuiltInComponentTypes* BuiltInComponents = FBuiltInComponentTypes::Get();
+		FMovieSceneTracksComponentTypes* TracksComponents = FMovieSceneTracksComponentTypes::Get();
 		for (FMovieSceneEntityID EntityID : ReboundMaterials)
 		{
-			Linker->EntityManager.AddComponent(EntityID, BuiltInComponents->Tags.NeedsLink);
+			Linker->EntityManager.AddComponent(EntityID, TracksComponents->Tags.BoundMaterialChanged);
 		}
 	}
 };
@@ -334,6 +334,7 @@ void TMovieSceneMaterialSystem<AccessorType, RequiredComponents...>::OnRun(UMovi
 		// Add bound materials for any NeedsLink entities that have material parameters
 		FEntityComponentFilter Filter(MaterialParameterFilter);
 		Filter.All({ BuiltInComponents->Tags.NeedsLink });
+		Filter.None({ BuiltInComponents->Tags.ImportedEntity });
 
 		using MutationType = TAddBoundMaterialMutationImpl<AccessorType, TMakeIntegerSequence<int, sizeof...(RequiredComponents)>, RequiredComponents...>;
 		Linker->EntityManager.MutateAll(Filter, MutationType(InRequiredComponents...));
@@ -346,6 +347,7 @@ void TMovieSceneMaterialSystem<AccessorType, RequiredComponents...>::OnRun(UMovi
 		.ReadEntityIDs()
 		.ReadAllOf(InRequiredComponents...)
 		.Write(TracksComponents->BoundMaterial)
+		.FilterAll({ BuiltInComponents->Tags.NeedsLink })
 		.FilterNone({ BuiltInComponents->Tags.NeedsUnlink })
 		.RunInline_PerEntity(&Linker->EntityManager, ReinitializeBoundMaterialsTask);
 	}
