Sequencer: speculative null check for rare crash

diff --git a/Engine/Source/Runtime/MovieSceneTracks/Private/TrackInstances/MovieSceneCameraCutTrackInstance.cpp b/Engine/Source/Runtime/MovieSceneTracks/Private/TrackInstances/MovieSceneCameraCutTrackInstance.cpp
--- a/Engine/Source/Runtime/MovieSceneTracks/Private/TrackInstances/MovieSceneCameraCutTrackInstance.cpp
+++ b/Engine/Source/Runtime/MovieSceneTracks/Private/TrackInstances/MovieSceneCameraCutTrackInstance.cpp
@@ -540,10 +540,13 @@ void UMovieSceneCameraCutTrackInstance::OnDestroyed()
 			if (PlayerUseCount.Value.bValid)
 			{
 				EMovieSceneCameraCutParams Params;
-	#if WITH_EDITOR
+#if WITH_EDITOR
 				Params.bCanBlend = PlayerUseCount.Value.bCanBlend;
-	#endif
-				PlayerUseCount.Key->UpdateCameraCut(nullptr, Params);
+#endif
+				if (PlayerUseCount.Key)
+				{
+					PlayerUseCount.Key->UpdateCameraCut(nullptr, Params);
+				}
 				break;  // Only do it on the first one.
 			}
 		}
