Adding BP support for setting animation custom mode and Swap Root Bone on Anim Params

diff --git a/Engine/Source/Runtime/AnimGraphRuntime/Public/AnimSequencerInstanceProxy.h b/Engine/Source/Runtime/AnimGraphRuntime/Public/AnimSequencerInstanceProxy.h
--- a/Engine/Source/Runtime/AnimGraphRuntime/Public/AnimSequencerInstanceProxy.h
+++ b/Engine/Source/Runtime/AnimGraphRuntime/Public/AnimSequencerInstanceProxy.h
@@ -10,7 +10,7 @@
 #include "AnimNodes/AnimNode_Mirror.h"
 #include "AnimSequencerInstanceProxy.generated.h"
 
-UENUM()
+UENUM(BlueprintType)
 enum class ESwapRootBone : uint8
 {
 	/* Swap the root bone to the component */
diff --git a/Engine/Source/Runtime/MovieSceneTracks/Public/Sections/MovieSceneSkeletalAnimationSection.h b/Engine/Source/Runtime/MovieSceneTracks/Public/Sections/MovieSceneSkeletalAnimationSection.h
--- a/Engine/Source/Runtime/MovieSceneTracks/Public/Sections/MovieSceneSkeletalAnimationSection.h
+++ b/Engine/Source/Runtime/MovieSceneTracks/Public/Sections/MovieSceneSkeletalAnimationSection.h
@@ -67,11 +67,11 @@ struct FMovieSceneSkeletalAnimationParams
 	bool bSkipAnimNotifiers;
 
 	/** If on animation sequence will always play when active even if the animation is controlled by a Blueprint or Anim Instance Class*/
-	UPROPERTY(EditAnywhere, Category = "Animation")
+	UPROPERTY(BlueprintReadWrite, EditAnywhere, Category = "Animation")
 	bool bForceCustomMode;
 
 	/** If on the root bone transform will be swapped to the specified root*/
-	UPROPERTY(EditAnywhere, Category = "Animation")
+	UPROPERTY(BlueprintReadWrite, EditAnywhere, Category = "Animation")
 	ESwapRootBone SwapRootBone;
 
 	UPROPERTY()
