Sequencer: simplify channel override entity import, fix float track channel override support

Before, channel overrides would get a type-erased parameter that gives information about what should be overriden and how. This was flawed, since a given type of channel can't know what kind of parameter was created by the caller, since various types of tracks and sections could have handled that. So we just simplify this for now, forcing the parameters to be (for now) just a channel name and result component. That is: parameters that tell the override which channel it's posing as, and where to put the result.

A second bug was that we would still create an ECS entity even if *all* of a section's channels had been overriden, which is easy to do when overriding the channel of a float track (since it only has 1 channel). This was particularly bad since it would allocate the entity but not initialize it with proper data. Now, the import helper checks whether there are non-overriden channels left before allocating space for an entity.

Finally, the float section was incorrectly giving any override channel a source component instead of a result component. The import helper API was to blame here, so this had to also be fixed.

diff --git a/Engine/Source/Runtime/MovieScene/Private/Channels/MovieSceneChannelProxy.cpp b/Engine/Source/Runtime/MovieScene/Private/Channels/MovieSceneChannelProxy.cpp
--- a/Engine/Source/Runtime/MovieScene/Private/Channels/MovieSceneChannelProxy.cpp
+++ b/Engine/Source/Runtime/MovieScene/Private/Channels/MovieSceneChannelProxy.cpp
@@ -4,6 +4,16 @@
 #include "Algo/BinarySearch.h"
 
 
+int32 FMovieSceneChannelProxy::NumChannels() const
+{
+	int32 Num = 0;
+	for (const FMovieSceneChannelEntry& Entry : Entries)
+	{
+		Num += Entry.GetChannels().Num();
+	}
+	return Num;
+}
+
 FMovieSceneChannelHandle FMovieSceneChannelProxy::MakeHandle(FName ChannelTypeName, int32 Index)
 {
 	TWeakPtr<FMovieSceneChannelProxy> WeakProxy = AsShared();
diff --git a/Engine/Source/Runtime/MovieScene/Private/Channels/MovieSceneSectionChannelOverrideRegistry.cpp b/Engine/Source/Runtime/MovieScene/Private/Channels/MovieSceneSectionChannelOverrideRegistry.cpp
--- a/Engine/Source/Runtime/MovieScene/Private/Channels/MovieSceneSectionChannelOverrideRegistry.cpp
+++ b/Engine/Source/Runtime/MovieScene/Private/Channels/MovieSceneSectionChannelOverrideRegistry.cpp
@@ -20,6 +20,11 @@ bool UMovieSceneSectionChannelOverrideRegistry::ContainsChannel(FName ChannelNam
 	return Overrides.Contains(ChannelName);
 }
 
+int32 UMovieSceneSectionChannelOverrideRegistry::NumChannels() const
+{
+	return Overrides.Num();
+}
+
 UMovieSceneChannelOverrideContainer* UMovieSceneSectionChannelOverrideRegistry::GetChannel(FName ChannelName) const
 {
 	return Overrides.FindRef(ChannelName);
@@ -30,14 +35,13 @@ void UMovieSceneSectionChannelOverrideRegistry::RemoveChannel(FName ChannelName)
 	Overrides.Remove(ChannelName);
 }
 
-void UMovieSceneSectionChannelOverrideRegistry::ImportEntityImpl(const FMovieSceneChannelOverrideEntityImportParamsHandle& OverrideParams, const UE::MovieScene::FEntityImportParams& ImportParams, UE::MovieScene::FImportedEntity* OutImportedEntity)
+void UMovieSceneSectionChannelOverrideRegistry::ImportEntityImpl(const FMovieSceneChannelOverrideEntityImportParams& OverrideParams, const UE::MovieScene::FEntityImportParams& ImportParams, UE::MovieScene::FImportedEntity* OutImportedEntity)
 {
-	using namespace UE::MovieScene;
-
-	check(OverrideParams.IsValid());
-	FName ChannelName = OverrideParams.GetPtr()->ChannelName;
-	check(!ChannelName.IsNone());
-	Overrides[ChannelName]->ImportEntityImpl(OverrideParams, ImportParams, OutImportedEntity);
+	TObjectPtr<UMovieSceneChannelOverrideContainer>* ContainerPtr = Overrides.Find(OverrideParams.ChannelName);
+	if (ensure(ContainerPtr))
+	{
+		(*ContainerPtr)->ImportEntityImpl(OverrideParams, ImportParams, OutImportedEntity);
+	}
 }
 
 void UMovieSceneSectionChannelOverrideRegistry::PopulateEvaluationFieldImpl(const TRange<FFrameNumber>& EffectiveRange, const FMovieSceneEvaluationFieldEntityMetaData& InMetaData, FMovieSceneEntityComponentFieldBuilder* OutFieldBuilder, UMovieSceneSection& OwnerSection)
diff --git a/Engine/Source/Runtime/MovieScene/Public/Channels/MovieSceneChannelOverrideContainer.h b/Engine/Source/Runtime/MovieScene/Public/Channels/MovieSceneChannelOverrideContainer.h
--- a/Engine/Source/Runtime/MovieScene/Public/Channels/MovieSceneChannelOverrideContainer.h
+++ b/Engine/Source/Runtime/MovieScene/Public/Channels/MovieSceneChannelOverrideContainer.h
@@ -22,57 +22,18 @@ struct FMovieSceneChannelProxyData;
 struct FMovieSceneEntityComponentFieldBuilder;
 
 /**
- * Base structure for entity import parameters of a channel override.
+ * Entity import parameters for a channel override.
  */
-struct IMovieSceneChannelOverrideEntityImportParams
+struct FMovieSceneChannelOverrideEntityImportParams
 {
 	FName ChannelName;
-};
-
-/*
- * Handle for storing channel override import parameters.
- *
- * These parameters carry type-specific information, such as the result component in which the channel
- * override should throw output values.
- */
-struct FMovieSceneChannelOverrideEntityImportParamsHandle : 
-	TInlineValue<IMovieSceneChannelOverrideEntityImportParams, sizeof(FName) + 16>
-{
-	using Super = TInlineValue<IMovieSceneChannelOverrideEntityImportParams, sizeof(FName) + 16>;
-
-	FMovieSceneChannelOverrideEntityImportParamsHandle()
-		: Super()
-	{
-	}
-
-	template<typename T>
-	FMovieSceneChannelOverrideEntityImportParamsHandle(T&& In)
-		: Super(In)
-	{
-	}
-
-	template<typename T>
-	const T& CastThis() const
-	{
-		return *static_cast<const T*>(this->GetPtr());
-	}
-};
-
-/**
- * Entity import parameters for channel overrides that just need an untyped result component.
- */
-struct FMovieSceneChannelOverrideResultComponentEntityImportParams : IMovieSceneChannelOverrideEntityImportParams
-{
 	UE::MovieScene::FComponentTypeID ResultComponent;
-};
 
-/**
- * Entity import parameters for channel overrides that need a typed result component.
- */
-template<typename ComponentType>
-struct TMovieSceneChannelOverrideResultComponentEntityImportParams : IMovieSceneChannelOverrideEntityImportParams
-{
-	UE::MovieScene::TComponentTypeID<ComponentType> ResultComponent;
+	FMovieSceneChannelOverrideEntityImportParams()
+	{}
+	FMovieSceneChannelOverrideEntityImportParams(FName InChannelName, UE::MovieScene::FComponentTypeID InResultComponent)
+		: ChannelName(InChannelName), ResultComponent(InResultComponent)
+	{}
 };
 
 /**
@@ -90,7 +51,7 @@ public:
 
 	/** Imports the entity for this channel */
 	virtual void ImportEntityImpl(
-			const FMovieSceneChannelOverrideEntityImportParamsHandle& OverrideParams, 
+			const FMovieSceneChannelOverrideEntityImportParams& OverrideParams, 
 			const UE::MovieScene::FEntityImportParams& ImportParams, UE::MovieScene::FImportedEntity* OutImportedEntity) {};
 
 	/** Gets the underlying channel */
diff --git a/Engine/Source/Runtime/MovieScene/Public/Channels/MovieSceneChannelProxy.h b/Engine/Source/Runtime/MovieScene/Public/Channels/MovieSceneChannelProxy.h
--- a/Engine/Source/Runtime/MovieScene/Public/Channels/MovieSceneChannelProxy.h
+++ b/Engine/Source/Runtime/MovieScene/Public/Channels/MovieSceneChannelProxy.h
@@ -273,6 +273,12 @@ public:
 	 */
 	FMovieSceneChannel* GetChannel(FName ChannelTypeName, int32 ChannelIndex) const;
 
+	/**
+	 * Returns the total number of channels
+	 * @return The total number of channels
+	 */
+	int32 NumChannels() const;
+
 	/**
 	 * Make a channel handle out for the specified index and channel type name
 	 *
diff --git a/Engine/Source/Runtime/MovieScene/Public/Channels/MovieSceneSectionChannelOverrideRegistry.h b/Engine/Source/Runtime/MovieScene/Public/Channels/MovieSceneSectionChannelOverrideRegistry.h
--- a/Engine/Source/Runtime/MovieScene/Public/Channels/MovieSceneSectionChannelOverrideRegistry.h
+++ b/Engine/Source/Runtime/MovieScene/Public/Channels/MovieSceneSectionChannelOverrideRegistry.h
@@ -61,6 +61,12 @@ public:
 	*/
 	bool ContainsChannel(FName ChannelName) const;
 
+	/**
+	 * Returns the number of override channels in this registry
+	 * @return	The number of override channels
+	 */
+	int32 NumChannels() const;
+
 	/**
 	 * Returns the channel override for a given name, or nullptr if not found
 	 * @param ChannelName	Name of the channel
@@ -93,7 +99,7 @@ public:
 	/**
 	* Forward ImportEntityImpl calls to an overriden channel
 	*/
-	void ImportEntityImpl(const FMovieSceneChannelOverrideEntityImportParamsHandle& OverrideParams, const UE::MovieScene::FEntityImportParams& ImportParams, UE::MovieScene::FImportedEntity* OutImportedEntity);
+	void ImportEntityImpl(const FMovieSceneChannelOverrideEntityImportParams& OverrideParams, const UE::MovieScene::FEntityImportParams& ImportParams, UE::MovieScene::FImportedEntity* OutImportedEntity);
 
 	/**
 	* Called when overridden channels should populate evaluation field
@@ -151,6 +157,32 @@ void AddChannelProxy(
 	}
 }
 
+/**
+ * Utility function for making a channel proxy with one channel that is maybe overriden by another channel.
+ */
+template<typename ChannelType, typename ExtendedEditorDataType>
+TSharedPtr<FMovieSceneChannelProxy> MakeChannelProxy(
+	const UMovieSceneSectionChannelOverrideRegistry* OverrideRegistry,
+	ChannelType& DefaultChannel,
+	const FMovieSceneChannelMetaData& InMetaData, ExtendedEditorDataType&& InExtendedEditorData)
+{
+	// When there's only one channel, it usually has no name. Create a channel proxy by hand if
+	// you need a named singled channel.
+	const FName ChannelName(NAME_None);
+	FMovieSceneChannelProxyData ProxyData;
+	UMovieSceneChannelOverrideContainer* OverrideContainer = GetChannelOverride(OverrideRegistry, ChannelName);
+	if (OverrideContainer)
+	{
+		FMovieSceneChannelHandle ChannelHandle = OverrideContainer->AddChannelProxy(ChannelName, ProxyData, InMetaData);
+		ProxyData.SetExtendedEditorData<ChannelType>(ChannelHandle, InExtendedEditorData);
+	}
+	else
+	{
+		ProxyData.Add(DefaultChannel, InMetaData, InExtendedEditorData);
+	}
+	return MakeShared<FMovieSceneChannelProxy>(MoveTemp(ProxyData));
+}
+
 #else
 
 /**
@@ -172,6 +204,28 @@ void AddChannelProxy(
 	}
 }
 
+/**
+ * Utility function for making a channel proxy with one channel that is maybe overriden by another channel.
+ */
+template<typename ChannelType>
+TSharedPtr<FMovieSceneChannelProxy> MakeChannelProxy(
+	const UMovieSceneSectionChannelOverrideRegistry* OverrideRegistry,
+	ChannelType& DefaultChannel)
+{
+	const FName ChannelName(NAME_None);
+	FMovieSceneChannelProxyData ProxyData;
+	UMovieSceneChannelOverrideContainer* OverrideContainer = GetChannelOverride(OverrideRegistry, ChannelName);
+	if (OverrideContainer)
+	{
+		OverrideContainer->AddChannelProxy(ChannelName, ProxyData);
+	}
+	else
+	{
+		ProxyData.Add(DefaultChannel);
+	}
+	return MakeShared<FMovieSceneChannelProxy>(MoveTemp(ProxyData));
+}
+
 #endif
 
 }  // namespace MovieScene
diff --git a/Engine/Source/Runtime/MovieSceneTracks/Private/Channels/MovieSceneDoublePerlinNoiseChannelContainer.cpp b/Engine/Source/Runtime/MovieSceneTracks/Private/Channels/MovieSceneDoublePerlinNoiseChannelContainer.cpp
--- a/Engine/Source/Runtime/MovieSceneTracks/Private/Channels/MovieSceneDoublePerlinNoiseChannelContainer.cpp
+++ b/Engine/Source/Runtime/MovieSceneTracks/Private/Channels/MovieSceneDoublePerlinNoiseChannelContainer.cpp
@@ -11,19 +11,19 @@ bool UMovieSceneDoublePerlinNoiseChannelContainer::SupportsOverride(FName Defaul
 	return FMovieSceneDoubleChannel::StaticStruct()->GetFName() == DefaultChannelTypeName;
 }
 
-void UMovieSceneDoublePerlinNoiseChannelContainer::ImportEntityImpl(const FMovieSceneChannelOverrideEntityImportParamsHandle& OverrideParams, const UE::MovieScene::FEntityImportParams& ImportParams, UE::MovieScene::FImportedEntity* OutImportedEntity)
+void UMovieSceneDoublePerlinNoiseChannelContainer::ImportEntityImpl(const FMovieSceneChannelOverrideEntityImportParams& OverrideParams, const UE::MovieScene::FEntityImportParams& ImportParams, UE::MovieScene::FImportedEntity* OutImportedEntity)
 {
 	using namespace UE::MovieScene;
 
 	const FMovieSceneTracksComponentTypes* TracksComponents = FMovieSceneTracksComponentTypes::Get();
 	const FBuiltInComponentTypes* BuiltInComponents = FBuiltInComponentTypes::Get();
 
-	const auto& DoubleResultParams = OverrideParams.CastThis<TMovieSceneChannelOverrideResultComponentEntityImportParams<double>>();
+	TComponentTypeID<double> ResultComponent = OverrideParams.ResultComponent.ReinterpretCast<double>();
 
 	OutImportedEntity->AddBuilder(
 		FEntityBuilder()
 		.Add(TracksComponents->DoublePerlinNoiseChannel, PerlinNoiseChannel.GetParams())
-		.Add(DoubleResultParams.ResultComponent, TNumericLimits<double>::Max())
+		.Add(ResultComponent, TNumericLimits<double>::Max())
 	);
 }
 
diff --git a/Engine/Source/Runtime/MovieSceneTracks/Private/Channels/MovieSceneFloatPerlinNoiseChannelContainer.cpp b/Engine/Source/Runtime/MovieSceneTracks/Private/Channels/MovieSceneFloatPerlinNoiseChannelContainer.cpp
--- a/Engine/Source/Runtime/MovieSceneTracks/Private/Channels/MovieSceneFloatPerlinNoiseChannelContainer.cpp
+++ b/Engine/Source/Runtime/MovieSceneTracks/Private/Channels/MovieSceneFloatPerlinNoiseChannelContainer.cpp
@@ -12,19 +12,19 @@ bool UMovieSceneFloatPerlinNoiseChannelContainer::SupportsOverride(FName Default
 	return FMovieSceneFloatChannel::StaticStruct()->GetFName() == DefaultChannelTypeName;
 }
 
-void UMovieSceneFloatPerlinNoiseChannelContainer::ImportEntityImpl(const FMovieSceneChannelOverrideEntityImportParamsHandle& OverrideParams, const UE::MovieScene::FEntityImportParams& ImportParams, UE::MovieScene::FImportedEntity* OutImportedEntity)
+void UMovieSceneFloatPerlinNoiseChannelContainer::ImportEntityImpl(const FMovieSceneChannelOverrideEntityImportParams& OverrideParams, const UE::MovieScene::FEntityImportParams& ImportParams, UE::MovieScene::FImportedEntity* OutImportedEntity)
 {
 	using namespace UE::MovieScene;
 
 	const FMovieSceneTracksComponentTypes* TracksComponents = FMovieSceneTracksComponentTypes::Get();
 	const FBuiltInComponentTypes* BuiltInComponents = FBuiltInComponentTypes::Get();
 
-	const auto& FloatResultParams = OverrideParams.CastThis<TMovieSceneChannelOverrideResultComponentEntityImportParams<float>>();
+	TComponentTypeID<double> ResultComponent = OverrideParams.ResultComponent.ReinterpretCast<double>();
 
 	OutImportedEntity->AddBuilder(
 		FEntityBuilder()
 		.Add(TracksComponents->FloatPerlinNoiseChannel, PerlinNoiseChannel.GetParams())
-		.Add(FloatResultParams.ResultComponent, TNumericLimits<float>::Max())
+		.Add(ResultComponent, TNumericLimits<float>::Max())
 	);
 }
 
diff --git a/Engine/Source/Runtime/MovieSceneTracks/Private/Sections/MovieScene3DTransformSection.cpp b/Engine/Source/Runtime/MovieSceneTracks/Private/Sections/MovieScene3DTransformSection.cpp
--- a/Engine/Source/Runtime/MovieSceneTracks/Private/Sections/MovieScene3DTransformSection.cpp
+++ b/Engine/Source/Runtime/MovieSceneTracks/Private/Sections/MovieScene3DTransformSection.cpp
@@ -421,9 +421,7 @@ void UMovieScene3DTransformSection::BuildEntity(BaseBuilderType& InBaseBuilder,
 		{
 			if (EnumHasAnyFlags(Channels, EMovieSceneTransformChannel::Weight))
 			{
-				TMovieSceneChannelOverrideResultComponentEntityImportParams<double> OverrideParams;
-				OverrideParams.ChannelName = ChannelOverrideName;
-				OverrideParams.ResultComponent = BuiltInComponentTypes->WeightResult;
+				FMovieSceneChannelOverrideEntityImportParams OverrideParams(ChannelOverrideName, BuiltInComponentTypes->WeightResult);
 				OverrideRegistry->ImportEntityImpl(OverrideParams, Params, OutImportedEntity);
 			}
 		}
@@ -431,9 +429,7 @@ void UMovieScene3DTransformSection::BuildEntity(BaseBuilderType& InBaseBuilder,
 		{
 			if (ActiveChannelsMask[ChannelOverrideIndex])
 			{
-				TMovieSceneChannelOverrideResultComponentEntityImportParams<double> OverrideParams;
-				OverrideParams.ChannelName = ChannelOverrideName;
-				OverrideParams.ResultComponent = BuiltInComponentTypes->DoubleResult[ChannelOverrideIndex];
+				FMovieSceneChannelOverrideEntityImportParams OverrideParams(ChannelOverrideName, BuiltInComponentTypes->DoubleResult[ChannelOverrideIndex]);
 				OverrideRegistry->ImportEntityImpl(OverrideParams, Params, OutImportedEntity);
 			}
 		}
diff --git a/Engine/Source/Runtime/MovieSceneTracks/Private/Sections/MovieSceneFloatSection.cpp b/Engine/Source/Runtime/MovieSceneTracks/Private/Sections/MovieSceneFloatSection.cpp
--- a/Engine/Source/Runtime/MovieSceneTracks/Private/Sections/MovieSceneFloatSection.cpp
+++ b/Engine/Source/Runtime/MovieSceneTracks/Private/Sections/MovieSceneFloatSection.cpp
@@ -30,13 +30,15 @@ UMovieSceneFloatSection::UMovieSceneFloatSection( const FObjectInitializer& Obje
 
 EMovieSceneChannelProxyType UMovieSceneFloatSection::CacheChannelProxy()
 {
+	using namespace UE::MovieScene;
+
 #if WITH_EDITOR
 
-	ChannelProxy = MakeShared<FMovieSceneChannelProxy>(FloatCurve, FMovieSceneChannelMetaData(), TMovieSceneExternalValue<float>::Make());
+	ChannelProxy = MakeChannelProxy(OverrideRegistry, FloatCurve, FMovieSceneChannelMetaData(), TMovieSceneExternalValue<float>::Make());
 
 #else
 
-	ChannelProxy = MakeShared<FMovieSceneChannelProxy>(FloatCurve);
+	ChannelProxy = MakeChannelProxy(OverrideRegistry, FloatCurve);
 
 #endif
 
@@ -53,7 +55,7 @@ void UMovieSceneFloatSection::ImportEntityImpl(UMovieSceneEntitySystemLinker* En
 {
 	using namespace UE::MovieScene;
 
-	if (!FloatCurve.HasAnyData())
+	if (!FloatCurve.HasAnyData() && (!OverrideRegistry || OverrideRegistry->NumChannels() == 0))
 	{
 		return;
 	}
@@ -62,7 +64,9 @@ void UMovieSceneFloatSection::ImportEntityImpl(UMovieSceneEntitySystemLinker* En
 	const FMovieSceneTracksComponentTypes* TracksComponents = FMovieSceneTracksComponentTypes::Get();
 
 	FPropertyTrackWithOverridableChannelsEntityImportHelper(TracksComponents->Float, this)
-		.Add(Components->FloatChannel[0], FName("FloatChannel"), &FloatCurve)
+		.Add(Components->FloatChannel[0], &FloatCurve,
+				// NAME_None below because we don't specify any names in CacheChannelProxy
+				FMovieSceneChannelOverrideEntityImportParams(NAME_None, Components->DoubleResult[0]))
 		.Commit(this, Params, OutImportedEntity);
 }
 
@@ -77,7 +81,8 @@ UMovieSceneSectionChannelOverrideRegistry* UMovieSceneFloatSection::GetChannelOv
 
 UE::MovieScene::FChannelOverrideProviderTraitsHandle UMovieSceneFloatSection::GetChannelOverrideProviderTraits() const
 {
-	UE::MovieScene::TSingleChannelOverrideProviderTraits<FMovieSceneFloatChannel> Traits(FName("FloatChannel"));
+	// NAME_None because we don't specify any names in CacheChannelProxy
+	UE::MovieScene::TSingleChannelOverrideProviderTraits<FMovieSceneFloatChannel> Traits(NAME_None);
 	return UE::MovieScene::FChannelOverrideProviderTraitsHandle(Traits);
 }
 
diff --git a/Engine/Source/Runtime/MovieSceneTracks/Private/Tracks/MovieScenePropertyTrack.cpp b/Engine/Source/Runtime/MovieSceneTracks/Private/Tracks/MovieScenePropertyTrack.cpp
--- a/Engine/Source/Runtime/MovieSceneTracks/Private/Tracks/MovieScenePropertyTrack.cpp
+++ b/Engine/Source/Runtime/MovieSceneTracks/Private/Tracks/MovieScenePropertyTrack.cpp
@@ -335,19 +335,25 @@ void FMovieScenePropertyTrackEntityImportHelper::PopulateEvaluationField(UMovieS
 {
 	using namespace UE::MovieScene;
 
+	int32 NumOverridenChannels = 0;
 	IMovieSceneChannelOverrideProvider* RegistryProvider = Cast<IMovieSceneChannelOverrideProvider>(&Section);
 	if (RegistryProvider)
 	{
 		if (UMovieSceneSectionChannelOverrideRegistry* OverrideRegistry = RegistryProvider->GetChannelOverrideRegistry(false))
 		{
+			NumOverridenChannels = OverrideRegistry->NumChannels();
 			OverrideRegistry->PopulateEvaluationFieldImpl(EffectiveRange, InMetaData, OutFieldBuilder, Section);
 		}
 	}
 
-	// Add the default entity for this section.
-	const int32 EntityIndex   = OutFieldBuilder->FindOrAddEntity(&Section, SectionPropertyValueImportingID);
-	const int32 MetaDataIndex = OutFieldBuilder->AddMetaData(InMetaData);
-	OutFieldBuilder->AddPersistentEntity(EffectiveRange, EntityIndex, MetaDataIndex);
+	const int32 NumChannels = Section.GetChannelProxy().NumChannels();
+	if (NumChannels > NumOverridenChannels)
+	{
+		// Add the default entity for this section.
+		const int32 EntityIndex = OutFieldBuilder->FindOrAddEntity(&Section, SectionPropertyValueImportingID);
+		const int32 MetaDataIndex = OutFieldBuilder->AddMetaData(InMetaData);
+		OutFieldBuilder->AddPersistentEntity(EffectiveRange, EntityIndex, MetaDataIndex);
+	}
 
 	// Check if this section is animating a property with an edit-condition. If so, we need to also animate a boolean toggle
 	// that will be set to true while the main property is animated.
diff --git a/Engine/Source/Runtime/MovieSceneTracks/Public/Channels/MovieSceneDoublePerlinNoiseChannelContainer.h b/Engine/Source/Runtime/MovieSceneTracks/Public/Channels/MovieSceneDoublePerlinNoiseChannelContainer.h
--- a/Engine/Source/Runtime/MovieSceneTracks/Public/Channels/MovieSceneDoublePerlinNoiseChannelContainer.h
+++ b/Engine/Source/Runtime/MovieSceneTracks/Public/Channels/MovieSceneDoublePerlinNoiseChannelContainer.h
@@ -20,7 +20,7 @@ public:
 
 	bool SupportsOverride(FName DefaultChannelTypeName) const override;
 	void ImportEntityImpl(
-			const FMovieSceneChannelOverrideEntityImportParamsHandle& OverrideParams, 
+			const FMovieSceneChannelOverrideEntityImportParams& OverrideParams, 
 			const UE::MovieScene::FEntityImportParams& ImportParams, UE::MovieScene::FImportedEntity* OutImportedEntity) override;
 
 	const FMovieSceneChannel* GetChannel() const override { return &PerlinNoiseChannel; }
diff --git a/Engine/Source/Runtime/MovieSceneTracks/Public/Channels/MovieSceneFloatPerlinNoiseChannelContainer.h b/Engine/Source/Runtime/MovieSceneTracks/Public/Channels/MovieSceneFloatPerlinNoiseChannelContainer.h
--- a/Engine/Source/Runtime/MovieSceneTracks/Public/Channels/MovieSceneFloatPerlinNoiseChannelContainer.h
+++ b/Engine/Source/Runtime/MovieSceneTracks/Public/Channels/MovieSceneFloatPerlinNoiseChannelContainer.h
@@ -21,7 +21,7 @@ public:
 
 	bool SupportsOverride(FName DefaultChannelTypeName) const override;
 	void ImportEntityImpl(
-			const FMovieSceneChannelOverrideEntityImportParamsHandle& OverrideParams, 
+			const FMovieSceneChannelOverrideEntityImportParams& OverrideParams, 
 			const UE::MovieScene::FEntityImportParams& ImportParams, UE::MovieScene::FImportedEntity* OutImportedEntity) override;
 
 	const FMovieSceneChannel* GetChannel() const override { return &PerlinNoiseChannel; }
diff --git a/Engine/Source/Runtime/MovieSceneTracks/Public/Tracks/MovieScenePropertyTrack.h b/Engine/Source/Runtime/MovieSceneTracks/Public/Tracks/MovieScenePropertyTrack.h
--- a/Engine/Source/Runtime/MovieSceneTracks/Public/Tracks/MovieScenePropertyTrack.h
+++ b/Engine/Source/Runtime/MovieSceneTracks/Public/Tracks/MovieScenePropertyTrack.h
@@ -283,7 +283,7 @@ struct FPropertyTrackEntityImportHelperParamsImpl
 	}
 	TScriptInterface<IMovieSceneChannelOverrideProvider> RegistryProvider;
 	TObjectPtr<UMovieSceneSectionChannelOverrideRegistry> Registry;
-	TArray<FMovieSceneChannelOverrideResultComponentEntityImportParams> ImportParams;
+	TArray<FMovieSceneChannelOverrideEntityImportParams> ImportParams;
 };
 
 /**
@@ -300,31 +300,23 @@ struct TPropertyTrackWithOverridableChannelsEntityImportHelperImpl
 	}
 
 	template<typename U, typename PayloadType>
-	TPropertyTrackEntityImportHelperImpl<T..., TAddConditional<U>> Add(TComponentTypeID<U> ComponentType, FName ChannelName, PayloadType&& InPayload)
+	TPropertyTrackEntityImportHelperImpl<T..., TAddConditional<U>> Add(TComponentTypeID<U> ComponentType, PayloadType&& InPayload, FMovieSceneChannelOverrideEntityImportParams OverrideParams)
 	{
-		check(!ChannelName.IsNone());
-		const bool bOverriden = (OverrideInfo.Registry && OverrideInfo.Registry->ContainsChannel(ChannelName));
+		const bool bOverriden = (OverrideInfo.Registry && OverrideInfo.Registry->ContainsChannel(OverrideParams.ChannelName));
 		if (bOverriden)
 		{
-			FMovieSceneChannelOverrideResultComponentEntityImportParams OverrideParams;
-			OverrideParams.ChannelName = ChannelName;
-			OverrideParams.ResultComponent = ComponentType;
-			OverrideInfo.ImportParams.Add(OverrideParams);
+			OverrideInfo.ImportParams.Add(MoveTemp(OverrideParams));
 		}
 		return TPropertyTrackEntityImportHelperImpl<T..., TAddConditional<U>>(Builder.Add(ComponentType, InPayload, !bOverriden), OverrideInfo, PropertyTag);
 	}
 
 	template<typename U, typename PayloadType>
-	TPropertyTrackEntityImportHelperImpl<T..., TAddConditional<U>> AddConditional(TComponentTypeID<U> ComponentType, FName ChannelName, PayloadType&& InPayload, bool bCondition)
+	TPropertyTrackEntityImportHelperImpl<T..., TAddConditional<U>> AddConditional(TComponentTypeID<U> ComponentType, PayloadType&& InPayload, bool bCondition, FMovieSceneChannelOverrideEntityImportParams OverrideParams)
 	{
-		check(!ChannelName.IsNone());
-		const bool bOverriden = (OverrideInfo.Registry && OverrideInfo.Registry->ContainsChannel(ChannelName));
+		const bool bOverriden = (OverrideInfo.Registry && OverrideInfo.Registry->ContainsChannel(OverrideParams.ChannelName));
 		if (bOverriden)
 		{
-			FMovieSceneChannelOverrideResultComponentEntityImportParams OverrideParams;
-			OverrideParams.ChannelName = ChannelName;
-			OverrideParams.ResultComponent = ComponentType;
-			OverrideInfo.ImportParams.Add(OverrideParams);
+			OverrideInfo.ImportParams.Add(MoveTemp(OverrideParams));
 		}
 		return TPropertyTrackEntityImportHelperImpl<T..., TAddConditional<U>>(Builder.AddConditional(ComponentType, InPayload, bCondition && !bOverriden), OverrideInfo, PropertyTag);
 	}
@@ -365,7 +357,7 @@ struct TPropertyTrackWithOverridableChannelsEntityImportHelperImpl
 					check(Traits.IsValid());
 					const FName ChannelOverrideName = Traits->GetChannelOverrideName(Params.EntityID);
 					auto* ChannelOverrideParams = OverrideInfo.ImportParams.FindByPredicate(
-							[=](const FMovieSceneChannelOverrideResultComponentEntityImportParams& CurParams)
+							[=](const FMovieSceneChannelOverrideEntityImportParams& CurParams)
 							{
 								return CurParams.ChannelName == ChannelOverrideName;
 							});
@@ -406,31 +398,23 @@ struct TPropertyTrackWithOverridableChannelsEntityImportHelperImpl<>
 	}
 
 	template<typename U, typename PayloadType>
-	TPropertyTrackWithOverridableChannelsEntityImportHelperImpl<FAdd, TAddConditional<U>> Add(TComponentTypeID<U> ComponentType, FName ChannelName, PayloadType&& InPayload)
+	TPropertyTrackWithOverridableChannelsEntityImportHelperImpl<FAdd, TAddConditional<U>> Add(TComponentTypeID<U> ComponentType, PayloadType&& InPayload, FMovieSceneChannelOverrideEntityImportParams OverrideParams)
 	{
-		check(!ChannelName.IsNone());
-		const bool bOverriden = (OverrideInfo.Registry && OverrideInfo.Registry->ContainsChannel(ChannelName));
+		const bool bOverriden = (OverrideInfo.Registry && OverrideInfo.Registry->ContainsChannel(OverrideParams.ChannelName));
 		if (bOverriden)
 		{
-			FMovieSceneChannelOverrideResultComponentEntityImportParams OverrideParams;
-			OverrideParams.ChannelName = ChannelName;
-			OverrideParams.ResultComponent = ComponentType;
-			OverrideInfo.ImportParams.Add(OverrideParams);
+			OverrideInfo.ImportParams.Add(MoveTemp(OverrideParams));
 		}
 		return TPropertyTrackWithOverridableChannelsEntityImportHelperImpl<FAdd, TAddConditional<U>>(Builder.AddConditional(ComponentType, InPayload, !bOverriden), OverrideInfo, PropertyTag);
 	}
 
 	template<typename U, typename PayloadType>
-	TPropertyTrackWithOverridableChannelsEntityImportHelperImpl<FAdd, TAddConditional<U>> AddConditional(TComponentTypeID<U> ComponentType, FName ChannelName, PayloadType&& InPayload, bool bCondition)
+	TPropertyTrackWithOverridableChannelsEntityImportHelperImpl<FAdd, TAddConditional<U>> AddConditional(TComponentTypeID<U> ComponentType, PayloadType&& InPayload, bool bCondition, FMovieSceneChannelOverrideEntityImportParams OverrideParams)
 	{
-		check(!ChannelName.IsNone());
-		const bool bOverriden = (OverrideInfo.Registry && OverrideInfo.Registry->ContainsChannel(ChannelName));
+		const bool bOverriden = (OverrideInfo.Registry && OverrideInfo.Registry->ContainsChannel(OverrideParams.ChannelName));
 		if (bOverriden)
 		{
-			FMovieSceneChannelOverrideResultComponentEntityImportParams OverrideParams;
-			OverrideParams.ChannelName = ChannelName;
-			OverrideParams.ResultComponent = ComponentType;
-			OverrideInfo.ImportParams.Add(OverrideParams);
+			OverrideInfo.ImportParams.Add(MoveTemp(OverrideParams));
 		}
 		return TPropertyTrackWithOverridableChannelsEntityImportHelperImpl<FAdd, TAddConditional<U>>(Builder.AddConditional(ComponentType, InPayload, bCondition && !bOverriden), OverrideInfo, PropertyTag);
 	}
