MovieScene: Fixed ExponentialHeightFog.SecondFogData properties not dirtying render state when set from Sequencer evaluation.

diff --git a/Engine/Source/Runtime/MovieSceneTracks/Private/MovieSceneTracksComponentTypes.cpp b/Engine/Source/Runtime/MovieSceneTracks/Private/MovieSceneTracksComponentTypes.cpp
--- a/Engine/Source/Runtime/MovieSceneTracks/Private/MovieSceneTracksComponentTypes.cpp
+++ b/Engine/Source/Runtime/MovieSceneTracks/Private/MovieSceneTracksComponentTypes.cpp
@@ -251,7 +251,7 @@ float GetSecondFogDataFogDensity(const UObject* Object)
 void SetSecondFogDataFogDensity(UObject* Object, float InFogDensity)
 {
 	UExponentialHeightFogComponent* ExponentialHeightFogComponent = CastChecked<UExponentialHeightFogComponent>(Object);
-	ExponentialHeightFogComponent->SecondFogData.FogDensity = InFogDensity;
+	ExponentialHeightFogComponent->SetSecondFogDensity(InFogDensity);
 }
 
 float GetSecondFogDataFogHeightFalloff(const UObject* Object)
@@ -263,7 +263,7 @@ float GetSecondFogDataFogHeightFalloff(const UObject* Object)
 void SetSecondFogDataFogHeightFalloff(UObject* Object, float InFogHeightFalloff)
 {
 	UExponentialHeightFogComponent* ExponentialHeightFogComponent = CastChecked<UExponentialHeightFogComponent>(Object);
-	ExponentialHeightFogComponent->SecondFogData.FogHeightFalloff = InFogHeightFalloff;
+	ExponentialHeightFogComponent->SetSecondFogHeightFalloff(InFogHeightFalloff);
 }
 
 float GetSecondFogDataFogHeightOffset(const UObject* Object)
@@ -275,7 +275,7 @@ float GetSecondFogDataFogHeightOffset(const UObject* Object)
 void SetSecondFogDataFogHeightOffset(UObject* Object, float InFogHeightOffset)
 {
 	UExponentialHeightFogComponent* ExponentialHeightFogComponent = CastChecked<UExponentialHeightFogComponent>(Object);
-	ExponentialHeightFogComponent->SecondFogData.FogHeightOffset = InFogHeightOffset;
+	ExponentialHeightFogComponent->SetSecondFogHeightOffset(InFogHeightOffset);
 }
 
 void FIntermediate3DTransform::ApplyTo(USceneComponent* SceneComponent) const
