Sequencer: Use the same params struct between float and double perlin noise channels

diff --git a/Engine/Source/Editor/MovieSceneTools/Public/Channels/PerlinNoiseChannelInterface.h b/Engine/Source/Editor/MovieSceneTools/Public/Channels/PerlinNoiseChannelInterface.h
--- a/Engine/Source/Editor/MovieSceneTools/Public/Channels/PerlinNoiseChannelInterface.h
+++ b/Engine/Source/Editor/MovieSceneTools/Public/Channels/PerlinNoiseChannelInterface.h
@@ -204,7 +204,7 @@ struct TPerlinNoiseChannelInterface : ISequencerChannelInterface
 		TArray<FVector2D> CurvePoints;
 		CurvePoints.Reserve(PaintArgs.Geometry.Size.X / 2.0);
 
-		const double Amplitude = TypedChannel->GetParams().Amplitude;
+		const double Amplitude = TypedChannel->PerlinNoiseParams.Amplitude;
 		const double YOffset = PaintArgs.Geometry.Size.Y / 2.0;
 		const double YScale = (Amplitude != 0) ? (PaintArgs.Geometry.Size.Y / Amplitude / 2.0) : 1;
 
diff --git a/Engine/Source/Runtime/MovieSceneTracks/Private/Channels/MovieSceneDoublePerlinNoiseChannel.cpp b/Engine/Source/Runtime/MovieSceneTracks/Private/Channels/MovieSceneDoublePerlinNoiseChannel.cpp
--- a/Engine/Source/Runtime/MovieSceneTracks/Private/Channels/MovieSceneDoublePerlinNoiseChannel.cpp
+++ b/Engine/Source/Runtime/MovieSceneTracks/Private/Channels/MovieSceneDoublePerlinNoiseChannel.cpp
@@ -9,7 +9,7 @@ FMovieSceneDoublePerlinNoiseChannel::FMovieSceneDoublePerlinNoiseChannel()
 {
 }
 
-FMovieSceneDoublePerlinNoiseChannel::FMovieSceneDoublePerlinNoiseChannel(const FDoublePerlinNoiseParams& InPerlinNoiseParams)
+FMovieSceneDoublePerlinNoiseChannel::FMovieSceneDoublePerlinNoiseChannel(const FPerlinNoiseParams& InPerlinNoiseParams)
 	: PerlinNoiseParams{ InPerlinNoiseParams }
 {
 }
diff --git a/Engine/Source/Runtime/MovieSceneTracks/Private/Channels/MovieSceneDoublePerlinNoiseChannelContainer.cpp b/Engine/Source/Runtime/MovieSceneTracks/Private/Channels/MovieSceneDoublePerlinNoiseChannelContainer.cpp
--- a/Engine/Source/Runtime/MovieSceneTracks/Private/Channels/MovieSceneDoublePerlinNoiseChannelContainer.cpp
+++ b/Engine/Source/Runtime/MovieSceneTracks/Private/Channels/MovieSceneDoublePerlinNoiseChannelContainer.cpp
@@ -22,7 +22,7 @@ void UMovieSceneDoublePerlinNoiseChannelContainer::ImportEntityImpl(const FMovie
 
 	OutImportedEntity->AddBuilder(
 		FEntityBuilder()
-		.Add(TracksComponents->DoublePerlinNoiseChannel, PerlinNoiseChannel.GetParams())
+		.Add(TracksComponents->DoublePerlinNoiseChannel, PerlinNoiseChannel.PerlinNoiseParams)
 		.Add(ResultComponent, TNumericLimits<double>::Max())
 	);
 }
diff --git a/Engine/Source/Runtime/MovieSceneTracks/Private/Channels/MovieSceneFloatPerlinNoiseChannel.cpp b/Engine/Source/Runtime/MovieSceneTracks/Private/Channels/MovieSceneFloatPerlinNoiseChannel.cpp
--- a/Engine/Source/Runtime/MovieSceneTracks/Private/Channels/MovieSceneFloatPerlinNoiseChannel.cpp
+++ b/Engine/Source/Runtime/MovieSceneTracks/Private/Channels/MovieSceneFloatPerlinNoiseChannel.cpp
@@ -9,7 +9,7 @@ FMovieSceneFloatPerlinNoiseChannel::FMovieSceneFloatPerlinNoiseChannel()
 {
 }
 
-FMovieSceneFloatPerlinNoiseChannel::FMovieSceneFloatPerlinNoiseChannel(const FFloatPerlinNoiseParams& InFloatPerlinNoiseParams)
+FMovieSceneFloatPerlinNoiseChannel::FMovieSceneFloatPerlinNoiseChannel(const FPerlinNoiseParams& InFloatPerlinNoiseParams)
 	: PerlinNoiseParams{ InFloatPerlinNoiseParams }
 {
 }
diff --git a/Engine/Source/Runtime/MovieSceneTracks/Private/Channels/MovieSceneFloatPerlinNoiseChannelContainer.cpp b/Engine/Source/Runtime/MovieSceneTracks/Private/Channels/MovieSceneFloatPerlinNoiseChannelContainer.cpp
--- a/Engine/Source/Runtime/MovieSceneTracks/Private/Channels/MovieSceneFloatPerlinNoiseChannelContainer.cpp
+++ b/Engine/Source/Runtime/MovieSceneTracks/Private/Channels/MovieSceneFloatPerlinNoiseChannelContainer.cpp
@@ -23,7 +23,7 @@ void UMovieSceneFloatPerlinNoiseChannelContainer::ImportEntityImpl(const FMovieS
 
 	OutImportedEntity->AddBuilder(
 		FEntityBuilder()
-		.Add(TracksComponents->FloatPerlinNoiseChannel, PerlinNoiseChannel.GetParams())
+		.Add(TracksComponents->FloatPerlinNoiseChannel, PerlinNoiseChannel.PerlinNoiseParams)
 		.Add(ResultComponent, TNumericLimits<float>::Max())
 	);
 }
diff --git a/Engine/Source/Runtime/MovieSceneTracks/Private/MovieSceneTracksComponentTypes.cpp b/Engine/Source/Runtime/MovieSceneTracks/Private/MovieSceneTracksComponentTypes.cpp
--- a/Engine/Source/Runtime/MovieSceneTracks/Private/MovieSceneTracksComponentTypes.cpp
+++ b/Engine/Source/Runtime/MovieSceneTracks/Private/MovieSceneTracksComponentTypes.cpp
@@ -768,26 +768,15 @@ FMovieSceneTracksComponentTypes* FMovieSceneTracksComponentTypes::Get()
 } // namespace MovieScene
 } // namespace UE
 
-FFloatPerlinNoiseParams::FFloatPerlinNoiseParams()
+FPerlinNoiseParams::FPerlinNoiseParams()
 	: Frequency(4.0f)
 	, Amplitude(1.0f)
 {
 }
 
-FFloatPerlinNoiseParams::FFloatPerlinNoiseParams(float InFrequency, float InAmplitude)
+FPerlinNoiseParams::FPerlinNoiseParams(float InFrequency, double InAmplitude)
 	: Frequency(InFrequency)
 	, Amplitude(InAmplitude)
 {
 }
 
-FDoublePerlinNoiseParams::FDoublePerlinNoiseParams()
-	: Frequency(4.0)
-	, Amplitude(1.0)
-{
-}
-
-FDoublePerlinNoiseParams::FDoublePerlinNoiseParams(double InFrequency, double InAmplitude)
-	: Frequency(InFrequency)
-	, Amplitude(InAmplitude)
-{
-}
diff --git a/Engine/Source/Runtime/MovieSceneTracks/Public/Channels/MovieSceneDoublePerlinNoiseChannel.h b/Engine/Source/Runtime/MovieSceneTracks/Public/Channels/MovieSceneDoublePerlinNoiseChannel.h
--- a/Engine/Source/Runtime/MovieSceneTracks/Public/Channels/MovieSceneDoublePerlinNoiseChannel.h
+++ b/Engine/Source/Runtime/MovieSceneTracks/Public/Channels/MovieSceneDoublePerlinNoiseChannel.h
@@ -15,27 +15,29 @@ struct MOVIESCENETRACKS_API FMovieSceneDoublePerlinNoiseChannel : public FMovieS
 	using CurveValueType = double;
 
 	FMovieSceneDoublePerlinNoiseChannel();
-	FMovieSceneDoublePerlinNoiseChannel(const FDoublePerlinNoiseParams& InDoublePerlinNoiseParams);
+	FMovieSceneDoublePerlinNoiseChannel(const FPerlinNoiseParams& InDoublePerlinNoiseParams);
 
 	/**
-	* Evaluate this channel with the frame resolution
+	* Evaluate this channel at the given time
 	*
-	* @param InSeconds  The Frame second to evaluate at
+	* @param InSeconds  The time, in seconds, to evaluate at
 	* @return A value to receive the PerlinNoise result
 	*/
 	double Evaluate(double InSeconds) const;
 
+	/**
+	 * Evaluate this channel at the given time
+	 *
+	 * @param InSection  The section that contains this channel, used to lookup the sequence's tick resolution
+	 * @param InTime     The time, in ticks, to evaluate at
+	 * @param OutValue   The evaluated noise value
+	 * @return           Whether the noise was successfully evaluated
+	 */
 	bool Evaluate(const UMovieSceneSection* InSection, FFrameTime InTime, double& OutValue) const;
 
-	const FDoublePerlinNoiseParams& GetParams() const { return PerlinNoiseParams; }
-	FDoublePerlinNoiseParams& GetParams() { return PerlinNoiseParams; }
-
-private:
-
+	/** The noise parameters */
 	UPROPERTY(EditAnywhere, Category = "Perlin Noise")
-	FDoublePerlinNoiseParams PerlinNoiseParams;
-
-	friend class FMovieSceneDoublePerlinNoiseChannelDetailsCustomization;
+	FPerlinNoiseParams PerlinNoiseParams;
 };
 
 template<>
diff --git a/Engine/Source/Runtime/MovieSceneTracks/Public/Channels/MovieSceneFloatPerlinNoiseChannel.h b/Engine/Source/Runtime/MovieSceneTracks/Public/Channels/MovieSceneFloatPerlinNoiseChannel.h
--- a/Engine/Source/Runtime/MovieSceneTracks/Public/Channels/MovieSceneFloatPerlinNoiseChannel.h
+++ b/Engine/Source/Runtime/MovieSceneTracks/Public/Channels/MovieSceneFloatPerlinNoiseChannel.h
@@ -15,7 +15,7 @@ struct MOVIESCENETRACKS_API FMovieSceneFloatPerlinNoiseChannel : public FMovieSc
 	using CurveValueType = float;
 
 	FMovieSceneFloatPerlinNoiseChannel();
-	FMovieSceneFloatPerlinNoiseChannel(const FFloatPerlinNoiseParams& InFloatPerlinNoiseParams);
+	FMovieSceneFloatPerlinNoiseChannel(const FPerlinNoiseParams& InFloatPerlinNoiseParams);
 
 	/**
 	* Evaluate this channel with the frame resolution
@@ -25,17 +25,19 @@ struct MOVIESCENETRACKS_API FMovieSceneFloatPerlinNoiseChannel : public FMovieSc
 	*/
 	float Evaluate(double InSeconds) const;
 
+	/**
+	 * Evaluate this channel at the given time
+	 *
+	 * @param InSection  The section that contains this channel, used to lookup the sequence's tick resolution
+	 * @param InTime     The time, in ticks, to evaluate at
+	 * @param OutValue   The evaluated noise value
+	 * @return           Whether the noise was successfully evaluated
+	 */
 	bool Evaluate(const UMovieSceneSection* InSection, FFrameTime InTime, float& OutValue) const;
 
-	FFloatPerlinNoiseParams& GetParams() { return PerlinNoiseParams; }
-	const FFloatPerlinNoiseParams& GetParams() const { return PerlinNoiseParams; }
-
-private:
-
+	/** The noise parameters */
 	UPROPERTY(EditAnywhere, Category="Perlin Noise")
-	FFloatPerlinNoiseParams PerlinNoiseParams;
-
-	friend class FMovieSceneFloatPerlinNoiseChannelDetailsCustomization;
+	FPerlinNoiseParams PerlinNoiseParams;
 };
 
 template<>
diff --git a/Engine/Source/Runtime/MovieSceneTracks/Public/MovieSceneTracksComponentTypes.h b/Engine/Source/Runtime/MovieSceneTracks/Public/MovieSceneTracksComponentTypes.h
--- a/Engine/Source/Runtime/MovieSceneTracks/Public/MovieSceneTracksComponentTypes.h
+++ b/Engine/Source/Runtime/MovieSceneTracks/Public/MovieSceneTracksComponentTypes.h
@@ -21,9 +21,9 @@ class UMovieSceneLevelVisibilitySection;
 struct FMovieSceneObjectBindingID;
 
 
-/** Component data for the Float Perlin Noise Channel */
+/** Component data for Perlin Noise channels */
 USTRUCT()
-struct MOVIESCENETRACKS_API FFloatPerlinNoiseParams
+struct MOVIESCENETRACKS_API FPerlinNoiseParams
 {
 	GENERATED_BODY()
 
@@ -31,30 +31,12 @@ struct MOVIESCENETRACKS_API FFloatPerlinNoiseParams
 	UPROPERTY(EditAnywhere, Category="Perlin Noise")
 	float Frequency;
 
-	/** The amplitude of the noise, which will vary between [-Amplitude, +Amplitude] */
-	UPROPERTY(EditAnywhere, Category = "Perlin Noise")
-	float Amplitude;
-
-	FFloatPerlinNoiseParams();
-	FFloatPerlinNoiseParams(float InFrequency, float InAmplitude);
-};
-
-/** Component data for the Float Perlin Noise Channel */
-USTRUCT()
-struct MOVIESCENETRACKS_API FDoublePerlinNoiseParams
-{
-	GENERATED_BODY()
-
-	/** The frequency of the noise, i.e. how many times per second does the noise peak */
-	UPROPERTY(EditAnywhere, Category = "Perlin Noise")
-	double Frequency;
-
 	/** The amplitude of the noise, which will vary between [-Amplitude, +Amplitude] */
 	UPROPERTY(EditAnywhere, Category = "Perlin Noise")
 	double Amplitude;
 
-	FDoublePerlinNoiseParams();
-	FDoublePerlinNoiseParams(double InFrequency, double InAmplitude);
+	FPerlinNoiseParams();
+	FPerlinNoiseParams(float InFrequency, double InAmplitude);
 };
 
 /** Component data for the level visibility system */
@@ -430,8 +412,8 @@ struct MOVIESCENETRACKS_API FMovieSceneTracksComponentTypes
 	TComponentTypeID<USceneComponent*> AttachParent;
 	TComponentTypeID<FAttachmentComponent> AttachComponent;
 	TComponentTypeID<FMovieSceneObjectBindingID> AttachParentBinding;
-	TComponentTypeID<FFloatPerlinNoiseParams> FloatPerlinNoiseChannel;
-	TComponentTypeID<FDoublePerlinNoiseParams> DoublePerlinNoiseChannel;
+	TComponentTypeID<FPerlinNoiseParams> FloatPerlinNoiseChannel;
+	TComponentTypeID<FPerlinNoiseParams> DoublePerlinNoiseChannel;
 
 	TComponentTypeID<int32> ComponentMaterialIndex;
 
