Sequencer: Move structure into UE::MovieScene namespace

diff --git a/Engine/Source/Runtime/MovieScene/Private/Channels/MovieSceneSectionChannelOverrideRegistry.cpp b/Engine/Source/Runtime/MovieScene/Private/Channels/MovieSceneSectionChannelOverrideRegistry.cpp
--- a/Engine/Source/Runtime/MovieScene/Private/Channels/MovieSceneSectionChannelOverrideRegistry.cpp
+++ b/Engine/Source/Runtime/MovieScene/Private/Channels/MovieSceneSectionChannelOverrideRegistry.cpp
@@ -35,7 +35,7 @@ void UMovieSceneSectionChannelOverrideRegistry::RemoveChannel(FName ChannelName)
 	Overrides.Remove(ChannelName);
 }
 
-void UMovieSceneSectionChannelOverrideRegistry::ImportEntityImpl(const FMovieSceneChannelOverrideEntityImportParams& OverrideParams, const UE::MovieScene::FEntityImportParams& ImportParams, UE::MovieScene::FImportedEntity* OutImportedEntity)
+void UMovieSceneSectionChannelOverrideRegistry::ImportEntityImpl(const UE::MovieScene::FChannelOverrideEntityImportParams& OverrideParams, const UE::MovieScene::FEntityImportParams& ImportParams, UE::MovieScene::FImportedEntity* OutImportedEntity)
 {
 	TObjectPtr<UMovieSceneChannelOverrideContainer>* ContainerPtr = Overrides.Find(OverrideParams.ChannelName);
 	if (ensure(ContainerPtr))
diff --git a/Engine/Source/Runtime/MovieScene/Public/Channels/MovieSceneChannelOverrideContainer.h b/Engine/Source/Runtime/MovieScene/Public/Channels/MovieSceneChannelOverrideContainer.h
--- a/Engine/Source/Runtime/MovieScene/Public/Channels/MovieSceneChannelOverrideContainer.h
+++ b/Engine/Source/Runtime/MovieScene/Public/Channels/MovieSceneChannelOverrideContainer.h
@@ -21,21 +21,23 @@ struct FMovieSceneChannelMetaData;
 struct FMovieSceneChannelProxyData;
 struct FMovieSceneEntityComponentFieldBuilder;
 
+namespace UE
+{
+namespace MovieScene
+{
+
 /**
  * Entity import parameters for a channel override.
  */
-struct FMovieSceneChannelOverrideEntityImportParams
+struct FChannelOverrideEntityImportParams
 {
 	FName ChannelName;
 	UE::MovieScene::FComponentTypeID ResultComponent;
-
-	FMovieSceneChannelOverrideEntityImportParams()
-	{}
-	FMovieSceneChannelOverrideEntityImportParams(FName InChannelName, UE::MovieScene::FComponentTypeID InResultComponent)
-		: ChannelName(InChannelName), ResultComponent(InResultComponent)
-	{}
 };
 
+}  // namespace MovieScene
+}  // namespace UE
+
 /**
  * A wrapper to implement polymorphism for FMovieSceneChannel.
  */
@@ -51,7 +53,7 @@ public:
 
 	/** Imports the entity for this channel */
 	virtual void ImportEntityImpl(
-			const FMovieSceneChannelOverrideEntityImportParams& OverrideParams, 
+			const UE::MovieScene::FChannelOverrideEntityImportParams& OverrideParams, 
 			const UE::MovieScene::FEntityImportParams& ImportParams, UE::MovieScene::FImportedEntity* OutImportedEntity) {};
 
 	/** Gets the underlying channel */
diff --git a/Engine/Source/Runtime/MovieScene/Public/Channels/MovieSceneSectionChannelOverrideRegistry.h b/Engine/Source/Runtime/MovieScene/Public/Channels/MovieSceneSectionChannelOverrideRegistry.h
--- a/Engine/Source/Runtime/MovieScene/Public/Channels/MovieSceneSectionChannelOverrideRegistry.h
+++ b/Engine/Source/Runtime/MovieScene/Public/Channels/MovieSceneSectionChannelOverrideRegistry.h
@@ -99,7 +99,7 @@ public:
 	/**
 	* Forward ImportEntityImpl calls to an overriden channel
 	*/
-	void ImportEntityImpl(const FMovieSceneChannelOverrideEntityImportParams& OverrideParams, const UE::MovieScene::FEntityImportParams& ImportParams, UE::MovieScene::FImportedEntity* OutImportedEntity);
+	void ImportEntityImpl(const UE::MovieScene::FChannelOverrideEntityImportParams& OverrideParams, const UE::MovieScene::FEntityImportParams& ImportParams, UE::MovieScene::FImportedEntity* OutImportedEntity);
 
 	/**
 	* Called when overridden channels should populate evaluation field
diff --git a/Engine/Source/Runtime/MovieSceneTracks/Private/Channels/MovieSceneDoublePerlinNoiseChannelContainer.cpp b/Engine/Source/Runtime/MovieSceneTracks/Private/Channels/MovieSceneDoublePerlinNoiseChannelContainer.cpp
--- a/Engine/Source/Runtime/MovieSceneTracks/Private/Channels/MovieSceneDoublePerlinNoiseChannelContainer.cpp
+++ b/Engine/Source/Runtime/MovieSceneTracks/Private/Channels/MovieSceneDoublePerlinNoiseChannelContainer.cpp
@@ -11,7 +11,7 @@ bool UMovieSceneDoublePerlinNoiseChannelContainer::SupportsOverride(FName Defaul
 	return FMovieSceneDoubleChannel::StaticStruct()->GetFName() == DefaultChannelTypeName;
 }
 
-void UMovieSceneDoublePerlinNoiseChannelContainer::ImportEntityImpl(const FMovieSceneChannelOverrideEntityImportParams& OverrideParams, const UE::MovieScene::FEntityImportParams& ImportParams, UE::MovieScene::FImportedEntity* OutImportedEntity)
+void UMovieSceneDoublePerlinNoiseChannelContainer::ImportEntityImpl(const UE::MovieScene::FChannelOverrideEntityImportParams& OverrideParams, const UE::MovieScene::FEntityImportParams& ImportParams, UE::MovieScene::FImportedEntity* OutImportedEntity)
 {
 	using namespace UE::MovieScene;
 
diff --git a/Engine/Source/Runtime/MovieSceneTracks/Private/Channels/MovieSceneFloatPerlinNoiseChannelContainer.cpp b/Engine/Source/Runtime/MovieSceneTracks/Private/Channels/MovieSceneFloatPerlinNoiseChannelContainer.cpp
--- a/Engine/Source/Runtime/MovieSceneTracks/Private/Channels/MovieSceneFloatPerlinNoiseChannelContainer.cpp
+++ b/Engine/Source/Runtime/MovieSceneTracks/Private/Channels/MovieSceneFloatPerlinNoiseChannelContainer.cpp
@@ -12,7 +12,7 @@ bool UMovieSceneFloatPerlinNoiseChannelContainer::SupportsOverride(FName Default
 	return FMovieSceneFloatChannel::StaticStruct()->GetFName() == DefaultChannelTypeName;
 }
 
-void UMovieSceneFloatPerlinNoiseChannelContainer::ImportEntityImpl(const FMovieSceneChannelOverrideEntityImportParams& OverrideParams, const UE::MovieScene::FEntityImportParams& ImportParams, UE::MovieScene::FImportedEntity* OutImportedEntity)
+void UMovieSceneFloatPerlinNoiseChannelContainer::ImportEntityImpl(const UE::MovieScene::FChannelOverrideEntityImportParams& OverrideParams, const UE::MovieScene::FEntityImportParams& ImportParams, UE::MovieScene::FImportedEntity* OutImportedEntity)
 {
 	using namespace UE::MovieScene;
 
diff --git a/Engine/Source/Runtime/MovieSceneTracks/Private/Sections/MovieScene3DTransformSection.cpp b/Engine/Source/Runtime/MovieSceneTracks/Private/Sections/MovieScene3DTransformSection.cpp
--- a/Engine/Source/Runtime/MovieSceneTracks/Private/Sections/MovieScene3DTransformSection.cpp
+++ b/Engine/Source/Runtime/MovieSceneTracks/Private/Sections/MovieScene3DTransformSection.cpp
@@ -421,7 +421,7 @@ void UMovieScene3DTransformSection::BuildEntity(BaseBuilderType& InBaseBuilder,
 		{
 			if (EnumHasAnyFlags(Channels, EMovieSceneTransformChannel::Weight))
 			{
-				FMovieSceneChannelOverrideEntityImportParams OverrideParams(ChannelOverrideName, BuiltInComponentTypes->WeightResult);
+				FChannelOverrideEntityImportParams OverrideParams{ ChannelOverrideName, BuiltInComponentTypes->WeightResult };
 				OverrideRegistry->ImportEntityImpl(OverrideParams, Params, OutImportedEntity);
 			}
 		}
@@ -429,7 +429,7 @@ void UMovieScene3DTransformSection::BuildEntity(BaseBuilderType& InBaseBuilder,
 		{
 			if (ActiveChannelsMask[ChannelOverrideIndex])
 			{
-				FMovieSceneChannelOverrideEntityImportParams OverrideParams(ChannelOverrideName, BuiltInComponentTypes->DoubleResult[ChannelOverrideIndex]);
+				FChannelOverrideEntityImportParams OverrideParams{ ChannelOverrideName, BuiltInComponentTypes->DoubleResult[ChannelOverrideIndex] };
 				OverrideRegistry->ImportEntityImpl(OverrideParams, Params, OutImportedEntity);
 			}
 		}
diff --git a/Engine/Source/Runtime/MovieSceneTracks/Private/Sections/MovieSceneFloatSection.cpp b/Engine/Source/Runtime/MovieSceneTracks/Private/Sections/MovieSceneFloatSection.cpp
--- a/Engine/Source/Runtime/MovieSceneTracks/Private/Sections/MovieSceneFloatSection.cpp
+++ b/Engine/Source/Runtime/MovieSceneTracks/Private/Sections/MovieSceneFloatSection.cpp
@@ -65,8 +65,8 @@ void UMovieSceneFloatSection::ImportEntityImpl(UMovieSceneEntitySystemLinker* En
 
 	FPropertyTrackWithOverridableChannelsEntityImportHelper(TracksComponents->Float, this)
 		.Add(Components->FloatChannel[0], &FloatCurve,
-				// NAME_None below because we don't specify any names in CacheChannelProxy
-				FMovieSceneChannelOverrideEntityImportParams(NAME_None, Components->DoubleResult[0]))
+			// NAME_None below because we don't specify any names in CacheChannelProxy
+			FChannelOverrideEntityImportParams{ NAME_None, Components->DoubleResult[0] })
 		.Commit(this, Params, OutImportedEntity);
 }
 
diff --git a/Engine/Source/Runtime/MovieSceneTracks/Public/Channels/MovieSceneDoublePerlinNoiseChannelContainer.h b/Engine/Source/Runtime/MovieSceneTracks/Public/Channels/MovieSceneDoublePerlinNoiseChannelContainer.h
--- a/Engine/Source/Runtime/MovieSceneTracks/Public/Channels/MovieSceneDoublePerlinNoiseChannelContainer.h
+++ b/Engine/Source/Runtime/MovieSceneTracks/Public/Channels/MovieSceneDoublePerlinNoiseChannelContainer.h
@@ -20,7 +20,7 @@ public:
 
 	bool SupportsOverride(FName DefaultChannelTypeName) const override;
 	void ImportEntityImpl(
-			const FMovieSceneChannelOverrideEntityImportParams& OverrideParams, 
+			const UE::MovieScene::FChannelOverrideEntityImportParams& OverrideParams, 
 			const UE::MovieScene::FEntityImportParams& ImportParams, UE::MovieScene::FImportedEntity* OutImportedEntity) override;
 
 	const FMovieSceneChannel* GetChannel() const override { return &PerlinNoiseChannel; }
diff --git a/Engine/Source/Runtime/MovieSceneTracks/Public/Channels/MovieSceneFloatPerlinNoiseChannelContainer.h b/Engine/Source/Runtime/MovieSceneTracks/Public/Channels/MovieSceneFloatPerlinNoiseChannelContainer.h
--- a/Engine/Source/Runtime/MovieSceneTracks/Public/Channels/MovieSceneFloatPerlinNoiseChannelContainer.h
+++ b/Engine/Source/Runtime/MovieSceneTracks/Public/Channels/MovieSceneFloatPerlinNoiseChannelContainer.h
@@ -21,7 +21,7 @@ public:
 
 	bool SupportsOverride(FName DefaultChannelTypeName) const override;
 	void ImportEntityImpl(
-			const FMovieSceneChannelOverrideEntityImportParams& OverrideParams, 
+			const UE::MovieScene::FChannelOverrideEntityImportParams& OverrideParams, 
 			const UE::MovieScene::FEntityImportParams& ImportParams, UE::MovieScene::FImportedEntity* OutImportedEntity) override;
 
 	const FMovieSceneChannel* GetChannel() const override { return &PerlinNoiseChannel; }
diff --git a/Engine/Source/Runtime/MovieSceneTracks/Public/Tracks/MovieScenePropertyTrack.h b/Engine/Source/Runtime/MovieSceneTracks/Public/Tracks/MovieScenePropertyTrack.h
--- a/Engine/Source/Runtime/MovieSceneTracks/Public/Tracks/MovieScenePropertyTrack.h
+++ b/Engine/Source/Runtime/MovieSceneTracks/Public/Tracks/MovieScenePropertyTrack.h
@@ -283,7 +283,7 @@ struct FPropertyTrackEntityImportHelperParamsImpl
 	}
 	TScriptInterface<IMovieSceneChannelOverrideProvider> RegistryProvider;
 	TObjectPtr<UMovieSceneSectionChannelOverrideRegistry> Registry;
-	TArray<FMovieSceneChannelOverrideEntityImportParams> ImportParams;
+	TArray<FChannelOverrideEntityImportParams> ImportParams;
 };
 
 /**
@@ -300,7 +300,7 @@ struct TPropertyTrackWithOverridableChannelsEntityImportHelperImpl
 	}
 
 	template<typename U, typename PayloadType>
-	TPropertyTrackEntityImportHelperImpl<T..., TAddConditional<U>> Add(TComponentTypeID<U> ComponentType, PayloadType&& InPayload, FMovieSceneChannelOverrideEntityImportParams OverrideParams)
+	TPropertyTrackEntityImportHelperImpl<T..., TAddConditional<U>> Add(TComponentTypeID<U> ComponentType, PayloadType&& InPayload, FChannelOverrideEntityImportParams OverrideParams)
 	{
 		const bool bOverriden = (OverrideInfo.Registry && OverrideInfo.Registry->ContainsChannel(OverrideParams.ChannelName));
 		if (bOverriden)
@@ -311,7 +311,7 @@ struct TPropertyTrackWithOverridableChannelsEntityImportHelperImpl
 	}
 
 	template<typename U, typename PayloadType>
-	TPropertyTrackEntityImportHelperImpl<T..., TAddConditional<U>> AddConditional(TComponentTypeID<U> ComponentType, PayloadType&& InPayload, bool bCondition, FMovieSceneChannelOverrideEntityImportParams OverrideParams)
+	TPropertyTrackEntityImportHelperImpl<T..., TAddConditional<U>> AddConditional(TComponentTypeID<U> ComponentType, PayloadType&& InPayload, bool bCondition, FChannelOverrideEntityImportParams OverrideParams)
 	{
 		const bool bOverriden = (OverrideInfo.Registry && OverrideInfo.Registry->ContainsChannel(OverrideParams.ChannelName));
 		if (bOverriden)
@@ -357,7 +357,7 @@ struct TPropertyTrackWithOverridableChannelsEntityImportHelperImpl
 					check(Traits.IsValid());
 					const FName ChannelOverrideName = Traits->GetChannelOverrideName(Params.EntityID);
 					auto* ChannelOverrideParams = OverrideInfo.ImportParams.FindByPredicate(
-							[=](const FMovieSceneChannelOverrideEntityImportParams& CurParams)
+							[=](const FChannelOverrideEntityImportParams& CurParams)
 							{
 								return CurParams.ChannelName == ChannelOverrideName;
 							});
@@ -398,7 +398,7 @@ struct TPropertyTrackWithOverridableChannelsEntityImportHelperImpl<>
 	}
 
 	template<typename U, typename PayloadType>
-	TPropertyTrackWithOverridableChannelsEntityImportHelperImpl<FAdd, TAddConditional<U>> Add(TComponentTypeID<U> ComponentType, PayloadType&& InPayload, FMovieSceneChannelOverrideEntityImportParams OverrideParams)
+	TPropertyTrackWithOverridableChannelsEntityImportHelperImpl<FAdd, TAddConditional<U>> Add(TComponentTypeID<U> ComponentType, PayloadType&& InPayload, FChannelOverrideEntityImportParams OverrideParams)
 	{
 		const bool bOverriden = (OverrideInfo.Registry && OverrideInfo.Registry->ContainsChannel(OverrideParams.ChannelName));
 		if (bOverriden)
@@ -409,7 +409,7 @@ struct TPropertyTrackWithOverridableChannelsEntityImportHelperImpl<>
 	}
 
 	template<typename U, typename PayloadType>
-	TPropertyTrackWithOverridableChannelsEntityImportHelperImpl<FAdd, TAddConditional<U>> AddConditional(TComponentTypeID<U> ComponentType, PayloadType&& InPayload, bool bCondition, FMovieSceneChannelOverrideEntityImportParams OverrideParams)
+	TPropertyTrackWithOverridableChannelsEntityImportHelperImpl<FAdd, TAddConditional<U>> AddConditional(TComponentTypeID<U> ComponentType, PayloadType&& InPayload, bool bCondition, FChannelOverrideEntityImportParams OverrideParams)
 	{
 		const bool bOverriden = (OverrideInfo.Registry && OverrideInfo.Registry->ContainsChannel(OverrideParams.ChannelName));
 		if (bOverriden)
