Sequencer: NLA fix root motion not blending based upon section order evaluation changing, originally assumption was the 'next' section would always evalauate first. To fix we just set the root motion override on both sections, but make sure to not do so before root motion overrides start.

diff --git a/Engine/Source/Runtime/MovieSceneTracks/Private/Sections/MovieSceneSkeletalAnimationSection.cpp b/Engine/Source/Runtime/MovieSceneTracks/Private/Sections/MovieSceneSkeletalAnimationSection.cpp
--- a/Engine/Source/Runtime/MovieSceneTracks/Private/Sections/MovieSceneSkeletalAnimationSection.cpp
+++ b/Engine/Source/Runtime/MovieSceneTracks/Private/Sections/MovieSceneSkeletalAnimationSection.cpp
@@ -394,8 +394,7 @@ void UMovieSceneSkeletalAnimationSection::PostEditUndo()
 
 TOptional<FTransform> UMovieSceneSkeletalAnimationSection::GetRootMotion(FFrameTime CurrentTime, bool& bBlendFirstChildOfRoot) const
 {
-	//if we don't have an offset transform set then it's fine to return nothing, we just use the default root. Also use smaller tolerance here just in case
-	if (GetRootMotionParams() && GetOffsetTransform().Equals(FTransform::Identity, KINDA_SMALL_NUMBER)  == false)
+	if (GetRootMotionParams())
 	{
 		UMovieSceneSkeletalAnimationTrack* Track = GetTypedOuter<UMovieSceneSkeletalAnimationTrack>();
 		if (GetRootMotionParams()->bRootMotionsDirty)
@@ -686,4 +685,3 @@ void UMovieSceneSkeletalAnimationSection::ToggleShowSkeleton()
 
 
 #undef LOCTEXT_NAMESPACE 
-
diff --git a/Engine/Source/Runtime/MovieSceneTracks/Private/Tracks/MovieSceneSkeletalAnimationTrack.cpp b/Engine/Source/Runtime/MovieSceneTracks/Private/Tracks/MovieSceneSkeletalAnimationTrack.cpp
--- a/Engine/Source/Runtime/MovieSceneTracks/Private/Tracks/MovieSceneSkeletalAnimationTrack.cpp
+++ b/Engine/Source/Runtime/MovieSceneTracks/Private/Tracks/MovieSceneSkeletalAnimationTrack.cpp
@@ -1072,7 +1072,7 @@ void UMovieSceneSkeletalAnimationTrack::ToggleShowRootMotionTrail()
 	bShowRootMotionTrail = bShowRootMotionTrail ? false : true;
 }
 #endif
-FTransform FMovieSceneSkeletalAnimRootMotionTrackParams::GetRootMotion(FFrameTime CurrentTime)  const
+TOptional<FTransform>  FMovieSceneSkeletalAnimRootMotionTrackParams::GetRootMotion(FFrameTime CurrentTime)  const
 {
 	if (RootTransforms.Num() > 0)
 	{
@@ -1113,7 +1113,8 @@ FTransform FMovieSceneSkeletalAnimRootMotionTrackParams::GetRootMotion(FFrameTim
 			return RootTransforms[0];
 		}
 	}
-	return FTransform::Identity;
+	return TOptional<FTransform>();
+
 }
 
 
diff --git a/Engine/Source/Runtime/MovieSceneTracks/Public/Tracks/MovieSceneSkeletalAnimationTrack.h b/Engine/Source/Runtime/MovieSceneTracks/Public/Tracks/MovieSceneSkeletalAnimationTrack.h
--- a/Engine/Source/Runtime/MovieSceneTracks/Public/Tracks/MovieSceneSkeletalAnimationTrack.h
+++ b/Engine/Source/Runtime/MovieSceneTracks/Public/Tracks/MovieSceneSkeletalAnimationTrack.h
@@ -26,7 +26,7 @@ struct FMovieSceneSkeletalAnimRootMotionTrackParams
 	TArray<FTransform> RootTransforms;
 
 	/** Get the Root Motion transform at the specified time.*/
-	MOVIESCENETRACKS_API FTransform GetRootMotion(FFrameTime CurrentTime) const;
+	MOVIESCENETRACKS_API TOptional<FTransform>  GetRootMotion(FFrameTime CurrentTime) const;
 	FMovieSceneSkeletalAnimRootMotionTrackParams() : bRootMotionsDirty(true) {}
 
 };
