Sequencer: Don't create useless entities when we are dealing with a muted override channel

It's possible for PopulateEvaluationField to create evaluation field entries for muted override channels if it doesn't have access to the editor data that tells us if a channel is enabled or not. In this case, we will be asked to create an entity builder but we will skip most of it since we will realize the channel was muted. There was a bug however where we would still add entity builders for the basic stuff (property binding, object binding, easing, etc). This would of course do nothing useful and was just wasteful. So this fix only adds these entity builders once we know the channel is good to go.

diff --git a/Engine/Source/Runtime/MovieSceneTracks/Private/Sections/MovieScene3DTransformSection.cpp b/Engine/Source/Runtime/MovieSceneTracks/Private/Sections/MovieScene3DTransformSection.cpp
--- a/Engine/Source/Runtime/MovieSceneTracks/Private/Sections/MovieScene3DTransformSection.cpp
+++ b/Engine/Source/Runtime/MovieSceneTracks/Private/Sections/MovieScene3DTransformSection.cpp
@@ -458,14 +458,19 @@ void UMovieScene3DTransformSection::BuildEntity(BaseBuilderType& InBaseBuilder,
 	// Then call into the override repository so it can add a second builder with the override stuff.
 	if (ChannelOverrideName != NAME_None)
 	{
-		check(OverrideRegistry && OverrideRegistry->ContainsChannel(ChannelOverrideName));
-		OutImportedEntity->AddBuilder(InBaseBuilder.AddTag(PropertyTag));
+		if (!ensureMsgf(OverrideRegistry && OverrideRegistry->ContainsChannel(ChannelOverrideName),
+				TEXT("We received an entity ID that corresponds to an override channel, but no such override channel exists in the registry.")))
+		{
+			return;
+		}
 
 		const int32 ChannelOverrideIndex = (Params.EntityID - ChannelOverrideNames.IndexOffset);
 		if (ChannelOverrideIndex == 9)
 		{
 			if (EnumHasAnyFlags(Channels, EMovieSceneTransformChannel::Weight))
 			{
+				OutImportedEntity->AddBuilder(InBaseBuilder.AddTag(PropertyTag));
+
 				FChannelOverrideEntityImportParams OverrideParams{ ChannelOverrideName, BuiltInComponentTypes->WeightResult };
 				OverrideRegistry->ImportEntityImpl(OverrideParams, Params, OutImportedEntity);
 			}
@@ -474,6 +479,8 @@ void UMovieScene3DTransformSection::BuildEntity(BaseBuilderType& InBaseBuilder,
 		{
 			if (ActiveChannelsMask[ChannelOverrideIndex])
 			{
+				OutImportedEntity->AddBuilder(InBaseBuilder.AddTag(PropertyTag));
+
 				FChannelOverrideEntityImportParams OverrideParams{ ChannelOverrideName, BuiltInComponentTypes->DoubleResult[ChannelOverrideIndex] };
 				OverrideRegistry->ImportEntityImpl(OverrideParams, Params, OutImportedEntity);
 			}
