Sequencer: Add support for perlin noise initial offset

Also add an initialization phase to channel override containers. In the case of the perlin noise channels, this is used to generate a random initial offset.

diff --git a/Engine/Source/Editor/Sequencer/Private/MVVM/ViewModels/ChannelModel.cpp b/Engine/Source/Editor/Sequencer/Private/MVVM/ViewModels/ChannelModel.cpp
--- a/Engine/Source/Editor/Sequencer/Private/MVVM/ViewModels/ChannelModel.cpp
+++ b/Engine/Source/Editor/Sequencer/Private/MVVM/ViewModels/ChannelModel.cpp
@@ -568,6 +568,7 @@ void FChannelGroupModel::OverrideChannels(TSubclassOf<UMovieSceneChannelOverride
 		check(OverrideRegistry);
 
 		UMovieSceneChannelOverrideContainer* OverrideChannelContainer = NewObject<UMovieSceneChannelOverrideContainer>(Section, OverrideClass);
+		OverrideChannelContainer->InitializeOverride(Channel->GetChannel());
 		OverrideRegistry->AddChannel(GetChannelName(), OverrideChannelContainer);	
 		
 		OverrideProvider->OnChannelOverridesChanged();
diff --git a/Engine/Source/Runtime/MovieScene/Public/Channels/MovieSceneChannelOverrideContainer.h b/Engine/Source/Runtime/MovieScene/Public/Channels/MovieSceneChannelOverrideContainer.h
--- a/Engine/Source/Runtime/MovieScene/Public/Channels/MovieSceneChannelOverrideContainer.h
+++ b/Engine/Source/Runtime/MovieScene/Public/Channels/MovieSceneChannelOverrideContainer.h
@@ -48,6 +48,12 @@ class MOVIESCENE_API UMovieSceneChannelOverrideContainer : public UMovieSceneSig
 
 public:
 
+	/** Initialize this container 
+	 *
+	 * @param InChannel		The channel being overriden
+	 */
+	virtual void InitializeOverride(FMovieSceneChannel* InChannel) {}
+
 	/** Returns whether this container's underlying channel can be used as an override to the given channel type */
 	virtual bool SupportsOverride(FName DefaultChannelTypeName) const { return false; }
 
diff --git a/Engine/Source/Runtime/MovieSceneTracks/Private/Channels/MovieSceneDoublePerlinNoiseChannel.cpp b/Engine/Source/Runtime/MovieSceneTracks/Private/Channels/MovieSceneDoublePerlinNoiseChannel.cpp
--- a/Engine/Source/Runtime/MovieSceneTracks/Private/Channels/MovieSceneDoublePerlinNoiseChannel.cpp
+++ b/Engine/Source/Runtime/MovieSceneTracks/Private/Channels/MovieSceneDoublePerlinNoiseChannel.cpp
@@ -7,20 +7,12 @@
 #include UE_INLINE_GENERATED_CPP_BY_NAME(MovieSceneDoublePerlinNoiseChannel)
 
 FMovieSceneDoublePerlinNoiseChannel::FMovieSceneDoublePerlinNoiseChannel()
-	: PerlinNoiseParams{}
-{
-}
-
-FMovieSceneDoublePerlinNoiseChannel::FMovieSceneDoublePerlinNoiseChannel(const FPerlinNoiseParams& InPerlinNoiseParams)
-	: PerlinNoiseParams{ InPerlinNoiseParams }
 {
 }
 
 double FMovieSceneDoublePerlinNoiseChannel::Evaluate(double InSeconds) const
 {
-	double Result = static_cast<double>(FMath::PerlinNoise1D(InSeconds * PerlinNoiseParams.Frequency));
-	Result *= PerlinNoiseParams.Amplitude;
-	return Result;
+	return Evaluate(PerlinNoiseParams, InSeconds);
 }
 
 bool FMovieSceneDoublePerlinNoiseChannel::Evaluate(const UMovieSceneSection* InSection, FFrameTime InTime, double& OutValue) const
@@ -39,3 +31,10 @@ bool FMovieSceneDoublePerlinNoiseChannel::Evaluate(const UMovieSceneSection* InS
 	return false;
 }
 
+double FMovieSceneDoublePerlinNoiseChannel::Evaluate(const FPerlinNoiseParams& InParams, double InSeconds)
+{
+	double Result = static_cast<double>(FMath::PerlinNoise1D((InSeconds + InParams.Offset) * InParams.Frequency));
+	Result *= InParams.Amplitude;
+	return Result;
+}
+
diff --git a/Engine/Source/Runtime/MovieSceneTracks/Private/Channels/MovieSceneDoublePerlinNoiseChannelContainer.cpp b/Engine/Source/Runtime/MovieSceneTracks/Private/Channels/MovieSceneDoublePerlinNoiseChannelContainer.cpp
--- a/Engine/Source/Runtime/MovieSceneTracks/Private/Channels/MovieSceneDoublePerlinNoiseChannelContainer.cpp
+++ b/Engine/Source/Runtime/MovieSceneTracks/Private/Channels/MovieSceneDoublePerlinNoiseChannelContainer.cpp
@@ -8,6 +8,11 @@
 
 #include UE_INLINE_GENERATED_CPP_BY_NAME(MovieSceneDoublePerlinNoiseChannelContainer)
 
+void UMovieSceneDoublePerlinNoiseChannelContainer::InitializeOverride(FMovieSceneChannel* InChannel)
+{
+	PerlinNoiseChannel.PerlinNoiseParams.RandomizeOffset();
+}
+
 bool UMovieSceneDoublePerlinNoiseChannelContainer::SupportsOverride(FName DefaultChannelTypeName) const
 {
 	return FMovieSceneDoubleChannel::StaticStruct()->GetFName() == DefaultChannelTypeName;
diff --git a/Engine/Source/Runtime/MovieSceneTracks/Private/Channels/MovieSceneFloatPerlinNoiseChannel.cpp b/Engine/Source/Runtime/MovieSceneTracks/Private/Channels/MovieSceneFloatPerlinNoiseChannel.cpp
--- a/Engine/Source/Runtime/MovieSceneTracks/Private/Channels/MovieSceneFloatPerlinNoiseChannel.cpp
+++ b/Engine/Source/Runtime/MovieSceneTracks/Private/Channels/MovieSceneFloatPerlinNoiseChannel.cpp
@@ -7,20 +7,12 @@
 #include UE_INLINE_GENERATED_CPP_BY_NAME(MovieSceneFloatPerlinNoiseChannel)
 
 FMovieSceneFloatPerlinNoiseChannel::FMovieSceneFloatPerlinNoiseChannel()
-	: PerlinNoiseParams{}
-{
-}
-
-FMovieSceneFloatPerlinNoiseChannel::FMovieSceneFloatPerlinNoiseChannel(const FPerlinNoiseParams& InFloatPerlinNoiseParams)
-	: PerlinNoiseParams{ InFloatPerlinNoiseParams }
 {
 }
 
 float FMovieSceneFloatPerlinNoiseChannel::Evaluate(double InSeconds) const
 {
-	float Result = FMath::PerlinNoise1D(InSeconds * PerlinNoiseParams.Frequency);
-	Result *= PerlinNoiseParams.Amplitude;
-	return Result;
+	return Evaluate(PerlinNoiseParams, InSeconds);
 }
 
 bool FMovieSceneFloatPerlinNoiseChannel::Evaluate(const UMovieSceneSection* InSection, FFrameTime InTime, float& OutValue) const
@@ -39,3 +31,10 @@ bool FMovieSceneFloatPerlinNoiseChannel::Evaluate(const UMovieSceneSection* InSe
 	return false;
 }
 
+float FMovieSceneFloatPerlinNoiseChannel::Evaluate(const FPerlinNoiseParams& InParams, double InSeconds)
+{
+	float Result = FMath::PerlinNoise1D((InSeconds + InParams.Offset) * InParams.Frequency);
+	Result *= InParams.Amplitude;
+	return Result;
+}
+
diff --git a/Engine/Source/Runtime/MovieSceneTracks/Private/Channels/MovieSceneFloatPerlinNoiseChannelContainer.cpp b/Engine/Source/Runtime/MovieSceneTracks/Private/Channels/MovieSceneFloatPerlinNoiseChannelContainer.cpp
--- a/Engine/Source/Runtime/MovieSceneTracks/Private/Channels/MovieSceneFloatPerlinNoiseChannelContainer.cpp
+++ b/Engine/Source/Runtime/MovieSceneTracks/Private/Channels/MovieSceneFloatPerlinNoiseChannelContainer.cpp
@@ -9,6 +9,11 @@
 
 #include UE_INLINE_GENERATED_CPP_BY_NAME(MovieSceneFloatPerlinNoiseChannelContainer)
 
+void UMovieSceneFloatPerlinNoiseChannelContainer::InitializeOverride(FMovieSceneChannel* InChannel)
+{
+	PerlinNoiseChannel.PerlinNoiseParams.RandomizeOffset();
+}
+
 bool UMovieSceneFloatPerlinNoiseChannelContainer::SupportsOverride(FName DefaultChannelTypeName) const
 {
 	return FMovieSceneFloatChannel::StaticStruct()->GetFName() == DefaultChannelTypeName;
diff --git a/Engine/Source/Runtime/MovieSceneTracks/Private/MovieSceneTracksComponentTypes.cpp b/Engine/Source/Runtime/MovieSceneTracks/Private/MovieSceneTracksComponentTypes.cpp
--- a/Engine/Source/Runtime/MovieSceneTracks/Private/MovieSceneTracksComponentTypes.cpp
+++ b/Engine/Source/Runtime/MovieSceneTracks/Private/MovieSceneTracksComponentTypes.cpp
@@ -773,13 +773,19 @@ FMovieSceneTracksComponentTypes* FMovieSceneTracksComponentTypes::Get()
 FPerlinNoiseParams::FPerlinNoiseParams()
 	: Frequency(4.0f)
 	, Amplitude(1.0f)
+	, Offset(0)
 {
 }
 
 FPerlinNoiseParams::FPerlinNoiseParams(float InFrequency, double InAmplitude)
 	: Frequency(InFrequency)
 	, Amplitude(InAmplitude)
+	, Offset(0)
 {
 }
 
+void FPerlinNoiseParams::RandomizeOffset(float InMaxOffset)
+{
+	Offset = FMath::FRand() * InMaxOffset;
+}
 
diff --git a/Engine/Source/Runtime/MovieSceneTracks/Private/Systems/DoublePerlinNoiseChannelEvaluatorSystem.cpp b/Engine/Source/Runtime/MovieSceneTracks/Private/Systems/DoublePerlinNoiseChannelEvaluatorSystem.cpp
--- a/Engine/Source/Runtime/MovieSceneTracks/Private/Systems/DoublePerlinNoiseChannelEvaluatorSystem.cpp
+++ b/Engine/Source/Runtime/MovieSceneTracks/Private/Systems/DoublePerlinNoiseChannelEvaluatorSystem.cpp
@@ -28,17 +28,15 @@ namespace UE
 				check(InstanceRegistry);
 			}
 
-			void ForEachEntity(FMovieSceneDoublePerlinNoiseChannel DoublePerlinNoiseChannel, FInstanceHandle InstanceHandle, double& OutResult)
+			void ForEachEntity(const FPerlinNoiseParams& PerlinNoiseParams, FInstanceHandle InstanceHandle, double& OutResult)
 			{
 				const FMovieSceneContext& Context = InstanceRegistry->GetContext(InstanceHandle);
 				FFrameTime Time = Context.GetTime();
 				FFrameRate Rate = Context.GetFrameRate();
 
-				OutResult = DoublePerlinNoiseChannel.Evaluate(Rate.AsSeconds(Time));
+				OutResult = FMovieSceneDoublePerlinNoiseChannel::Evaluate(PerlinNoiseParams, Rate.AsSeconds(Time));
 			}
 		};
-
-
 	} // namespace MovieScene
 } // namespace UE
 
diff --git a/Engine/Source/Runtime/MovieSceneTracks/Private/Systems/FloatPerlinNoiseChannelEvaluatorSystem.cpp b/Engine/Source/Runtime/MovieSceneTracks/Private/Systems/FloatPerlinNoiseChannelEvaluatorSystem.cpp
--- a/Engine/Source/Runtime/MovieSceneTracks/Private/Systems/FloatPerlinNoiseChannelEvaluatorSystem.cpp
+++ b/Engine/Source/Runtime/MovieSceneTracks/Private/Systems/FloatPerlinNoiseChannelEvaluatorSystem.cpp
@@ -28,18 +28,16 @@ namespace UE
 				check(InstanceRegistry);
 			}
 
-			void ForEachEntity(FMovieSceneFloatPerlinNoiseChannel FloatPerlinNoiseChannel, FInstanceHandle InstanceHandle, double& OutResult)
+			void ForEachEntity(const FPerlinNoiseParams& PerlinNoiseParams, FInstanceHandle InstanceHandle, double& OutResult)
 			{
 				const FMovieSceneContext& Context = InstanceRegistry->GetContext(InstanceHandle);
 				FFrameTime Time = Context.GetTime();
 				FFrameRate Rate = Context.GetFrameRate();
 
-				const float NoiseResult = FloatPerlinNoiseChannel.Evaluate(Rate.AsSeconds(Time));
+				const float NoiseResult = FMovieSceneFloatPerlinNoiseChannel::Evaluate(PerlinNoiseParams, Rate.AsSeconds(Time));
 				OutResult = (double)NoiseResult;
 			}
 		};
-
-
 	} // namespace MovieScene
 } // namespace UE
 
diff --git a/Engine/Source/Runtime/MovieSceneTracks/Public/Channels/MovieSceneDoublePerlinNoiseChannel.h b/Engine/Source/Runtime/MovieSceneTracks/Public/Channels/MovieSceneDoublePerlinNoiseChannel.h
--- a/Engine/Source/Runtime/MovieSceneTracks/Public/Channels/MovieSceneDoublePerlinNoiseChannel.h
+++ b/Engine/Source/Runtime/MovieSceneTracks/Public/Channels/MovieSceneDoublePerlinNoiseChannel.h
@@ -15,29 +15,38 @@ struct MOVIESCENETRACKS_API FMovieSceneDoublePerlinNoiseChannel : public FMovieS
 	using CurveValueType = double;
 
 	FMovieSceneDoublePerlinNoiseChannel();
-	FMovieSceneDoublePerlinNoiseChannel(const FPerlinNoiseParams& InDoublePerlinNoiseParams);
 
 	/**
-	* Evaluate this channel at the given time
-	*
-	* @param InSeconds  The time, in seconds, to evaluate at
-	* @return A value to receive the PerlinNoise result
-	*/
+	 * Evaluate this channel at the given time
+	 *
+	 * @param InSeconds		The time, in seconds, to evaluate at
+	 * @return				The evaluated noise value
+	 */
 	double Evaluate(double InSeconds) const;
 
 	/**
 	 * Evaluate this channel at the given time
 	 *
-	 * @param InSection  The section that contains this channel, used to lookup the sequence's tick resolution
-	 * @param InTime     The time, in ticks, to evaluate at
-	 * @param OutValue   The evaluated noise value
-	 * @return           Whether the noise was successfully evaluated
+	 * @param InSection		The section that contains this channel, used to lookup the sequence's tick resolution
+	 * @param InTime		The time, in ticks, to evaluate at
+	 * @param OutValue		The evaluated noise value
+	 * @return				Whether the noise was successfully evaluated
 	 */
 	bool Evaluate(const UMovieSceneSection* InSection, FFrameTime InTime, double& OutValue) const;
 
 	/** The noise parameters */
 	UPROPERTY(EditAnywhere, Category = "Perlin Noise")
 	FPerlinNoiseParams PerlinNoiseParams;
+
+	/**
+	 * Evaluate perlin noise
+	 *
+	 * @params InParams			Perlin noise parameters
+	 * @params InInitialOffset	The initial offset for the noise
+	 * @params InSeconds		The time at which to evaluate the noise
+	 * @return					The evaluated noise value
+	 */
+	static double Evaluate(const FPerlinNoiseParams& InParams, double InSeconds);
 };
 
 template<>
diff --git a/Engine/Source/Runtime/MovieSceneTracks/Public/Channels/MovieSceneDoublePerlinNoiseChannelContainer.h b/Engine/Source/Runtime/MovieSceneTracks/Public/Channels/MovieSceneDoublePerlinNoiseChannelContainer.h
--- a/Engine/Source/Runtime/MovieSceneTracks/Public/Channels/MovieSceneDoublePerlinNoiseChannelContainer.h
+++ b/Engine/Source/Runtime/MovieSceneTracks/Public/Channels/MovieSceneDoublePerlinNoiseChannelContainer.h
@@ -18,6 +18,7 @@ public:
 
 	using ChannelType = FMovieSceneDoublePerlinNoiseChannel;
 
+	void InitializeOverride(FMovieSceneChannel* InChannel) override;
 	bool SupportsOverride(FName DefaultChannelTypeName) const override;
 	void ImportEntityImpl(
 			const UE::MovieScene::FChannelOverrideEntityImportParams& OverrideParams, 
diff --git a/Engine/Source/Runtime/MovieSceneTracks/Public/Channels/MovieSceneFloatPerlinNoiseChannel.h b/Engine/Source/Runtime/MovieSceneTracks/Public/Channels/MovieSceneFloatPerlinNoiseChannel.h
--- a/Engine/Source/Runtime/MovieSceneTracks/Public/Channels/MovieSceneFloatPerlinNoiseChannel.h
+++ b/Engine/Source/Runtime/MovieSceneTracks/Public/Channels/MovieSceneFloatPerlinNoiseChannel.h
@@ -15,29 +15,38 @@ struct MOVIESCENETRACKS_API FMovieSceneFloatPerlinNoiseChannel : public FMovieSc
 	using CurveValueType = float;
 
 	FMovieSceneFloatPerlinNoiseChannel();
-	FMovieSceneFloatPerlinNoiseChannel(const FPerlinNoiseParams& InFloatPerlinNoiseParams);
 
 	/**
-	* Evaluate this channel with the frame resolution
-	*
-	* @param InSeconds  The Frame second to evaluate at
-	* @return A value to receive the PerlinNoise result
-	*/
+	 * Evaluate this channel at the given time
+	 *
+	 * @param InSeconds		The time, in seconds, to evaluate at
+	 * @return				The evaluated noise value
+	 */
 	float Evaluate(double InSeconds) const;
 
 	/**
 	 * Evaluate this channel at the given time
 	 *
-	 * @param InSection  The section that contains this channel, used to lookup the sequence's tick resolution
-	 * @param InTime     The time, in ticks, to evaluate at
-	 * @param OutValue   The evaluated noise value
-	 * @return           Whether the noise was successfully evaluated
+	 * @param InSection		The section that contains this channel, used to lookup the sequence's tick resolution
+	 * @param InTime		The time, in ticks, to evaluate at
+	 * @param OutValue		The evaluated noise value
+	 * @return				Whether the noise was successfully evaluated
 	 */
 	bool Evaluate(const UMovieSceneSection* InSection, FFrameTime InTime, float& OutValue) const;
 
 	/** The noise parameters */
 	UPROPERTY(EditAnywhere, Category="Perlin Noise")
 	FPerlinNoiseParams PerlinNoiseParams;
+
+	/**
+	 * Evaluate perlin noise
+	 *
+	 * @params InParams			Perlin noise parameters
+	 * @params InInitialOffset	The initial offset for the noise
+	 * @params InSeconds		The time at which to evaluate the noise
+	 * @return					The evaluated noise value
+	 */
+	static float Evaluate(const FPerlinNoiseParams& InParams, double InSeconds);
 };
 
 template<>
diff --git a/Engine/Source/Runtime/MovieSceneTracks/Public/Channels/MovieSceneFloatPerlinNoiseChannelContainer.h b/Engine/Source/Runtime/MovieSceneTracks/Public/Channels/MovieSceneFloatPerlinNoiseChannelContainer.h
--- a/Engine/Source/Runtime/MovieSceneTracks/Public/Channels/MovieSceneFloatPerlinNoiseChannelContainer.h
+++ b/Engine/Source/Runtime/MovieSceneTracks/Public/Channels/MovieSceneFloatPerlinNoiseChannelContainer.h
@@ -19,6 +19,7 @@ public:
 
 	using ChannelType = FMovieSceneFloatPerlinNoiseChannel;
 
+	void InitializeOverride(FMovieSceneChannel* InChannel) override;
 	bool SupportsOverride(FName DefaultChannelTypeName) const override;
 	void ImportEntityImpl(
 			const UE::MovieScene::FChannelOverrideEntityImportParams& OverrideParams, 
diff --git a/Engine/Source/Runtime/MovieSceneTracks/Public/MovieSceneTracksComponentTypes.h b/Engine/Source/Runtime/MovieSceneTracks/Public/MovieSceneTracksComponentTypes.h
--- a/Engine/Source/Runtime/MovieSceneTracks/Public/MovieSceneTracksComponentTypes.h
+++ b/Engine/Source/Runtime/MovieSceneTracks/Public/MovieSceneTracksComponentTypes.h
@@ -35,8 +35,15 @@ struct MOVIESCENETRACKS_API FPerlinNoiseParams
 	UPROPERTY(EditAnywhere, Category = "Perlin Noise")
 	double Amplitude;
 
+	/** Starting offset, in seconds, into the noise pattern */
+	UPROPERTY(EditAnywhere, Category = "Perlin Noise")
+	float Offset;
+
 	FPerlinNoiseParams();
 	FPerlinNoiseParams(float InFrequency, double InAmplitude);
+
+	/** Generates a new random offset between [0, InMaxOffset] and sets it */
+	void RandomizeOffset(float InMaxOffset = 100.f);
 };
 
 /** Component data for the level visibility system */
