Constraints: To get constraints to be finally spawnable we have to
1) Save a copy of the constraint that may not exist in the world,that's the ConstraintCopyToSpawn item added to the constraint channel. Make sure it's saved in PreSave
2)Duplication in UE is wack so addded DuplicateMe functions to make sure sub objecst or whatever else we want get duplicated.
3) When we evaluate, via the CR template or the the transform system, we check to see if the weak constraint owned by the constraint manager is valid, if not we copy our copy back to the constraint manager.
4) When deleting Constraint from an Actor deletion, I added a bool to see if we should compenstate, this was causing issues with the spawnable system.
5) I added an extra SubObject parameter with ResolveBoundObjects to see if a control rig has gotten changed to make sure we handle it.

diff --git a/Engine/Plugins/Animation/ControlRig/Source/ControlRig/Private/Constraints/ControlRigTransformableHandle.cpp b/Engine/Plugins/Animation/ControlRig/Source/ControlRig/Private/Constraints/ControlRigTransformableHandle.cpp
--- a/Engine/Plugins/Animation/ControlRig/Source/ControlRig/Private/Constraints/ControlRigTransformableHandle.cpp
+++ b/Engine/Plugins/Animation/ControlRig/Source/ControlRig/Private/Constraints/ControlRigTransformableHandle.cpp
@@ -543,11 +543,42 @@ bool UTransformableControlHandle::AddTransformKeys(const TArray<FFrameNumber>& I
 	return true;
 }
 
-void UTransformableControlHandle::ResolveBoundObjects(FMovieSceneSequenceID LocalSequenceID, IMovieScenePlayer& Player)
+//for control rig need to check to see if the control rig is different then we may need to update it based upon what we are now bound to
+void UTransformableControlHandle::ResolveBoundObjects(FMovieSceneSequenceID LocalSequenceID, IMovieScenePlayer& Player, UObject* SubObject)
 {
-	//don't do anything for ControlRigs since they don't get removed when it's actor is destroyed or respawned since the track owns it. In the template
+	if (UControlRig* InControlRig = Cast<UControlRig>(SubObject))
+	{
+		if (ControlRig != InControlRig)
+		{
+			for (TWeakObjectPtr<> ParentObject : ConstraintBindingID.ResolveBoundObjects(LocalSequenceID, Player))
+			{
+				USceneComponent* Component = nullptr;
+				if (AActor* Actor = Cast<AActor>(ParentObject.Get()))
+				{
+					Component = Actor->GetRootComponent();
+				}
+				else if (USceneComponent* Comp = Cast<USceneComponent>(ParentObject.Get()))
+				{
+					Component = Comp;
+				}
+
+				if (InControlRig->GetObjectBinding() && InControlRig->GetObjectBinding()->GetBoundObject() == Component)
+				{
+					ControlRig = InControlRig;
+				}
+				break; //just do one
+			}
+		}
+	}
 }
 
+UTransformableHandle* UTransformableControlHandle::Duplicate(UObject* NewOuter) const
+{
+	UTransformableControlHandle* HandleCopy = DuplicateObject<UTransformableControlHandle>(this, NewOuter, GetFName());
+	HandleCopy->ControlRig = ControlRig;
+	HandleCopy->ControlName = ControlName;
+	return HandleCopy;
+}
 #if WITH_EDITOR
 
 FString UTransformableControlHandle::GetLabel() const
diff --git a/Engine/Plugins/Animation/ControlRig/Source/ControlRig/Private/Sequencer/MovieSceneControlRigParameterSection.cpp b/Engine/Plugins/Animation/ControlRig/Source/ControlRig/Private/Sequencer/MovieSceneControlRigParameterSection.cpp
--- a/Engine/Plugins/Animation/ControlRig/Source/ControlRig/Private/Sequencer/MovieSceneControlRigParameterSection.cpp
+++ b/Engine/Plugins/Animation/ControlRig/Source/ControlRig/Private/Sequencer/MovieSceneControlRigParameterSection.cpp
@@ -21,6 +21,7 @@
 #include "UObject/UE5MainStreamObjectVersion.h"
 #include "TransformConstraint.h"
 #include "TransformableHandle.h"
+#include "UObject/ObjectSaveContext.h"
 
 #include UE_INLINE_GENERATED_CPP_BY_NAME(MovieSceneControlRigParameterSection)
 
@@ -739,6 +740,17 @@ void UMovieSceneControlRigParameterSection::OnBindingIDsUpdated(const TMap<UE::M
 				TransformConstraint->ParentTRSHandle->OnBindingIDsUpdated(OldFixedToNewFixedMap, LocalSequenceID, Hierarchy, Player);
 			}
 		}
+		if (UTickableTransformConstraint* SpawnCopy = Cast< UTickableTransformConstraint>(ConstraintChannel.ConstraintCopyToSpawn))
+		{
+			if (SpawnCopy->ChildTRSHandle)
+			{
+				SpawnCopy->ChildTRSHandle->OnBindingIDsUpdated(OldFixedToNewFixedMap, LocalSequenceID, Hierarchy, Player);
+			}
+			if (SpawnCopy->ParentTRSHandle)
+			{
+				SpawnCopy->ParentTRSHandle->OnBindingIDsUpdated(OldFixedToNewFixedMap, LocalSequenceID, Hierarchy, Player);
+			}
+		}
 	}
 }
 
@@ -760,6 +772,20 @@ void UMovieSceneControlRigParameterSection::GetReferencedBindings(TArray<FGuid>&
 	}
 }
 
+void UMovieSceneControlRigParameterSection::PreSave(FObjectPreSaveContext SaveContext)
+{
+	Super::PreSave(SaveContext);
+	
+	for (FConstraintAndActiveChannel& ActiveChannel : ConstraintsChannels)
+	{
+		if (ActiveChannel.Constraint.IsValid())
+		{
+			ActiveChannel.ConstraintCopyToSpawn = ActiveChannel.Constraint->Duplicate(this);
+		}
+	}
+}
+
+
 bool UMovieSceneControlRigParameterSection::RenameParameterName(const FName& OldParameterName, const FName& NewParameterName)
 {
 	bool bWasReplaced = false;
@@ -936,6 +962,29 @@ void UMovieSceneControlRigParameterSection::PostEditImport()
 void UMovieSceneControlRigParameterSection::PostLoad()
 {
 	Super::PostLoad();
+	//for spawnables the control rig saved in our channels may have changed so we need to update thaem
+	if (ControlRig)
+	{
+		for (FConstraintAndActiveChannel& ConstraintChannel : ConstraintsChannels)
+		{
+			if (UTickableTransformConstraint* TransformConstraint = Cast<UTickableTransformConstraint>(ConstraintChannel.ConstraintCopyToSpawn))
+			{
+				if (UTransformableControlHandle* Handle = Cast<UTransformableControlHandle>(TransformConstraint->ChildTRSHandle))
+				{
+					Handle->ControlRig = ControlRig;
+				}
+			}
+			/*
+			if (ConstraintChannel.Constraint.IsValid() == false && ConstraintChannel.ConstraintCopyToSpawn)
+			{
+				const FConstraintsManagerController& Controller = FConstraintsManagerController::Get(ControlRig->GetWorld());
+				UTickableConstraint* NewOne = Controller.AddConstraintFromCopy(ConstraintChannel.ConstraintCopyToSpawn);
+				ConstraintChannel.Constraint = NewOne;
+				ReconstructChannelProxy();
+			}
+			*/
+		}
+	}
 }
 
 bool UMovieSceneControlRigParameterSection::HasScalarParameter(FName InParameterName) const
@@ -1345,9 +1394,24 @@ FConstraintAndActiveChannel* UMovieSceneControlRigParameterSection::GetConstrain
 	return (Index != INDEX_NONE) ? &ConstraintsChannels[Index] : nullptr;	
 }
 
+void UMovieSceneControlRigParameterSection::ReplaceConstraint(const FName InConstraintName, UTickableConstraint* InConstraint) 
+{
+	const int32 Index = ConstraintsChannels.IndexOfByPredicate([InConstraintName](const FConstraintAndActiveChannel& InChannel)
+	{
+		return InChannel.ConstraintCopyToSpawn ? InChannel.ConstraintCopyToSpawn->GetFName() == InConstraintName : false;
+	});
+	if (Index != INDEX_NONE)
+	{
+		Modify();
+		ConstraintsChannels[Index].Constraint = InConstraint;
+		ReconstructChannelProxy();
+	}
+}
+
 void UMovieSceneControlRigParameterSection::AddConstraintChannel(UTickableConstraint* InConstraint)
 {
-	if (!HasConstraintChannel(InConstraint->GetFName()))
+	
+	if (InConstraint && !HasConstraintChannel(InConstraint->GetFName()))
 	{
 		Modify();
 		
@@ -1355,6 +1419,9 @@ void UMovieSceneControlRigParameterSection::AddConstraintChannel(UTickableConstr
 
 		FMovieSceneConstraintChannel* ExistingChannel = &ConstraintsChannels[NewIndex].ActiveChannel;
 		ExistingChannel->SetDefault(false);
+		
+		//make copy that we can spawn if it doesn't exist
+		ConstraintsChannels[NewIndex].ConstraintCopyToSpawn = InConstraint->Duplicate(this);
 
 		if (OnConstraintChannelAdded.IsBound())
 		{
diff --git a/Engine/Plugins/Animation/ControlRig/Source/ControlRig/Private/Sequencer/MovieSceneControlRigParameterTemplate.cpp b/Engine/Plugins/Animation/ControlRig/Source/ControlRig/Private/Sequencer/MovieSceneControlRigParameterTemplate.cpp
--- a/Engine/Plugins/Animation/ControlRig/Source/ControlRig/Private/Sequencer/MovieSceneControlRigParameterTemplate.cpp
+++ b/Engine/Plugins/Animation/ControlRig/Source/ControlRig/Private/Sequencer/MovieSceneControlRigParameterTemplate.cpp
@@ -22,6 +22,8 @@
 #include "Rigs/FKControlRig.h"
 #include "UObject/UObjectAnnotation.h"
 #include "Rigs/RigHierarchy.h"
+#include "TransformConstraint.h"
+#include "TransformableHandle.h"
 
 #include UE_INLINE_GENERATED_CPP_BY_NAME(MovieSceneControlRigParameterTemplate)
 
@@ -174,11 +176,13 @@ struct FEulerTransformParameterStringAndValue
 
 struct FConstraintAndActiveValue
 {
-	FConstraintAndActiveValue(TWeakObjectPtr<UTickableConstraint> InConstraint, bool InValue)
+	FConstraintAndActiveValue(TWeakObjectPtr<UTickableConstraint> InConstraint, TWeakObjectPtr<UTickableConstraint> InSpawnConstraint, bool InValue)
 		: Constraint(InConstraint)
+		, SpawnConstraint(InSpawnConstraint)
 		, Value(InValue)
 	{}
 	TWeakObjectPtr<UTickableConstraint> Constraint;
+	TWeakObjectPtr<UTickableConstraint> SpawnConstraint;
 	bool Value;
 };
 
@@ -1259,13 +1263,37 @@ struct FControlRigParameterExecutionToken : IMovieSceneExecutionToken
 					}
 				}
 
-				for (const FConstraintAndActiveValue& ConstraintValue : ConstraintsValues)
+				const FConstraintsManagerController& Controller = FConstraintsManagerController::Get(ControlRig->GetWorld());
+				for (FConstraintAndActiveValue& ConstraintValue : ConstraintsValues)
 				{
+					//if the constraint isn't valid try to make a copy of our spawn copy if it exists and add that
+					if (ConstraintValue.Constraint.IsValid() == false && ConstraintValue.SpawnConstraint.IsValid())
+					{
+						UTickableConstraint* NewOne = Controller.AddConstraintFromCopy(ConstraintValue.SpawnConstraint.Get());
+						ConstraintValue.Constraint = NewOne;
+						UMovieSceneControlRigParameterSection* NonConstSection = const_cast<UMovieSceneControlRigParameterSection*>(Section);
+						NonConstSection->ReplaceConstraint(ConstraintValue.SpawnConstraint.Get()->GetFName(), NewOne);
+					}
 					if (ConstraintValue.Constraint.IsValid())
 					{
+						//For Control Rig we may need to explicitly set the control rig
+						ConstraintValue.Constraint->ResolveBoundObjects(Operand.SequenceID, Player, ControlRig);
 						ConstraintValue.Constraint->SetActive(ConstraintValue.Value);
 					}
 				}
+				//unfortunately for Constraints with ControlRig we need to resolve all Parents also. Don't need to do children since they wil be handled by
+				//the channel resolve above
+				TArray< TObjectPtr<UTickableConstraint>> Constraints =  Controller.GetAllConstraints();
+				for (UTickableConstraint* TickConstraint : Constraints)
+				{
+					if (UTickableTransformConstraint* TransformConstraint = Cast< UTickableTransformConstraint>(TickConstraint))
+					{
+						if (TransformConstraint->ParentTRSHandle)
+						{
+							TransformConstraint->ParentTRSHandle->ResolveBoundObjects(Operand.SequenceID, Player, ControlRig);
+						}
+					}
+				}
 			}
 			Section->SetDoNotKey(bWasDoNotKey);
 		}
@@ -1801,7 +1829,7 @@ void FMovieSceneControlRigParameterTemplate::EvaluateCurvesWithMasks(const FMovi
 			const FConstraintAndActiveChannel& ConstraintAndActiveChannel = Constraints[Index];
 			ConstraintAndActiveChannel.ActiveChannel.Evaluate(RoundTime, Value);
 			
-			Values.ConstraintsValues.Emplace(ConstraintAndActiveChannel.Constraint.Get(), Value);
+			Values.ConstraintsValues.Emplace(ConstraintAndActiveChannel.Constraint.Get(), ConstraintAndActiveChannel.ConstraintCopyToSpawn, Value);
 		}
 		
 		for (int32 Index = 0; Index < Bools.Num(); ++Index)
diff --git a/Engine/Plugins/Animation/ControlRig/Source/ControlRig/Public/Constraints/ControlRigTransformableHandle.h b/Engine/Plugins/Animation/ControlRig/Source/ControlRig/Public/Constraints/ControlRigTransformableHandle.h
--- a/Engine/Plugins/Animation/ControlRig/Source/ControlRig/Public/Constraints/ControlRigTransformableHandle.h
+++ b/Engine/Plugins/Animation/ControlRig/Source/ControlRig/Public/Constraints/ControlRigTransformableHandle.h
@@ -66,7 +66,10 @@ public:
 		const bool bLocal = true) const override;
 
 	/** Resolve the bound objects so that any object it references are resolved and correctly set up*/
-	virtual void ResolveBoundObjects(FMovieSceneSequenceID LocalSequenceID, IMovieScenePlayer& Player) override;
+	virtual void ResolveBoundObjects(FMovieSceneSequenceID LocalSequenceID, IMovieScenePlayer& Player, UObject* SubObject = nullptr) override;
+
+	/** Make a duplicate of myself with this outer*/
+	virtual UTransformableHandle* Duplicate(UObject* NewOuter) const override;
 
 	/** tick things when baking*/ 
 	virtual void TickForBaking() override;
diff --git a/Engine/Plugins/Animation/ControlRig/Source/ControlRig/Public/Sequencer/MovieSceneControlRigParameterSection.h b/Engine/Plugins/Animation/ControlRig/Source/ControlRig/Public/Sequencer/MovieSceneControlRigParameterSection.h
--- a/Engine/Plugins/Animation/ControlRig/Source/ControlRig/Public/Sequencer/MovieSceneControlRigParameterSection.h
+++ b/Engine/Plugins/Animation/ControlRig/Source/ControlRig/Public/Sequencer/MovieSceneControlRigParameterSection.h
@@ -287,6 +287,12 @@ public:
 	*/
 	virtual TArray<FConstraintAndActiveChannel>& GetConstraintsChannels()  override;
 
+	/*
+	*  Replace the constraint with the specified name with the new one
+	*/
+	virtual void ReplaceConstraint(const FName InName, UTickableConstraint* InConstraint)  override;
+
+
 	//not override but needed
 	const TArray<FConstraintAndActiveChannel>& GetConstraintsChannels() const;
 
@@ -491,6 +497,7 @@ protected:
 	virtual float GetTotalWeightValue(FFrameTime InTime) const override;
 	virtual void OnBindingIDsUpdated(const TMap<UE::MovieScene::FFixedObjectBindingID, UE::MovieScene::FFixedObjectBindingID>& OldFixedToNewFixedMap, FMovieSceneSequenceID LocalSequenceID, const FMovieSceneSequenceHierarchy* Hierarchy, IMovieScenePlayer& Player) override;
 	virtual void GetReferencedBindings(TArray<FGuid>& OutBindings) override;
+	virtual void PreSave(FObjectPreSaveContext SaveContext) override;
 
 
 	// When true we do not set a key on the section, since it will be set because we changed the value
diff --git a/Engine/Plugins/Animation/ControlRig/Source/ControlRigEditor/Private/Sequencer/ControlRigParameterTrackEditor.cpp b/Engine/Plugins/Animation/ControlRig/Source/ControlRigEditor/Private/Sequencer/ControlRigParameterTrackEditor.cpp
--- a/Engine/Plugins/Animation/ControlRig/Source/ControlRigEditor/Private/Sequencer/ControlRigParameterTrackEditor.cpp
+++ b/Engine/Plugins/Animation/ControlRig/Source/ControlRigEditor/Private/Sequencer/ControlRigParameterTrackEditor.cpp
@@ -2730,7 +2730,7 @@ void FControlRigParameterTrackEditor::HandleConstraintRemoved(IMovieSceneConstra
 		if (!InSection->OnConstraintRemovedHandle.IsValid())
 		{
 			InSection->OnConstraintRemovedHandle =
-			Controller.OnConstraintRemoved().AddLambda([InSection, Section, this](FName InConstraintName)
+			Controller.OnConstraintRemoved().AddLambda([InSection, Section, this](FName InConstraintName, bool bDoNotCompensate)
 			{
 				const FConstraintAndActiveChannel* ConstraintChannel = InSection->GetConstraintChannel(InConstraintName);
 				if (!ConstraintChannel)
@@ -2738,7 +2738,7 @@ void FControlRigParameterTrackEditor::HandleConstraintRemoved(IMovieSceneConstra
 					return;
 				}
 				
-				if (ConstraintChannel->Constraint.IsValid())
+				if (bDoNotCompensate == false && ConstraintChannel->Constraint.IsValid())
 				{
 					FMovieSceneConstraintChannelHelper::HandleConstraintRemoved(
 						ConstraintChannel->Constraint.Get(),
diff --git a/Engine/Source/Editor/MovieSceneTools/Private/TrackEditors/TransformTrackEditor.cpp b/Engine/Source/Editor/MovieSceneTools/Private/TrackEditors/TransformTrackEditor.cpp
--- a/Engine/Source/Editor/MovieSceneTools/Private/TrackEditors/TransformTrackEditor.cpp
+++ b/Engine/Source/Editor/MovieSceneTools/Private/TrackEditors/TransformTrackEditor.cpp
@@ -1587,7 +1587,7 @@ void F3DTransformTrackEditor::HandleConstraintRemoved(IMovieSceneConstrainedSect
 	if (!InSection->OnConstraintRemovedHandle.IsValid())
 	{
 		InSection->OnConstraintRemovedHandle =
-			Controller.OnConstraintRemoved().AddLambda([InSection,Section, this](FName InConstraintName)
+			Controller.OnConstraintRemoved().AddLambda([InSection,Section, this](FName InConstraintName, bool bDoNotCompensate)
 				{
 					const FConstraintAndActiveChannel* ConstraintChannel = InSection->GetConstraintChannel(InConstraintName);
 					if (!ConstraintChannel)
@@ -1595,7 +1595,7 @@ void F3DTransformTrackEditor::HandleConstraintRemoved(IMovieSceneConstrainedSect
 						return;
 					}
 
-					if (ConstraintChannel->Constraint.IsValid())
+					if (bDoNotCompensate == false && ConstraintChannel->Constraint.IsValid())
 					{
 						FMovieSceneConstraintChannelHelper::HandleConstraintRemoved(
 							ConstraintChannel->Constraint.Get(),
diff --git a/Engine/Source/Runtime/Experimental/Animation/Constraints/Private/ConstraintsManager.cpp b/Engine/Source/Runtime/Experimental/Animation/Constraints/Private/ConstraintsManager.cpp
--- a/Engine/Source/Runtime/Experimental/Animation/Constraints/Private/ConstraintsManager.cpp
+++ b/Engine/Source/Runtime/Experimental/Animation/Constraints/Private/ConstraintsManager.cpp
@@ -82,6 +82,11 @@ void UTickableConstraint::Evaluate(bool bTickHandlesAlso) const
 	ConstraintTick.EvaluateFunctions();
 }
 
+UTickableConstraint* UTickableConstraint::Duplicate(UObject* NewOuter) const
+{
+	return DuplicateObject<UTickableConstraint>(this, NewOuter, GetFName());
+}
+
 #if WITH_EDITOR
 
 FString UTickableConstraint::GetLabel() const
@@ -154,11 +159,13 @@ void UConstraintsManager::OnActorDestroyed(AActor* InActor)
 {
 	if (USceneComponent* SceneComponent = InActor->GetRootComponent())
 	{
+
 		TArray< int32 > IndicesToRemove;
 		for (int32 Index = 0; Index < Constraints.Num(); ++Index)
 		{
 			//if the constraint has a bound object(in sequencer) we don't remove the constraint, it could be a spawnable
-			if (IsValid(Constraints[Index]) &&  Constraints[Index]->HasBoundObjects() == false && Constraints[Index]->ReferencesObject(SceneComponent))
+			if (IsValid(Constraints[Index]) && Constraints[Index]->HasBoundObjects() == false && Constraints[Index]->ReferencesObject(SceneComponent))
+
 			{
 				IndicesToRemove.Add(Index);
 			}
@@ -169,7 +176,7 @@ void UConstraintsManager::OnActorDestroyed(AActor* InActor)
 			FConstraintsManagerController& Controller = FConstraintsManagerController::Get(InActor->GetWorld());
 			for (int32 Index = IndicesToRemove.Num()-1; Index >= 0; --Index)
 			{
-				Controller.RemoveConstraint(Index);
+				Controller.RemoveConstraint(Index,/*do not compensate*/ true);
 			}
 		}
 	}
@@ -282,7 +289,7 @@ void UConstraintsManager::Clear(UWorld* World)
 	UnregisterDelegates(World);
 	for (UTickableConstraint* Constraint : Constraints)
 	{
-		OnConstraintRemoved_BP.Broadcast(this, Constraint);
+		OnConstraintRemoved_BP.Broadcast(this, Constraint, true);
 	}
 	Constraints.Empty();
 }
@@ -405,6 +412,27 @@ bool FConstraintsManagerController::AddConstraint(UTickableConstraint* InConstra
 	return true;
 }
 
+UTickableConstraint* FConstraintsManagerController::AddConstraintFromCopy(UTickableConstraint* CopyOfConstraint) const
+{
+	if (!CopyOfConstraint)
+	{
+		return nullptr;
+	}
+	UConstraintsManager* Manager = GetManager(); //this will allocate if doesn't exist
+	if (!Manager)
+	{
+		return nullptr;
+	}
+
+	UTickableConstraint* OurCopy = CopyOfConstraint->Duplicate(Manager);
+	if (AddConstraint(OurCopy))
+	{
+		return OurCopy;
+	}
+	return nullptr;
+}
+
+
 int32 FConstraintsManagerController::GetConstraintIndex(const FName& InConstraintName) const
 {
 	const UConstraintsManager* Manager = FindManager();
@@ -419,7 +447,7 @@ int32 FConstraintsManagerController::GetConstraintIndex(const FName& InConstrain
 	} );
 }
 	
-bool  FConstraintsManagerController::RemoveConstraint(const FName& InConstraintName) const
+bool  FConstraintsManagerController::RemoveConstraint(const FName& InConstraintName, bool bDoNotCompensate) const
 {
 	const int32 Index = GetConstraintIndex(InConstraintName);
 	if (Index == INDEX_NONE)
@@ -427,10 +455,10 @@ bool  FConstraintsManagerController::RemoveConstraint(const FName& InConstraintN
 		return false;
 	}
 	
-	return RemoveConstraint(Index);
+	return RemoveConstraint(Index, bDoNotCompensate);
 }
 
-bool FConstraintsManagerController::RemoveConstraint(const int32 InConstraintIndex) const
+bool FConstraintsManagerController::RemoveConstraint(const int32 InConstraintIndex,bool bDoNotCompensate) const
 {
 	UConstraintsManager* Manager = FindManager();
 	if (!Manager)
@@ -448,8 +476,8 @@ bool FConstraintsManagerController::RemoveConstraint(const int32 InConstraintInd
 	
 	// notify deletion
 
-	ConstraintRemoved.Broadcast(ConstraintName);
-	Manager->OnConstraintRemoved_BP.Broadcast(Manager, Constraint);
+	ConstraintRemoved.Broadcast(ConstraintName, bDoNotCompensate);
+	Manager->OnConstraintRemoved_BP.Broadcast(Manager, Constraint, bDoNotCompensate);
 	Manager->Constraints[InConstraintIndex]->Modify();
 	Manager->Modify();
 
diff --git a/Engine/Source/Runtime/Experimental/Animation/Constraints/Private/TransformConstraint.cpp b/Engine/Source/Runtime/Experimental/Animation/Constraints/Private/TransformConstraint.cpp
--- a/Engine/Source/Runtime/Experimental/Animation/Constraints/Private/TransformConstraint.cpp
+++ b/Engine/Source/Runtime/Experimental/Animation/Constraints/Private/TransformConstraint.cpp
@@ -36,6 +36,24 @@ EMovieSceneTransformChannel UTickableTransformConstraint::GetChannelsToKey() con
 	return EMovieSceneTransformChannel::AllTransform;
 }
 
+UTickableConstraint* UTickableTransformConstraint::Duplicate(UObject* NewOuter) const
+{
+	UTickableTransformConstraint* Dup = DuplicateObject<UTickableTransformConstraint>(this, NewOuter);
+	if (ChildTRSHandle)
+	{
+		UTransformableHandle* HandleCopy = ChildTRSHandle->Duplicate(NewOuter);
+		Dup->ChildTRSHandle = HandleCopy;
+	}
+	if (ParentTRSHandle)
+	{
+		UTransformableHandle* HandleCopy = ParentTRSHandle->Duplicate(NewOuter);
+		Dup->ParentTRSHandle = HandleCopy;
+	}
+	Dup->SetupDependencies();
+	Dup->RegisterDelegates();
+	return Dup;
+}
+
 #if WITH_EDITOR
 
 FString UTickableTransformConstraint::GetLabel() const
@@ -241,15 +259,15 @@ bool UTickableTransformConstraint::HasBoundObjects() const
 	return false;
 }
 
-void UTickableTransformConstraint::ResolveBoundObjects(FMovieSceneSequenceID LocalSequenceID, IMovieScenePlayer& Player)
+void UTickableTransformConstraint::ResolveBoundObjects(FMovieSceneSequenceID LocalSequenceID, IMovieScenePlayer& Player, UObject* SubObject)
 {
 	if (ChildTRSHandle && ChildTRSHandle->HasBoundObjects())
 	{
-		ChildTRSHandle->ResolveBoundObjects(LocalSequenceID, Player);
+		ChildTRSHandle->ResolveBoundObjects(LocalSequenceID, Player, SubObject);
 	}
 	if (ParentTRSHandle && ParentTRSHandle->HasBoundObjects())
 	{
-		ParentTRSHandle->ResolveBoundObjects(LocalSequenceID, Player);
+		ParentTRSHandle->ResolveBoundObjects(LocalSequenceID, Player, SubObject);
 	}
 }
 
diff --git a/Engine/Source/Runtime/Experimental/Animation/Constraints/Private/TransformableHandle.cpp b/Engine/Source/Runtime/Experimental/Animation/Constraints/Private/TransformableHandle.cpp
--- a/Engine/Source/Runtime/Experimental/Animation/Constraints/Private/TransformableHandle.cpp
+++ b/Engine/Source/Runtime/Experimental/Animation/Constraints/Private/TransformableHandle.cpp
@@ -492,7 +492,7 @@ bool UTransformableComponentHandle::AddTransformKeys(const TArray<FFrameNumber>&
 	return true;
 }
 
-void UTransformableComponentHandle::ResolveBoundObjects(FMovieSceneSequenceID LocalSequenceID, IMovieScenePlayer& Player)
+void UTransformableComponentHandle::ResolveBoundObjects(FMovieSceneSequenceID LocalSequenceID, IMovieScenePlayer& Player, UObject* SubObject)
 {
 	for (TWeakObjectPtr<> ParentObject : ConstraintBindingID.ResolveBoundObjects(LocalSequenceID, Player))
 	{
@@ -508,6 +508,14 @@ void UTransformableComponentHandle::ResolveBoundObjects(FMovieSceneSequenceID Lo
 	}
 }
 
+UTransformableHandle* UTransformableComponentHandle::Duplicate(UObject* NewOuter) const
+{
+	UTransformableComponentHandle* HandleCopy = DuplicateObject<UTransformableComponentHandle>(this, NewOuter, GetFName());
+	HandleCopy->Component = Component;
+	HandleCopy->SocketName = SocketName;
+	return HandleCopy;
+}
+
 #if WITH_EDITOR
 FString UTransformableComponentHandle::GetLabel() const
 {
diff --git a/Engine/Source/Runtime/Experimental/Animation/Constraints/Public/ConstraintChannel.h b/Engine/Source/Runtime/Experimental/Animation/Constraints/Public/ConstraintChannel.h
--- a/Engine/Source/Runtime/Experimental/Animation/Constraints/Public/ConstraintChannel.h
+++ b/Engine/Source/Runtime/Experimental/Animation/Constraints/Public/ConstraintChannel.h
@@ -107,4 +107,7 @@ struct CONSTRAINTS_API FConstraintAndActiveChannel
 
 	UPROPERTY()
 	FMovieSceneConstraintChannel ActiveChannel;
+
+	UPROPERTY()
+	TObjectPtr<UTickableConstraint> ConstraintCopyToSpawn;
 };
\ No newline at end of file
diff --git a/Engine/Source/Runtime/Experimental/Animation/Constraints/Public/ConstraintsManager.h b/Engine/Source/Runtime/Experimental/Animation/Constraints/Public/ConstraintsManager.h
--- a/Engine/Source/Runtime/Experimental/Animation/Constraints/Public/ConstraintsManager.h
+++ b/Engine/Source/Runtime/Experimental/Animation/Constraints/Public/ConstraintsManager.h
@@ -93,12 +93,14 @@ public:
 	/** If true it contains objects bound to an external system, like sequencer so we don't do certain things, like remove constraints when they don't resolve*/
 	virtual bool HasBoundObjects() const PURE_VIRTUAL(HasBoundObjects, return false;);
 	/** Resolve the bound objects so that any object it references are resovled and correctly set up*/
-	virtual void ResolveBoundObjects(FMovieSceneSequenceID LocalSequenceID, IMovieScenePlayer& Player)  PURE_VIRTUAL(ResolveBoundObjects);
+	virtual void ResolveBoundObjects(FMovieSceneSequenceID LocalSequenceID, IMovieScenePlayer& Player, UObject* SubObject = nullptr)  PURE_VIRTUAL(ResolveBoundObjects);
 
 	/** @todo document */
 	virtual uint32 GetTargetHash() const PURE_VIRTUAL(GetTargetHash, return 0;);
 	/** Whether or not this object references this object. If it HasBoundObjects you should ResolveBoundObjects if you expect this to be up to date */
 	virtual bool ReferencesObject(TWeakObjectPtr<UObject> InObject) const PURE_VIRTUAL(ReferencesObject, return false;);
+	/** Create duplicate with new Outer*/
+	virtual UTickableConstraint* Duplicate(UObject* NewOuter) const;
 
 #if WITH_EDITOR
 	/** Returns the constraint's label used for UI. */
@@ -134,7 +136,7 @@ public:
 
 	/** Dynamic blueprintable delegates for knowing when a constraints are added or deleted*/
 	DECLARE_DYNAMIC_MULTICAST_SPARSE_DELEGATE_TwoParams(FOnConstraintAdded, UConstraintsManager, OnConstraintAdded_BP, UConstraintsManager*, Mananger, UTickableConstraint*, Constraint);
-	DECLARE_DYNAMIC_MULTICAST_SPARSE_DELEGATE_TwoParams(FOnConstraintRemoved, UConstraintsManager, OnConstraintRemoved_BP, UConstraintsManager*, Mananger, UTickableConstraint*, Constraint);
+	DECLARE_DYNAMIC_MULTICAST_SPARSE_DELEGATE_ThreeParams(FOnConstraintRemoved, UConstraintsManager, OnConstraintRemoved_BP, UConstraintsManager*, Mananger, UTickableConstraint*, Constraint, bool, bDoNotCompensate);
 
 	
 	UConstraintsManager();
@@ -166,7 +168,7 @@ public:
 	
 	/** BP Delegate fired when constraints are removed*/
 	UPROPERTY(BlueprintAssignable, Category = Constraints, meta = (DisplayName = "OnConstraintRemoved"))
-	FOnConstraintAdded OnConstraintRemoved_BP;
+	FOnConstraintRemoved OnConstraintRemoved_BP;
 
 private:
 
@@ -181,10 +183,11 @@ private:
 	/** @todo document */
 	void Dump() const;
 
-	/** @todo document */
+	/** All of the constraints*/
 	UPROPERTY()
 	TArray< TObjectPtr<UTickableConstraint> > Constraints;
 	
+
 	friend class FConstraintsManagerController;
 	friend class AConstraintsActor;
 };
@@ -217,17 +220,20 @@ public:
 		return NewConstraint;
 	}
 
-	/** @todo document */
+	/** Add this constraint to the manager */
 	bool AddConstraint(UTickableConstraint* InConstraint) const;
 	
+	/** Make a copy of this constraint and add it to the manager, returns the copy mode if it was added*/
+	UTickableConstraint* AddConstraintFromCopy(UTickableConstraint* CopyOfConstraint) const;
+
 	/** Get the index of the given constraint's name. */
 	int32 GetConstraintIndex(const FName& InConstraintName) const;
 	
 	/** Remove the constraint by name. */
-	bool RemoveConstraint(const FName& InConstraintName) const;
+	bool RemoveConstraint(const FName& InConstraintName, bool bDoNotCompensate = false) const;
 
 	/** Remove the constraint at the given index. */
-	bool RemoveConstraint(const int32 InConstraintIndex) const;
+	bool RemoveConstraint(const int32 InConstraintIndex, bool bDoNotCompensate = false) const;
 
 	/** @todo document */
 	UTickableConstraint* GetConstraint(const FName& InConstraintName) const;
@@ -253,6 +259,7 @@ public:
 
 	/** Go through each constraint in order and evaluate and tick them*/
 	void EvaluateAllConstraints() const;
+
 	
 private:
 	/** Delegeate that's fired when a scene component is constrained, this is needed to make sure things like gizmo's get updated after the constraint tick happens*/
@@ -260,7 +267,7 @@ private:
 	FOnSceneComponentConstrained SceneComponentConstrained;
 
 	/** @todo document */
-	DECLARE_MULTICAST_DELEGATE_OneParam(FOnConstraintRemoved, FName /*InConstraintName*/);
+	DECLARE_MULTICAST_DELEGATE_TwoParams(FOnConstraintRemoved, FName /*InConstraintName*/, bool /*bDoNotCompensate*/);
 	FOnConstraintRemoved ConstraintRemoved;
 	
 	/** Find the existing Constraints Manager in World or create a new one. */
diff --git a/Engine/Source/Runtime/Experimental/Animation/Constraints/Public/TransformConstraint.h b/Engine/Source/Runtime/Experimental/Animation/Constraints/Public/TransformConstraint.h
--- a/Engine/Source/Runtime/Experimental/Animation/Constraints/Public/TransformConstraint.h
+++ b/Engine/Source/Runtime/Experimental/Animation/Constraints/Public/TransformConstraint.h
@@ -44,7 +44,7 @@ public:
 	virtual bool HasBoundObjects() const override;
 	
 	/** Resolve the bound objects so that any object it references are resovled and correctly set up*/
-	virtual void ResolveBoundObjects(FMovieSceneSequenceID LocalSequenceID, IMovieScenePlayer& Player) override;
+	virtual void ResolveBoundObjects(FMovieSceneSequenceID LocalSequenceID, IMovieScenePlayer& Player,UObject* SubObject) override;
 
 	/** If Active and the handles and targets are valid*/
 	virtual bool IsFullyActive() const override;
@@ -52,6 +52,9 @@ public:
 	/** If that constraint needs to be handled by the compensation system. */
 	virtual bool NeedsCompensation() const;
 
+	/** Create duplicate with new Outer*/
+	virtual UTickableConstraint* Duplicate(UObject* NewOuter) const override;
+
 	/** Override the evaluate so we can tick our handles*/
 	virtual void Evaluate(bool bTickHandlesAlso = false) const override;
 
diff --git a/Engine/Source/Runtime/Experimental/Animation/Constraints/Public/TransformableHandle.h b/Engine/Source/Runtime/Experimental/Animation/Constraints/Public/TransformableHandle.h
--- a/Engine/Source/Runtime/Experimental/Animation/Constraints/Public/TransformableHandle.h
+++ b/Engine/Source/Runtime/Experimental/Animation/Constraints/Public/TransformableHandle.h
@@ -64,7 +64,11 @@ public:
 	virtual bool HasBoundObjects() const;
 
 	/** Resolve the bound objects so that any object it references are resovled and correctly set up*/
-	virtual void ResolveBoundObjects(FMovieSceneSequenceID LocalSequenceID, IMovieScenePlayer& Player) PURE_VIRTUAL(ResolveBoundObjects);
+	virtual void ResolveBoundObjects(FMovieSceneSequenceID LocalSequenceID, IMovieScenePlayer& Player, UObject* SubObject = nullptr) PURE_VIRTUAL(ResolveBoundObjects);
+
+	/** Make a duplicate of myself with this outer*/
+	virtual UTransformableHandle* Duplicate(UObject* NewOuter) const PURE_VIRTUAL(Duplicate, return nullptr;);
+
 
 	/** Fix up Binding in case it has changed*/
 	void OnBindingIDsUpdated(const TMap<UE::MovieScene::FFixedObjectBindingID, UE::MovieScene::FFixedObjectBindingID>& OldFixedToNewFixedMap, FMovieSceneSequenceID LocalSequenceID, const FMovieSceneSequenceHierarchy* Hierarchy, IMovieScenePlayer& Player);
@@ -178,7 +182,10 @@ public:
 		const bool bLocal = true) const override;
 
 	/** Resolve the bound objects so that any object it references are resovled and correctly set up*/
-	virtual void ResolveBoundObjects(FMovieSceneSequenceID LocalSequenceID, IMovieScenePlayer& Player) override;
+	virtual void ResolveBoundObjects(FMovieSceneSequenceID LocalSequenceID, IMovieScenePlayer& Player, UObject* SubObject = nullptr) override;
+
+	/** Make a duplicate of myself with this outer*/
+	virtual UTransformableHandle* Duplicate(UObject* NewOuter) const override;
 
 #if WITH_EDITOR
 	/** Returns labels used for UI. */
diff --git a/Engine/Source/Runtime/MovieSceneTracks/Private/Sections/MovieScene3DTransformSection.cpp b/Engine/Source/Runtime/MovieSceneTracks/Private/Sections/MovieScene3DTransformSection.cpp
--- a/Engine/Source/Runtime/MovieSceneTracks/Private/Sections/MovieScene3DTransformSection.cpp
+++ b/Engine/Source/Runtime/MovieSceneTracks/Private/Sections/MovieScene3DTransformSection.cpp
@@ -18,6 +18,8 @@
 #include "Algo/AnyOf.h"
 #include "TransformConstraint.h"
 #include "TransformableHandle.h"
+#include "UObject/ObjectSaveContext.h"
+
 
 #include UE_INLINE_GENERATED_CPP_BY_NAME(MovieScene3DTransformSection)
 
@@ -367,6 +369,17 @@ void UMovieScene3DTransformSection::OnBindingIDsUpdated(const TMap<UE::MovieScen
 					TransformConstraint->ParentTRSHandle->OnBindingIDsUpdated(OldFixedToNewFixedMap, LocalSequenceID, Hierarchy, Player);
 				}
 			}
+			if (UTickableTransformConstraint* SpawnCopy = Cast< UTickableTransformConstraint>(ConstraintChannel.ConstraintCopyToSpawn))
+			{
+				if (SpawnCopy->ChildTRSHandle)
+				{
+					SpawnCopy->ChildTRSHandle->OnBindingIDsUpdated(OldFixedToNewFixedMap, LocalSequenceID, Hierarchy, Player);
+				}
+				if (SpawnCopy->ParentTRSHandle)
+				{
+					SpawnCopy->ParentTRSHandle->OnBindingIDsUpdated(OldFixedToNewFixedMap, LocalSequenceID, Hierarchy, Player);
+				}
+			}
 		}
 	}
 }
@@ -392,6 +405,21 @@ void UMovieScene3DTransformSection::GetReferencedBindings(TArray<FGuid>& OutBind
 	}
 }
 
+void UMovieScene3DTransformSection::PreSave(FObjectPreSaveContext SaveContext)
+{
+	Super::PreSave(SaveContext);
+	if (Constraints)
+	{
+		for (FConstraintAndActiveChannel& ActiveChannel : Constraints->ConstraintsChannels)
+		{
+			if (ActiveChannel.Constraint.IsValid())
+			{
+				ActiveChannel.ConstraintCopyToSpawn = ActiveChannel.Constraint->Duplicate(this);
+			}
+		}
+	}
+}
+
 template<typename BaseBuilderType>
 void UMovieScene3DTransformSection::BuildEntity(BaseBuilderType& InBaseBuilder, UMovieSceneEntitySystemLinker* Linker, const FEntityImportParams& Params, FImportedEntity* OutImportedEntity)
 {
@@ -557,19 +585,21 @@ void UMovieScene3DTransformSection::ImportConstraintEntity(UMovieSceneEntitySyst
 		const int32 ConstraintIndex = static_cast<int32>(Params.EntityID & ~ConstraintTypeMask);
 
 		checkf(Constraints->ConstraintsChannels.IsValidIndex(ConstraintIndex), TEXT("Encoded constraint (%d) index is not valid within array size %d. Data must have been manipulated without re-compilaition."), ConstraintIndex, Constraints->ConstraintsChannels.Num());
-
-		if (Constraints->ConstraintsChannels[ConstraintIndex].Constraint.IsValid())
+		//add if constraint or spawn copy is valid
+		if (Constraints->ConstraintsChannels[ConstraintIndex].Constraint.IsValid() ||
+			Constraints->ConstraintsChannels[ConstraintIndex].ConstraintCopyToSpawn)
 		{
-			FName ConstraintName = Constraints->ConstraintsChannels[ConstraintIndex].Constraint->GetFName();
+			FName ConstraintName = Constraints->ConstraintsChannels[ConstraintIndex].Constraint.IsValid() ? Constraints->ConstraintsChannels[ConstraintIndex].Constraint->GetFName()
+				: Constraints->ConstraintsChannels[ConstraintIndex].ConstraintCopyToSpawn->GetFName();
 
 			FConstraintComponentData ComponentData;
 			ComponentData.ConstraintName = ConstraintName;
-			ComponentData.Channel = &(Constraints->ConstraintsChannels[ConstraintIndex].ActiveChannel);
+			ComponentData.ConstraintAndActiveChannel = &(Constraints->ConstraintsChannels[ConstraintIndex]);
+			ComponentData.Section = this;
 			OutImportedEntity->AddBuilder(
 				FEntityBuilder()
 				.Add(BuiltInComponentTypes->SceneComponentBinding, ObjectBindingID)
-				//.Add(TrackComponents->ConstraintName, ConstraintName)
-				.Add(TrackComponents->ConstraintChannel,ComponentData)
+				.Add(TrackComponents->ConstraintChannel, ComponentData)
 			);
 		}
 	}
@@ -952,7 +982,7 @@ void UMovieScene3DTransformSection::AddConstraintChannel(UTickableConstraint* In
 	{
 		Constraints = NewObject<UMovieScene3dTransformSectionConstraints>(this, NAME_None, RF_Public| RF_Transactional);
 	}
-	if (!HasConstraintChannel(InConstraint->GetFName()))
+	if (InConstraint && !HasConstraintChannel(InConstraint->GetFName()))
 	{
 		Constraints->SetFlags(RF_Transactional);
 		Constraints->Modify();
@@ -961,6 +991,9 @@ void UMovieScene3DTransformSection::AddConstraintChannel(UTickableConstraint* In
 		FMovieSceneConstraintChannel* ExistingChannel = &Constraints->ConstraintsChannels[NewIndex].ActiveChannel;
 		ExistingChannel->SetDefault(false);
 
+		//make copy that we can spawn if it doesn't exist
+		Constraints->ConstraintsChannels[NewIndex].ConstraintCopyToSpawn = InConstraint->Duplicate(this);
+
 		CacheChannelProxy();
 
 		if (OnConstraintChannelAdded.IsBound())
@@ -970,6 +1003,23 @@ void UMovieScene3DTransformSection::AddConstraintChannel(UTickableConstraint* In
 	}
 }
 
+void UMovieScene3DTransformSection::ReplaceConstraint(const FName InConstraintName, UTickableConstraint* InConstraint) 
+{
+	if (Constraints)
+	{
+		const int32 Index = Constraints->ConstraintsChannels.IndexOfByPredicate([InConstraintName](const FConstraintAndActiveChannel& InChannel)
+			{
+				return InChannel.ConstraintCopyToSpawn ? InChannel.ConstraintCopyToSpawn->GetFName() == InConstraintName : false;
+			});
+		if (Index != INDEX_NONE)
+		{
+			Modify();
+			Constraints->ConstraintsChannels[Index].Constraint = InConstraint;
+			CacheChannelProxy();
+		}
+	}
+}
+
 void UMovieScene3DTransformSection::RemoveConstraintChannel(const FName& InConstraintName)
 {
 	if (Constraints)
diff --git a/Engine/Source/Runtime/MovieSceneTracks/Private/Systems/MovieSceneConstraintSystem.cpp b/Engine/Source/Runtime/MovieSceneTracks/Private/Systems/MovieSceneConstraintSystem.cpp
--- a/Engine/Source/Runtime/MovieSceneTracks/Private/Systems/MovieSceneConstraintSystem.cpp
+++ b/Engine/Source/Runtime/MovieSceneTracks/Private/Systems/MovieSceneConstraintSystem.cpp
@@ -94,7 +94,6 @@ void UMovieSceneConstraintSystem::OnRun(FSystemTaskPrerequisites& InPrerequisite
 	FMovieSceneEntitySystemRunner* ActiveRunner = Linker->GetActiveRunner();
 	ESystemPhase CurrentPhase = ActiveRunner->GetCurrentPhase();
 
-
 	if (CurrentPhase == ESystemPhase::Instantiation)
 	{
 		// Save pre-animated state
@@ -117,11 +116,17 @@ void UMovieSceneConstraintSystem::OnRun(FSystemTaskPrerequisites& InPrerequisite
 			{
 				const FSequenceInstance& TargetInstance = InstanceRegistry->GetInstance(InstanceHandle);
 
-				UTickableConstraint* Constraint = Controller->GetConstraint(ConstraintChannel.ConstraintName);
+				UTickableConstraint* Constraint = ConstraintChannel.ConstraintAndActiveChannel->Constraint.Get();
+				if (!Constraint)
+				{
+					UTickableConstraint* NewOne = Controller->AddConstraintFromCopy(ConstraintChannel.ConstraintAndActiveChannel->ConstraintCopyToSpawn);
+					Constraint = Controller->GetConstraint(ConstraintChannel.ConstraintName);
+					ConstraintChannel.Section->ReplaceConstraint(ConstraintChannel.ConstraintName, Constraint);
+				}
 				if (Constraint)
 				{
 					bool Result = false;
-					ConstraintChannel.Channel->Evaluate(FrameTime, Result);
+					ConstraintChannel.ConstraintAndActiveChannel->ActiveChannel.Evaluate(FrameTime, Result);
 					Constraint->SetActive(Result);
 					Constraint->ResolveBoundObjects(TargetInstance.GetSequenceID(), *TargetInstance.GetPlayer());
 					if (UTickableTransformConstraint* TransformConstraint = Cast< UTickableTransformConstraint>(Constraint))
diff --git a/Engine/Source/Runtime/MovieSceneTracks/Public/MovieSceneTracksComponentTypes.h b/Engine/Source/Runtime/MovieSceneTracks/Public/MovieSceneTracksComponentTypes.h
--- a/Engine/Source/Runtime/MovieSceneTracks/Public/MovieSceneTracksComponentTypes.h
+++ b/Engine/Source/Runtime/MovieSceneTracks/Public/MovieSceneTracksComponentTypes.h
@@ -13,11 +13,13 @@
 #include "TransformData.h"
 #include "MovieSceneTracksPropertyTypes.h"
 #include "Styling/SlateColor.h"
+#include "ConstraintChannel.h"
 #include "MovieSceneTracksComponentTypes.generated.h"
 
 class UMaterialParameterCollection;
 class UMovieSceneDataLayerSection;
 class UMovieSceneLevelVisibilitySection;
+class UMovieScene3DTransformSection;
 struct FMovieSceneObjectBindingID;
 
 
@@ -74,8 +76,8 @@ struct FConstraintComponentData
 
 	UPROPERTY()
 	FName ConstraintName;
-
-	FMovieSceneBoolChannel* Channel;
+	FConstraintAndActiveChannel* ConstraintAndActiveChannel;
+	UMovieScene3DTransformSection* Section;
 };
 
 namespace UE
diff --git a/Engine/Source/Runtime/MovieSceneTracks/Public/Sections/MovieScene3DTransformSection.h b/Engine/Source/Runtime/MovieSceneTracks/Public/Sections/MovieScene3DTransformSection.h
--- a/Engine/Source/Runtime/MovieSceneTracks/Public/Sections/MovieScene3DTransformSection.h
+++ b/Engine/Source/Runtime/MovieSceneTracks/Public/Sections/MovieScene3DTransformSection.h
@@ -221,6 +221,7 @@ public:
 #if WITH_EDITOR
 	/* From UObject*/
 	virtual bool Modify(bool bAlwaysMarkDirty = true) override;
+
 #endif
 	/* From UMovieSection*/
 	
@@ -228,6 +229,7 @@ public:
 	MOVIESCENETRACKS_API virtual void SetBlendType(EMovieSceneBlendType InBlendType) override;
 	MOVIESCENETRACKS_API virtual void OnBindingIDsUpdated(const TMap<UE::MovieScene::FFixedObjectBindingID, UE::MovieScene::FFixedObjectBindingID>& OldFixedToNewFixedMap, FMovieSceneSequenceID LocalSequenceID, const FMovieSceneSequenceHierarchy* Hierarchy, IMovieScenePlayer& Player) override;
 	MOVIESCENETRACKS_API virtual void GetReferencedBindings(TArray<FGuid>& OutBindings) override;
+	MOVIESCENETRACKS_API virtual void PreSave(FObjectPreSaveContext SaveContext) override;
 
 public:
 
@@ -341,7 +343,10 @@ public:
 	*/
 	virtual TArray<FConstraintAndActiveChannel>& GetConstraintsChannels()  override;
 
-
+	/*
+	*  Replace the constraint with the specified name with the new one
+	*/
+	virtual void ReplaceConstraint(const FName InName, UTickableConstraint* InConstraint)  override;
 private:
 
 #if WITH_EDITORONLY_DATA
diff --git a/Engine/Source/Runtime/MovieSceneTracks/Public/Sections/MovieSceneConstrainedSection.h b/Engine/Source/Runtime/MovieSceneTracks/Public/Sections/MovieSceneConstrainedSection.h
--- a/Engine/Source/Runtime/MovieSceneTracks/Public/Sections/MovieSceneConstrainedSection.h
+++ b/Engine/Source/Runtime/MovieSceneTracks/Public/Sections/MovieSceneConstrainedSection.h
@@ -62,6 +62,11 @@ public:
 	*/
 	virtual TArray<FConstraintAndActiveChannel>& GetConstraintsChannels() = 0;
 
+	/*
+	*  Replace the constraint with the specified name with the new one
+	*/
+	virtual void ReplaceConstraint(const FName InName, UTickableConstraint* InConstraint) = 0;
+
 	/*
 	* Added Delegate
 	*/
