Sequencer: Support "additive-from-base" blend type when recomposing values for keying

This change also removes the unused single-precision Recompose method and CombineWeighted method.

diff --git a/Engine/Source/Runtime/MovieScene/Private/EntitySystem/BuiltInComponentTypes.cpp b/Engine/Source/Runtime/MovieScene/Private/EntitySystem/BuiltInComponentTypes.cpp
--- a/Engine/Source/Runtime/MovieScene/Private/EntitySystem/BuiltInComponentTypes.cpp
+++ b/Engine/Source/Runtime/MovieScene/Private/EntitySystem/BuiltInComponentTypes.cpp
@@ -267,6 +267,13 @@ FBuiltInComponentTypes::FBuiltInComponentTypes()
 		}
 	}
 
+	{
+		for (int32 Index = 0; Index < UE_ARRAY_COUNT(DoubleResult); ++Index)
+		{
+			ResultToBase.Add(DoubleResult[Index], BaseDouble[Index]);
+		}
+	}
+
 	// Easing component relationships
 	{
 		// Easing components should be duplicated to children
@@ -346,6 +353,14 @@ FBuiltInComponentTypes* FBuiltInComponentTypes::Get()
 	return GMovieSceneBuiltInComponentTypes.Get();
 }
 
+FComponentTypeID FBuiltInComponentTypes::GetBaseValueComponentType(const FComponentTypeID& InResultComponentType)
+{
+	if (FComponentTypeID* BaseComponentType = ResultToBase.Find(InResultComponentType))
+	{
+		return *BaseComponentType;
+	}
+	return FComponentTypeID::Invalid();
+}
 
 } // namespace MovieScene
 } // namespace UE
diff --git a/Engine/Source/Runtime/MovieScene/Private/EntitySystem/MovieSceneDecompositionQuery.cpp b/Engine/Source/Runtime/MovieScene/Private/EntitySystem/MovieSceneDecompositionQuery.cpp
--- a/Engine/Source/Runtime/MovieScene/Private/EntitySystem/MovieSceneDecompositionQuery.cpp
+++ b/Engine/Source/Runtime/MovieScene/Private/EntitySystem/MovieSceneDecompositionQuery.cpp
@@ -11,14 +11,16 @@ namespace UE
 namespace MovieScene
 {
 
-void FDecomposedValue::Decompose(FMovieSceneEntityID EntityID, FWeightedValue& ThisValue, bool& bOutIsAdditive, FWeightedValue& Absolutes, FWeightedValue& Additives) const
+void FDecomposedValue::Decompose(
+		FMovieSceneEntityID EntityID, FWeightedValue& ThisValue, EDecomposedValueBlendType& OutBlendType, 
+		FWeightedValue& Absolutes, FWeightedValue& Additives, FWeightedValue& AdditivesFromBase) const
 {
 	for (TTuple<FMovieSceneEntityID, FWeightedValue> Pair : DecomposedAbsolutes)
 	{
 		if (Pair.Get<0>() == EntityID)
 		{
 			ThisValue = Pair.Value;
-			bOutIsAdditive = false;
+			OutBlendType = EDecomposedValueBlendType::Absolute;
 		}
 		else
 		{
@@ -31,7 +33,7 @@ void FDecomposedValue::Decompose(FMovieSceneEntityID EntityID, FWeightedValue& T
 		if (Pair.Get<0>() == EntityID)
 		{
 			ThisValue = Pair.Value;
-			bOutIsAdditive = true;
+			OutBlendType = EDecomposedValueBlendType::Additive;
 		}
 		else
 		{
@@ -39,18 +41,18 @@ void FDecomposedValue::Decompose(FMovieSceneEntityID EntityID, FWeightedValue& T
 			Additives.Weight += Pair.Value.Weight;
 		}
 	}
-}
-
-float FDecomposedValue::Recompose(FMovieSceneEntityID EntityID, float CurrentValue, const float* InitialValue) const
-{
-	if (InitialValue)
+	for (TTuple<FMovieSceneEntityID, FWeightedValue> Pair : DecomposedAdditivesFromBase)
 	{
-		double WidenedInitialValue = *InitialValue;
-		return (float)Recompose(EntityID, (double)CurrentValue, &WidenedInitialValue);
-	}
-	else
-	{
-		return (float)Recompose(EntityID, (double)CurrentValue, nullptr);
+		if (Pair.Get<0>() == EntityID)
+		{
+			ThisValue = Pair.Value;
+			OutBlendType = EDecomposedValueBlendType::AdditiveFromBase;
+		}
+		else
+		{
+			AdditivesFromBase.Value += (Pair.Value.Value - Pair.Value.BaseValue) * Pair.Value.Weight;
+			AdditivesFromBase.Weight += Pair.Value.Weight;
+		}
 	}
 }
 
@@ -58,13 +60,15 @@ double FDecomposedValue::Recompose(FMovieSceneEntityID RecomposeEntity, double C
 {
 	FWeightedValue DecomposedAbsolute;
 	FWeightedValue DecomposedAdditive;
+	FWeightedValue DecomposedAdditiveFromBase;
 
 	FWeightedValue Channel;
-	bool bIsAdditive = false;
-	Decompose(RecomposeEntity, Channel, bIsAdditive, DecomposedAbsolute, DecomposedAdditive);
+	EDecomposedValueBlendType BlendType(EDecomposedValueBlendType::Absolute);
+	Decompose(RecomposeEntity, Channel, BlendType, DecomposedAbsolute, DecomposedAdditive, DecomposedAdditiveFromBase);
 
 	FWeightedValue ResultAbsolute = Result.Absolute;
 	float TotalAbsoluteWeight = ResultAbsolute.Weight + DecomposedAbsolute.Weight;
+	const bool bIsAdditive = (BlendType != EDecomposedValueBlendType::Absolute);
 	if (!bIsAdditive)
 	{
 		TotalAbsoluteWeight += Channel.Weight;
@@ -77,12 +81,12 @@ double FDecomposedValue::Recompose(FMovieSceneEntityID RecomposeEntity, double C
 	}
 
 	// If this channel is the only thing we decomposed, that is simple
-	if (DecomposedAbsolute.Weight == 0.f && DecomposedAdditive.Weight == 0.f)
+	if (DecomposedAbsolute.Weight == 0.f && DecomposedAdditive.Weight == 0.f && DecomposedAdditiveFromBase.Weight == 0.f)
 	{
 		if (bIsAdditive)
 		{
 			const double WeightedAdditiveResult = CurrentValue - ResultAbsolute.Combine(DecomposedAbsolute).WeightedValue() - Result.Additive;
-			return Channel.Weight == 0.f ? WeightedAdditiveResult : WeightedAdditiveResult / Channel.Weight;
+			return (Channel.Weight == 0.f ? WeightedAdditiveResult : WeightedAdditiveResult / Channel.Weight) + Channel.BaseValue;
 		}
 		else
 		{
@@ -110,20 +114,20 @@ double FDecomposedValue::Recompose(FMovieSceneEntityID RecomposeEntity, double C
 		CurrentValue -= ResultAbsolute.Combine(DecomposedAbsolute).WeightedValue();
 
 		const double ThisAdditive = Channel.WeightedValue();
-		if (ThisAdditive == 0.f && DecomposedAdditive.WeightedValue() == 0.f)
+		if (ThisAdditive == 0.f && DecomposedAdditive.WeightedValue() == 0.f && DecomposedAdditiveFromBase.WeightedValue() == 0.f)
 		{
 			const float TotalAdditiveWeight = DecomposedAdditive.Weight + Channel.Weight;
-			return CurrentValue * Channel.Weight / TotalAdditiveWeight;
+			return (CurrentValue * Channel.Weight / TotalAdditiveWeight) + Channel.BaseValue;
 		}
 
 		// Use the fractions of the values for the recomposition if we have non-zero values
-		const double DecoposeFactor = ThisAdditive / (DecomposedAdditive.WeightedValue() + ThisAdditive);
-		return CurrentValue * DecoposeFactor / Channel.Weight;
+		const double DecomposeFactor = ThisAdditive / (DecomposedAdditive.WeightedValue() + ThisAdditive);
+		return CurrentValue * DecomposeFactor / Channel.Weight + Channel.BaseValue;
 	}
 	else if (DecomposedAdditives.Num() != 0)
 	{
 		// Absolute channel, but we're keying additives, put the full weight to the additives
-		return Channel.Value;
+		return Channel.Value - Channel.BaseValue;
 	}
 	else
 	{
diff --git a/Engine/Source/Runtime/MovieScene/Public/EntitySystem/BuiltInComponentTypes.h b/Engine/Source/Runtime/MovieScene/Public/EntitySystem/BuiltInComponentTypes.h
--- a/Engine/Source/Runtime/MovieScene/Public/EntitySystem/BuiltInComponentTypes.h
+++ b/Engine/Source/Runtime/MovieScene/Public/EntitySystem/BuiltInComponentTypes.h
@@ -406,6 +406,8 @@ public:
 
 	FComponentMask FinishedMask;
 
+public:
+
 	static void Destroy();
 
 	static FBuiltInComponentTypes* Get();
@@ -415,8 +417,13 @@ public:
 		return InObject == nullptr || !IsValidChecked(InObject) || InObject->IsUnreachable();
 	}
 
+	FComponentTypeID GetBaseValueComponentType(const FComponentTypeID& InResultComponentType);
+
 private:
+
 	FBuiltInComponentTypes();
+
+	TMap<FComponentTypeID, FComponentTypeID> ResultToBase;
 };
 
 
diff --git a/Engine/Source/Runtime/MovieScene/Public/EntitySystem/MovieSceneDecompositionQuery.h b/Engine/Source/Runtime/MovieScene/Public/EntitySystem/MovieSceneDecompositionQuery.h
--- a/Engine/Source/Runtime/MovieScene/Public/EntitySystem/MovieSceneDecompositionQuery.h
+++ b/Engine/Source/Runtime/MovieScene/Public/EntitySystem/MovieSceneDecompositionQuery.h
@@ -46,20 +46,17 @@ namespace MovieScene
 	{
 		double Value = 0.f;
 		float Weight = 0.f;
+		
+		double BaseValue = 0.f;  // Should only be set when the value has additive-from-base blend type
 
 		double WeightedValue() const
 		{
-			return Weight != 0.f ? Value / Weight : 0.f;
+			return Weight != 0.f ? (Value - BaseValue) / Weight : 0.f;
 		}
 
 		FWeightedValue Combine(FWeightedValue Other) const
 		{
-			return FWeightedValue{Value + Other.Value, Weight + Other.Weight};
-		}
-
-		FWeightedValue CombineWeighted(FWeightedValue Other) const
-		{
-			return FWeightedValue{Value + Other.Value * Other.Weight, Weight + Other.Weight};
+			return FWeightedValue{Value + Other.Value, Weight + Other.Weight, BaseValue + Other.BaseValue};
 		}
 	};
 
@@ -75,10 +72,21 @@ namespace MovieScene
 
 		TArray<TTuple<FMovieSceneEntityID, FWeightedValue>> DecomposedAbsolutes;
 		TArray<TTuple<FMovieSceneEntityID, FWeightedValue>> DecomposedAdditives;
+		TArray<TTuple<FMovieSceneEntityID, FWeightedValue>> DecomposedAdditivesFromBase;
 
-		MOVIESCENE_API float Recompose(FMovieSceneEntityID EntityID, float CurrentValue, const float* InitialValue) const;
 		MOVIESCENE_API double Recompose(FMovieSceneEntityID EntityID, double CurrentValue, const double* InitialValue) const;
-		MOVIESCENE_API void Decompose(FMovieSceneEntityID EntityID, FWeightedValue& ThisValue, bool& bOutIsAdditive, FWeightedValue& Absolutes, FWeightedValue& Additives) const;
+
+	private:
+		enum class EDecomposedValueBlendType
+		{
+			Absolute,
+			Additive,
+			AdditiveFromBase
+		};
+
+		void Decompose(
+				FMovieSceneEntityID EntityID, FWeightedValue& ThisValue, EDecomposedValueBlendType& OutBlendType, 
+				FWeightedValue& Absolutes, FWeightedValue& Additives, FWeightedValue& AdditivesFromBase) const;
 	};
 
 	// Align results to cache lines so there's no contention between cores
diff --git a/Engine/Source/Runtime/MovieSceneTracks/Private/Systems/MovieScenePiecewiseDoubleBlenderSystem.cpp b/Engine/Source/Runtime/MovieSceneTracks/Private/Systems/MovieScenePiecewiseDoubleBlenderSystem.cpp
--- a/Engine/Source/Runtime/MovieSceneTracks/Private/Systems/MovieScenePiecewiseDoubleBlenderSystem.cpp
+++ b/Engine/Source/Runtime/MovieSceneTracks/Private/Systems/MovieScenePiecewiseDoubleBlenderSystem.cpp
@@ -679,7 +679,10 @@ FGraphEventRef UMovieScenePiecewiseDoubleBlenderSystem::DispatchDecomposeTask(co
 		return nullptr;
 	}
 
+	FBuiltInComponentTypes* BuiltInComponents = FBuiltInComponentTypes::Get();
+
 	TComponentTypeID<double> ResultComponentType = Params.ResultComponentType.ReinterpretCast<double>();
+	TComponentTypeID<double> BaseValueComponentType = BuiltInComponents->GetBaseValueComponentType(Params.ResultComponentType).ReinterpretCast<double>();
 
 	struct FChannelResultTask
 	{
@@ -687,20 +690,26 @@ FGraphEventRef UMovieScenePiecewiseDoubleBlenderSystem::DispatchDecomposeTask(co
 		FAlignedDecomposedValue* Result;
 		uint16 DecomposeBlendChannel;
 		FComponentTypeID AdditiveBlendTag;
+		FComponentTypeID AdditiveFromBaseBlendTag;
 
 		explicit FChannelResultTask(const FValueDecompositionParams& Params, FAlignedDecomposedValue* InResult)
 			: Result(InResult)
 			, DecomposeBlendChannel(Params.DecomposeBlendChannel)
 			, AdditiveBlendTag(FBuiltInComponentTypes::Get()->Tags.AdditiveBlend)
+			, AdditiveFromBaseBlendTag(FBuiltInComponentTypes::Get()->Tags.AdditiveFromBaseBlend)
 		{
 			EntitiesToDecompose.Append(Params.Query.Entities.GetData(), Params.Query.Entities.Num());
 		}
 
-		void ForEachAllocation(const FEntityAllocation* Allocation, TRead<FMovieSceneEntityID> EntityToDecomposeIDs, TRead<FMovieSceneBlendChannelID> BlendChannels, TRead<double> ValueResultComponent, TReadOptional<double> OptionalWeightComponent)
+		void ForEachAllocation(
+				const FEntityAllocation* Allocation, TRead<FMovieSceneEntityID> EntityToDecomposeIDs, 
+				TRead<FMovieSceneBlendChannelID> BlendChannels, TRead<double> ValueResultComponent, 
+				TReadOptional<double> OptionalBaseValueComponent, TReadOptional<double> OptionalWeightComponent)
 		{
 			static const FMovieSceneBlenderSystemID BlenderSystemID = UMovieSceneBlenderSystem::GetBlenderSystemID<UMovieScenePiecewiseDoubleBlenderSystem>();
 
 			const bool bAdditive = Allocation->HasComponent(AdditiveBlendTag);
+			const bool bAdditiveFromBase = Allocation->HasComponent(AdditiveFromBaseBlendTag);
 
 			const int32 Num = Allocation->Num();
 			for (int32 EntityIndex = 0; EntityIndex < Num; ++EntityIndex)
@@ -715,8 +724,9 @@ FGraphEventRef UMovieScenePiecewiseDoubleBlenderSystem::DispatchDecomposeTask(co
 
 				// We've found a contributor for this blend channel
 				const FMovieSceneEntityID EntityToDecompose = EntityToDecomposeIDs[EntityIndex];
-				const float               Weight            = OptionalWeightComponent ? OptionalWeightComponent[EntityIndex] : 1.f;
 				const double              ValueResult       = ValueResultComponent[EntityIndex];
+				const double              BaseValue         = OptionalBaseValueComponent ? OptionalBaseValueComponent[EntityIndex] : 0.f;
+				const float               Weight            = OptionalWeightComponent ? OptionalWeightComponent[EntityIndex] : 1.f;
 
 				if (EntitiesToDecompose.Contains(EntityToDecompose))
 				{
@@ -724,6 +734,10 @@ FGraphEventRef UMovieScenePiecewiseDoubleBlenderSystem::DispatchDecomposeTask(co
 					{
 						Result->Value.DecomposedAdditives.Add(MakeTuple(EntityToDecompose, FWeightedValue{ ValueResult, Weight }));
 					}
+					else if (bAdditiveFromBase)
+					{
+						Result->Value.DecomposedAdditivesFromBase.Add(MakeTuple(EntityToDecompose, FWeightedValue{ ValueResult, Weight, BaseValue }));
+					}
 					else
 					{
 						Result->Value.DecomposedAbsolutes.Add(MakeTuple(EntityToDecompose, FWeightedValue{ ValueResult, Weight }));
@@ -733,6 +747,10 @@ FGraphEventRef UMovieScenePiecewiseDoubleBlenderSystem::DispatchDecomposeTask(co
 				{
 					Result->Value.Result.Additive += ValueResult * Weight;
 				}
+				else if (bAdditiveFromBase)
+				{
+					Result->Value.Result.Additive += (ValueResult - BaseValue) * Weight;
+				}
 				else
 				{
 					Result->Value.Result.Absolute.Value  += ValueResult * Weight;
@@ -742,14 +760,13 @@ FGraphEventRef UMovieScenePiecewiseDoubleBlenderSystem::DispatchDecomposeTask(co
 		}
 	};
 
-	FBuiltInComponentTypes* BuiltInComponents = FBuiltInComponentTypes::Get();
-
 	if (Params.Query.bConvertFromSourceEntityIDs)
 	{
 		return FEntityTaskBuilder()
 			.Read(BuiltInComponents->ParentEntity)
 			.Read(BuiltInComponents->BlendChannelInput)
 			.Read(ResultComponentType)
+			.ReadOptional(BaseValueComponentType)
 			.ReadOptional(BuiltInComponents->WeightAndEasingResult)
 			.FilterAll({ Params.PropertyTag })
 			.template Dispatch_PerAllocation<FChannelResultTask>(&Linker->EntityManager, FSystemTaskPrerequisites(), nullptr, Params, Output);
@@ -760,6 +777,7 @@ FGraphEventRef UMovieScenePiecewiseDoubleBlenderSystem::DispatchDecomposeTask(co
 			.ReadEntityIDs()
 			.Read(BuiltInComponents->BlendChannelInput)
 			.Read(ResultComponentType)
+			.ReadOptional(BaseValueComponentType)
 			.ReadOptional(BuiltInComponents->WeightAndEasingResult)
 			.FilterAll({ Params.PropertyTag })
 			.template Dispatch_PerAllocation<FChannelResultTask>(&Linker->EntityManager, FSystemTaskPrerequisites(), nullptr, Params, Output);
