Sequencer: Constraints:  When undoing a change with a constraint only reset the proxy if the number of channels has changed, previously we would always do this which would cause the curve to get reset.

diff --git a/Engine/Source/Runtime/MovieSceneTracks/Private/Sections/MovieScene3DTransformSection.cpp b/Engine/Source/Runtime/MovieSceneTracks/Private/Sections/MovieScene3DTransformSection.cpp
--- a/Engine/Source/Runtime/MovieSceneTracks/Private/Sections/MovieScene3DTransformSection.cpp
+++ b/Engine/Source/Runtime/MovieSceneTracks/Private/Sections/MovieScene3DTransformSection.cpp
@@ -1088,11 +1088,18 @@ void UMovieScene3DTransformSection::OnChannelOverridesChanged()
 	ChannelProxy = nullptr;
 }
 
-void UMovieScene3DTransformSectionConstraints::Serialize(FArchive& Ar)
+#if WITH_EDITOR
+
+void UMovieScene3DTransformSectionConstraints::PreEditUndo()
+{
+	Super::PreEditUndo();
+	PreEditUndoNumChannels = ConstraintsChannels.Num();
+}
+void UMovieScene3DTransformSectionConstraints::PostEditUndo()
 {
-	Super::Serialize(Ar);
-	//if we are loading during an undo/redo set the cache to be empty
-	if (Ar.IsLoading() && Ar.IsTransacting())
+	Super::PostEditUndo();
+	//if number of channels has changed tell parent section we have changed which will recreate the proxy
+	if (PreEditUndoNumChannels != ConstraintsChannels.Num())
 	{
 		if (UMovieScene3DTransformSection* Section = GetTypedOuter<UMovieScene3DTransformSection>())
 		{
@@ -1102,5 +1109,8 @@ void UMovieScene3DTransformSectionConstraints::Serialize(FArchive& Ar)
 			}
 		}
 	}
+	PreEditUndoNumChannels = ConstraintsChannels.Num();
 }
 
+#endif
+
diff --git a/Engine/Source/Runtime/MovieSceneTracks/Public/Sections/MovieScene3DTransformSection.h b/Engine/Source/Runtime/MovieSceneTracks/Public/Sections/MovieScene3DTransformSection.h
--- a/Engine/Source/Runtime/MovieSceneTracks/Public/Sections/MovieScene3DTransformSection.h
+++ b/Engine/Source/Runtime/MovieSceneTracks/Public/Sections/MovieScene3DTransformSection.h
@@ -198,12 +198,16 @@ class MOVIESCENETRACKS_API UMovieScene3DTransformSectionConstraints : public UOb
 
 public:
 
-	virtual void Serialize(FArchive& Ar) override;
-
 	/** Constraint Channels*/
 	UPROPERTY()
 	TArray<FConstraintAndActiveChannel> ConstraintsChannels;
 
+	/** When undo/redoing if we get a different number of channels we need to make sure our parent section recreates it's channel proxies*/
+#if WITH_EDITOR
+	int32 PreEditUndoNumChannels = -1;
+	virtual void PreEditUndo() override;
+	virtual void PostEditUndo() override;
+#endif
 };
 
 /**
