Sequencer: Constraints: We do need to always recache the proxy on undo otherwise the constraint channel get's wiped.

diff --git a/Engine/Source/Runtime/MovieSceneTracks/Private/Sections/MovieScene3DTransformSection.cpp b/Engine/Source/Runtime/MovieSceneTracks/Private/Sections/MovieScene3DTransformSection.cpp
--- a/Engine/Source/Runtime/MovieSceneTracks/Private/Sections/MovieScene3DTransformSection.cpp
+++ b/Engine/Source/Runtime/MovieSceneTracks/Private/Sections/MovieScene3DTransformSection.cpp
@@ -1090,26 +1090,16 @@ void UMovieScene3DTransformSection::OnChannelOverridesChanged()
 
 #if WITH_EDITOR
 
-void UMovieScene3DTransformSectionConstraints::PreEditUndo()
-{
-	Super::PreEditUndo();
-	PreEditUndoNumChannels = ConstraintsChannels.Num();
-}
 void UMovieScene3DTransformSectionConstraints::PostEditUndo()
 {
 	Super::PostEditUndo();
-	//if number of channels has changed tell parent section we have changed which will recreate the proxy
-	if (PreEditUndoNumChannels != ConstraintsChannels.Num())
+	if (UMovieScene3DTransformSection* Section = GetTypedOuter<UMovieScene3DTransformSection>())
 	{
-		if (UMovieScene3DTransformSection* Section = GetTypedOuter<UMovieScene3DTransformSection>())
+		if (IMovieSceneConstrainedSection* ConstrainedSection = Cast<IMovieSceneConstrainedSection>(Section))
 		{
-			if (IMovieSceneConstrainedSection* ConstrainedSection = Cast<IMovieSceneConstrainedSection>(Section))
-			{
-				ConstrainedSection->OnConstraintsChanged();
-			}
+			ConstrainedSection->OnConstraintsChanged();
 		}
 	}
-	PreEditUndoNumChannels = ConstraintsChannels.Num();
 }
 
 #endif
diff --git a/Engine/Source/Runtime/MovieSceneTracks/Public/Sections/MovieScene3DTransformSection.h b/Engine/Source/Runtime/MovieSceneTracks/Public/Sections/MovieScene3DTransformSection.h
--- a/Engine/Source/Runtime/MovieSceneTracks/Public/Sections/MovieScene3DTransformSection.h
+++ b/Engine/Source/Runtime/MovieSceneTracks/Public/Sections/MovieScene3DTransformSection.h
@@ -202,10 +202,8 @@ public:
 	UPROPERTY()
 	TArray<FConstraintAndActiveChannel> ConstraintsChannels;
 
-	/** When undo/redoing if we get a different number of channels we need to make sure our parent section recreates it's channel proxies*/
+	/** When undo/redoing we need to recreate channel proxies after we are done*/
 #if WITH_EDITOR
-	int32 PreEditUndoNumChannels = -1;
-	virtual void PreEditUndo() override;
 	virtual void PostEditUndo() override;
 #endif
 };
