Sequencer: NLA expose start offset translate/rotate and show skeleton to BP.
#trivial

diff --git a/Engine/Source/Runtime/MovieSceneTracks/Public/Sections/MovieSceneSkeletalAnimationSection.h b/Engine/Source/Runtime/MovieSceneTracks/Public/Sections/MovieSceneSkeletalAnimationSection.h
--- a/Engine/Source/Runtime/MovieSceneTracks/Public/Sections/MovieSceneSkeletalAnimationSection.h
+++ b/Engine/Source/Runtime/MovieSceneTracks/Public/Sections/MovieSceneSkeletalAnimationSection.h
@@ -159,12 +159,12 @@ private:
 	FName SlotName_DEPRECATED;
 
 public:
-	UPROPERTY(EditAnywhere, Category = "Root Motions")
-	/* Location Offset applied to this this animations start root motion*/
+	UPROPERTY(BlueprintReadWrite, EditAnywhere, Category = "Root Motions")
+	/* Location Offset applied to this animations start root motion*/
 	FVector StartLocationOffset;
 
-	UPROPERTY(EditAnywhere, Category = "Root Motions")
-	/* Rotation Offset applied to this this animations start root motion*/
+	UPROPERTY(BlueprintReadWrite,EditAnywhere, Category = "Root Motions")
+	/* Rotation Offset applied to this animations start root motion*/
 	FRotator StartRotationOffset;
 
 	UPROPERTY()
@@ -196,7 +196,7 @@ public:
 
 #if WITH_EDITORONLY_DATA
 	/** Whether to show the underlying skeleton for this section. */
-	UPROPERTY(EditAnywhere, Category = "Root Motions")
+	UPROPERTY(BlueprintReadWrite, EditAnywhere, Category = "Root Motions")
 	bool bShowSkeleton;
 #endif
 	//Temporary transform used to specify the global OffsetTransform while calculting the root motions.
