Sequener: fix initial value system not working in interrogation linkers

diff --git a/Engine/Source/Runtime/MovieSceneTracks/Private/Systems/MovieSceneInitialValueSystem.cpp b/Engine/Source/Runtime/MovieSceneTracks/Private/Systems/MovieSceneInitialValueSystem.cpp
--- a/Engine/Source/Runtime/MovieSceneTracks/Private/Systems/MovieSceneInitialValueSystem.cpp
+++ b/Engine/Source/Runtime/MovieSceneTracks/Private/Systems/MovieSceneInitialValueSystem.cpp
@@ -157,7 +157,8 @@ void UMovieSceneInitialValueSystem::OnRun(FSystemTaskPrerequisites& InPrerequisi
 		// When there is no caching extension, or we are interrogating we simply initialize any initial values directly without going through the cache
 		FEntityComponentFilter Filter;
 		Filter.Any(Mutation.AnyInitialValue);
-		Filter.All({ BuiltInComponents->BoundObject, BuiltInComponents->Tags.NeedsLink });
+		Filter.Any({ BuiltInComponents->BoundObject, BuiltInComponents->Interrogation.OutputKey });
+		Filter.All({ BuiltInComponents->Tags.NeedsLink });
 
 		for (FEntityAllocationIteratorItem Item : Linker->EntityManager.Iterate(&Filter))
 		{
