Fix CIS issue

diff --git a/Engine/Source/Runtime/MovieSceneTracks/Public/Tests/MovieSceneTestDataBuilders.h b/Engine/Source/Runtime/MovieSceneTracks/Public/Tests/MovieSceneTestDataBuilders.h
--- a/Engine/Source/Runtime/MovieSceneTracks/Public/Tests/MovieSceneTestDataBuilders.h
+++ b/Engine/Source/Runtime/MovieSceneTracks/Public/Tests/MovieSceneTestDataBuilders.h
@@ -2,12 +2,15 @@
 
 #pragma once
 
+#include "Channels/MovieSceneChannelHandle.h"
 #include "Channels/MovieSceneChannelTraits.h"
+#include "Channels/MovieSceneChannelProxy.h"
 #include "CoreTypes.h"
 #include "Generators/MovieSceneEasingCurves.h"
 #include "IMovieScenePlayer.h"
 #include "KeyParams.h"
 #include "MovieScene.h"
+#include "MovieSceneSection.h"
 #include "MovieSceneSequence.h"
 #include "Templates/EnableIf.h"
 #include "Templates/PointerIsConvertibleFromTo.h"
@@ -236,6 +239,7 @@ struct FSequenceSectionBuilder
 		return AddKeys(Channel, InTimes, InValues, Interpolation);
 	}
 
+#if WITH_EDITOR
 	template<typename ChannelType, typename ValueType>
 	FSequenceSectionBuilder& AddKeys(
 			FName InChannelName, 
@@ -247,6 +251,7 @@ struct FSequenceSectionBuilder
 		TMovieSceneChannelHandle<ChannelType> ChannelHandle = ChannelProxy.GetChannelByName<ChannelType>(InChannelName);
 		return AddKeys(ChannelHandle.Get(), InTimes, InValues, Interpolation);
 	}
+#endif
 
 	template<typename ChannelType, typename ValueType>
 	FSequenceSectionBuilder& AddKeys(
