Sequencer: Fix edge cases with value decomposition, add lots of comments, and add unit tests

- Fix the math in some cases that didn't correctly include all sorts of absolutes or additives, or didn't handle weight correctly.
- Clarify the difference between weighted values and accumulated weighted values.
- Merge additives and additives-from-base together.

diff --git a/Engine/Source/Runtime/MovieScene/Private/EntitySystem/MovieSceneDecompositionQuery.cpp b/Engine/Source/Runtime/MovieScene/Private/EntitySystem/MovieSceneDecompositionQuery.cpp
--- a/Engine/Source/Runtime/MovieScene/Private/EntitySystem/MovieSceneDecompositionQuery.cpp
+++ b/Engine/Source/Runtime/MovieScene/Private/EntitySystem/MovieSceneDecompositionQuery.cpp
@@ -1,20 +1,22 @@
 // Copyright Epic Games, Inc. All Rights Reserved.
 
 #include "EntitySystem/MovieSceneDecompositionQuery.h"
-#include "Algo/Accumulate.h"
 
 #include UE_INLINE_GENERATED_CPP_BY_NAME(MovieSceneDecompositionQuery)
 
-
 namespace UE
 {
 namespace MovieScene
 {
 
 void FDecomposedValue::Decompose(
-		FMovieSceneEntityID EntityID, FWeightedValue& ThisValue, EDecomposedValueBlendType& OutBlendType, 
-		FWeightedValue& Absolutes, FWeightedValue& Additives, FWeightedValue& AdditivesFromBase) const
+		FMovieSceneEntityID EntityID, 
+		FWeightedValue& ThisValue, EDecomposedValueBlendType& OutBlendType, 
+		FAccumulatedWeightedValue& Absolutes, 
+		FAccumulatedWeightedValue& Additives) const
 {
+	// Look through all our decomposed values and isolate the one that corresponds to the particular given entity that
+	// we're interested in. All other decomposed values get accumulated in the given output parameters.
 	for (TTuple<FMovieSceneEntityID, FWeightedValue> Pair : DecomposedAbsolutes)
 	{
 		if (Pair.Get<0>() == EntityID)
@@ -24,8 +26,7 @@ void FDecomposedValue::Decompose(
 		}
 		else
 		{
-			Absolutes.Value += Pair.Value.Value * Pair.Value.Weight;
-			Absolutes.Weight += Pair.Value.Weight;
+			Absolutes.AccumulateThis(Pair.Value);
 		}
 	}
 	for (TTuple<FMovieSceneEntityID, FWeightedValue> Pair : DecomposedAdditives)
@@ -33,126 +34,183 @@ void FDecomposedValue::Decompose(
 		if (Pair.Get<0>() == EntityID)
 		{
 			ThisValue = Pair.Value;
-			OutBlendType = EDecomposedValueBlendType::Additive;
-		}
-		else
-		{
-			Additives.Value += Pair.Value.Value * Pair.Value.Weight;
-			Additives.Weight += Pair.Value.Weight;
-		}
-	}
-	for (TTuple<FMovieSceneEntityID, FWeightedValue> Pair : DecomposedAdditivesFromBase)
-	{
-		if (Pair.Get<0>() == EntityID)
-		{
-			ThisValue = Pair.Value;
-			OutBlendType = EDecomposedValueBlendType::AdditiveFromBase;
+			OutBlendType = (Pair.Value.BaseValue != 0.f) ? 
+				EDecomposedValueBlendType::AdditiveFromBase : EDecomposedValueBlendType::Additive;
 		}
 		else
 		{
-			AdditivesFromBase.Value += (Pair.Value.Value - Pair.Value.BaseValue) * Pair.Value.Weight;
-			AdditivesFromBase.Weight += Pair.Value.Weight;
+			Additives.AccumulateThis(Pair.Value);
 		}
 	}
 }
 
 double FDecomposedValue::Recompose(FMovieSceneEntityID RecomposeEntity, double CurrentValue, const double* InitialValue) const
 {
-	FWeightedValue DecomposedAbsolute;
-	FWeightedValue DecomposedAdditive;
-	FWeightedValue DecomposedAdditiveFromBase;
-
+	// First, a little reminder... the formula for blending values is:
+	//
+	// Value = Absolutes + Additives
+	//
+	// Absolutes = (Abs1 * AbsWeight1 + ... + AbsN * AbsWeightN) / (AbsWeight1 + ... + AbsWeightN)
+	// Additives = (Add1 * AddWeight1 + ... AddN * AddWeightN)
+	//
+	//
+	// Sort through all the data we have so that the contribution of RecomposeEntity is set aside in Channel.  The
+	// contributions of other decomposed entities (if any) are combined into OtherAbsolute and OtherAdditive, which
+	// should also be combined with Result (which contains non-decomposed entities) to get the full picture.
 	FWeightedValue Channel;
 	EDecomposedValueBlendType BlendType(EDecomposedValueBlendType::Absolute);
-	Decompose(RecomposeEntity, Channel, BlendType, DecomposedAbsolute, DecomposedAdditive, DecomposedAdditiveFromBase);
+	FAccumulatedWeightedValue OtherAbsolute;
+	FAccumulatedWeightedValue OtherAdditive;
+	Decompose(RecomposeEntity, Channel, BlendType, OtherAbsolute, OtherAdditive);
 
-	FWeightedValue ResultAbsolute = Result.Absolute;
-	float TotalAbsoluteWeight = ResultAbsolute.Weight + DecomposedAbsolute.Weight;
+	FAccumulatedWeightedValue ResultAbsolute = Result.Absolute;
 	const bool bIsAdditive = (BlendType != EDecomposedValueBlendType::Absolute);
+	float TotalAbsoluteWeight = ResultAbsolute.TotalWeight + OtherAbsolute.TotalWeight;
 	if (!bIsAdditive)
 	{
 		TotalAbsoluteWeight += Channel.Weight;
 	}
 	if (TotalAbsoluteWeight < 1.f && InitialValue != nullptr)
 	{
+		// If all the absolutes we know about (the decomposed and the non-decomposed) do not amount up to 100% weight,
+		// let's fill the gap with the initial value.
 		const float InitialValueWeight = (1.f - TotalAbsoluteWeight);
-		ResultAbsolute.Value = (*InitialValue) * InitialValueWeight + ResultAbsolute.WeightedValue();
-		ResultAbsolute.Weight = 1.f;
+		ResultAbsolute.Total = (*InitialValue) * InitialValueWeight + ResultAbsolute.Total;
+		ResultAbsolute.TotalWeight = InitialValueWeight + ResultAbsolute.TotalWeight;
+		TotalAbsoluteWeight = 1.f;
 	}
 
-	// If this channel is the only thing we decomposed, that is simple
-	if (DecomposedAbsolute.Weight == 0.f && DecomposedAdditive.Weight == 0.f && DecomposedAdditiveFromBase.Weight == 0.f)
+	// If this channel is the only thing we decomposed, we just need to adjust it by the difference
+	// between what we have and what the desired current value is.
+	if (OtherAbsolute.TotalWeight == 0.f && OtherAdditive.TotalWeight == 0.f)
 	{
 		if (bIsAdditive)
 		{
-			const double WeightedAdditiveResult = CurrentValue - ResultAbsolute.Combine(DecomposedAbsolute).WeightedValue() - Result.Additive;
+			// For an additive channel, we just put the missing difference on it, adjusted by its weight.
+			const double WeightedAdditiveResult = CurrentValue - ResultAbsolute.Normalize() - Result.Additive;
 			return (Channel.Weight == 0.f ? WeightedAdditiveResult : WeightedAdditiveResult / Channel.Weight) + Channel.BaseValue;
 		}
-		else
+		else // is absolute
 		{
-			if (Channel.Weight != 0.f)
+			// Absolutes get all added together, and normalized by the total weight, so we can't simply find the
+			// missing difference like we do with additives. The overall formula for our case is this:
+			//
+			// Value = (WeightedResultAbsoluteValue + WeightedChannelValue) / (ResultAbsoluteWeight + ChannelWeight) + ResultAdditive
+			//
+			// Therefore:
+			//
+			// WeightedChannelValue = (Value - ResultAdditive) * (ResultAbsoluteWeight + ChannelWeight) - WeightedResultAbsoluteValue
+			// ChannelValue = WeightedChannelValue / ChannelWeight
+			//
+			// Note that if the channel's weight is zero, it doesn't matter what value we key, we can never reach
+			// whatever value we want to hit. In order to provide a better default (and avoid a division by zero), we
+			// single this case out and set the channel value as if it was weighted to 100%.
+			//
+			// Note: this assumes that weights are always positive or zero though (we don't support negative weights
+			// that could, when combined with equal positive weights, have a sum of zero).
+			//
+			if (Channel.Weight == 0.f)
 			{
-				const float TotalWeight = Channel.Weight + ResultAbsolute.Weight;
-				const double WeightedAbsoluteResult = CurrentValue - Result.Additive - ResultAbsolute.Value / TotalWeight;
-				return WeightedAbsoluteResult * TotalWeight / Channel.Weight;
+				// Note that we use the pure result here (Result instead of AbsoluteResult). This is because
+				// AbsoluteResult may include the initial value. As mentioned, we key the channel as if it had 100%
+				// weight. If it has 100% weight, the initial value wouldn't be involved, so we shouldn't count it in.
+				return (CurrentValue - Result.Additive) * (Result.Absolute.TotalWeight + 1.f) - Result.Absolute.Total;
 			}
 			else
 			{
-				return CurrentValue - Result.Additive - ResultAbsolute.WeightedValue();
+				// Normal use-case... just run the formula from above.
+				const double WeightedAbsoluteResult = (CurrentValue - Result.Additive) * TotalAbsoluteWeight - ResultAbsolute.Total;
+				return (Channel.Weight == 0.f ? WeightedAbsoluteResult : WeightedAbsoluteResult / Channel.Weight);
 			}
 		}
 	}
 
-	// If the channel had no weight, we can't recompose it - everything else will get the full weighting
+	// If we're here, we have multiple channels (entities) that we are decomposing for.
+	// 
+	// This generally means we are keying multiple channels at once, and we're expecting this function to be called once
+	// for each channel being decomposed.  This is important because we will be spreading the needed value changes over
+	// these multiple channels, as needed, which means only passing some of it onto the channel (entity) given as the
+	// first parameter.
+
+	// First, check if the current channel has any weight. If it doesn't, let's just ignore it and hope to be able to do
+	// something better in subsequent calls to this function for which we look at another channel.
 	if (Channel.Weight == 0.f)
 	{
 		return Channel.Value;
 	}
 
+	// Main case for multiple decomposed values. The gist of it is as follows:
+	//
+	// - Mix of absolutes and additives: leave the absolutes alone, and spread the required changes among additives,
+	//      based on these additives' proportional values (so channels with larger values get a bigger slice of the pie).
+	//
+	// - Additives only: as above, we proportionally spread the required changes among the additive channels.
+	//
+	// - Absolutes only: we proportionally spread the required changes among the absolute channels based on their
+	//		weight (so channels with greater weight get a bigger slice of the pie).
+	//
 	if (bIsAdditive)
 	{
-		CurrentValue -= ResultAbsolute.Combine(DecomposedAbsolute).WeightedValue();
+		const double ThisAdditive = Channel.Get();
+		const double AllOtherAbsolutes = ResultAbsolute.Accumulate(OtherAbsolute).Normalize();
 
-		const double ThisAdditive = Channel.WeightedValue();
-		if (ThisAdditive == 0.f && DecomposedAdditive.WeightedValue() == 0.f && DecomposedAdditiveFromBase.WeightedValue() == 0.f)
+		// We are keying an additive channel, but it's currently zero and no other additives are involved.
+		// Key our channel with the difference to reach the desired value.
+		if (OtherAdditive.Total + ThisAdditive == 0.f)
 		{
-			const float TotalAdditiveWeight = DecomposedAdditive.Weight + Channel.Weight;
-			return (CurrentValue * Channel.Weight / TotalAdditiveWeight) + Channel.BaseValue;
+			const double WeightedChannelAdditive = (CurrentValue - AllOtherAbsolutes - Result.Additive);
+			return (Channel.Weight == 0.f) ? (WeightedChannelAdditive + Channel.BaseValue) : (WeightedChannelAdditive / Channel.Weight + Channel.BaseValue);
 		}
 
-		// Use the fractions of the values for the recomposition if we have non-zero values
-		const double DecomposeFactor = ThisAdditive / (DecomposedAdditive.WeightedValue() + ThisAdditive);
-		return CurrentValue * DecomposeFactor / Channel.Weight + Channel.BaseValue;
+		// Scale up/down each of the additives we want to key proportionally to their value. That is, bigger additives
+		// will be changed proportionally more than smaller additives. But the sum of all increases or decreases will
+		// amount to whatever is needed to reach the desired final value.
+		const double DecomposeFactor = ThisAdditive / (OtherAdditive.Total + ThisAdditive);
+		const double WeightedChannelAdditive = (CurrentValue - AllOtherAbsolutes - Result.Additive) * DecomposeFactor;
+		return (Channel.Weight == 0.f) ? (WeightedChannelAdditive + Channel.BaseValue) : (WeightedChannelAdditive / Channel.Weight + Channel.BaseValue);
 	}
 	else if (DecomposedAdditives.Num() != 0)
 	{
-		// Absolute channel, but we're keying additives, put the full weight to the additives
-		return Channel.Value - Channel.BaseValue;
+		// We are currently dealing with keying an absolute channel, but we're also keying additives in the same
+		// operation (since some of them are decomposed). Let's put the full weight of keying on the additives, and
+		// leave the absolute value alone.
+		return Channel.Value;
 	}
 	else
 	{
-		const float TotalDecomposedWeight = DecomposedAbsolute.Weight + Channel.Weight;
-		CurrentValue -= Result.Additive;
-		CurrentValue *= ResultAbsolute.Weight + TotalDecomposedWeight;
-		CurrentValue -= ResultAbsolute.Value;
-
-		const double AbsValue = Algo::Accumulate(DecomposedAbsolutes, 0.f, [](double Accum, TTuple<FMovieSceneEntityID, FWeightedValue> In) { return Accum + FMath::Abs(In.Value.Value)*In.Value.Weight; });
-		if (AbsValue != 0.f)
+		// We are keying multiple absolute channels only. We just need to take the difference between the current value
+		// and the desired value, and spread that across all our absolute channels proportionally based on weight. This
+		// proportion, for a given channel, is:
+		//
+		// DecomposeFactor = Channel.Weight / (AllOtherWeights + Channel.Weight)
+		//
+		// The current value is:
+		//
+		// AllOtherAbsolutes = (WeightedResultAbsolute + WeightedOtherAbsolute)
+		// AllOtherWeights = (ResultAbsoluteWeight + OtherAbsoluteWeight)
+		// ActualValue = (WeightedChannelValue + AllOtherAbsolutes) / (ChannelWeight + AllOtherWeights) + ResultAdditive
+		// 
+		// The missing difference for the current channel is:
+		//
+		// ChannelDelta = (Value - ActualValue) * DecomposeFactor
+		//
+		// So we want the channel to be increased by ChannelDelta.
+		//
+		const FAccumulatedWeightedValue TotalAbsolutes = ResultAbsolute.Accumulate(OtherAbsolute).Accumulate(Channel);
+		if (TotalAbsolutes.TotalWeight == 0.f)
 		{
-			return ((CurrentValue * FMath::Abs(Channel.Value) * Channel.Weight / AbsValue) - Channel.Value) / Channel.Weight;
+			const double ChannelDelta = CurrentValue - Result.Additive;
+			return Channel.Value + ChannelDelta;
 		}
-		else if (TotalDecomposedWeight == 0.f)
+		else
 		{
-			return Channel.Value;
+			const double ActualValue = TotalAbsolutes.Normalize() + Result.Additive;
+			const double ChannelDelta = (CurrentValue - ActualValue);
+			return Channel.Value + ChannelDelta;
 		}
-		return (CurrentValue * Channel.Weight / TotalDecomposedWeight) / Channel.Weight;
 	}
-
-	return Channel.Value;
 }
 
 } // namespace MovieScene
 } // namespace UE
 
-
-
diff --git a/Engine/Source/Runtime/MovieScene/Public/EntitySystem/MovieSceneDecompositionQuery.h b/Engine/Source/Runtime/MovieScene/Public/EntitySystem/MovieSceneDecompositionQuery.h
--- a/Engine/Source/Runtime/MovieScene/Public/EntitySystem/MovieSceneDecompositionQuery.h
+++ b/Engine/Source/Runtime/MovieScene/Public/EntitySystem/MovieSceneDecompositionQuery.h
@@ -40,40 +40,100 @@ namespace MovieScene
 		UObject* Object = nullptr;
 	};
 
-
-	/** Used for decomposing how a final blended value was blended */
+	/** Holds a weighted value */
 	struct FWeightedValue
 	{
+		/** The raw value */
 		double Value = 0.f;
+		/** The weight of the value */
 		float Weight = 0.f;
-		
-		double BaseValue = 0.f;  // Should only be set when the value has additive-from-base blend type
+		/** The base value, which should only be set when the value has additive-from-base blend type */
+		double BaseValue = 0.f;
 
-		double WeightedValue() const
+		/** Multiplies the value by the weight */
+		double Get() const
 		{
-			return Weight != 0.f ? (Value - BaseValue) / Weight : 0.f;
+			return (Value - BaseValue) * Weight;
 		}
+	};
 
-		FWeightedValue Combine(FWeightedValue Other) const
+	/** Holds an accumulation of weighted values */
+	struct FAccumulatedWeightedValue
+	{
+		/** Accumulated value */
+		double Total = 0.f;
+		/** The total weight of all accumulated values */
+		float TotalWeight = 0.f;
+
+		/** Normalizes the total value by the total weight */
+		double Normalize() const
+		{
+			return TotalWeight != 0.f ? Total / TotalWeight : Total;
+		}
+
+		/** Accumulates more values into this value */
+		FAccumulatedWeightedValue& AccumulateThis(const FWeightedValue& Other)
+		{
+			Total += Other.Get();
+			TotalWeight += Other.Weight;
+			return *this;
+		}
+		/** Accumulates more values into this value */
+		FAccumulatedWeightedValue& AccumulateThis(const FAccumulatedWeightedValue& Other)
 		{
-			return FWeightedValue{Value + Other.Value, Weight + Other.Weight, BaseValue + Other.BaseValue};
+			Total += Other.Total;
+			TotalWeight += Other.TotalWeight;
+			return *this;
+		}
+		/** Accumulates values */
+		FAccumulatedWeightedValue Accumulate(const FWeightedValue& Other) const
+		{
+			return FAccumulatedWeightedValue{ Total + Other.Get(), TotalWeight + Other.Weight };
+		}
+		/** Accumulates values */
+		FAccumulatedWeightedValue Accumulate(const FAccumulatedWeightedValue& Other) const
+		{
+			return FAccumulatedWeightedValue{ Total + Other.Total, TotalWeight + Other.TotalWeight };
 		}
 	};
 
+	/**
+	 * A structure for holding all the values of channels contributing to an animated property.
+	 *
+	 * Some of these channels are "decomposed", i.e. their values are set aside and matched to the channel entity that
+	 * produced them. The other channels are mixed together in a non-decomposed result.
+	 */
 	struct FDecomposedValue
 	{
 		struct FResult
 		{
-			FWeightedValue Absolute;
+			/** Accumulated absolute values */
+			FAccumulatedWeightedValue Absolute;
+
+			/** Accumulated additive values (already weighted) */
 			double Additive = 0.f;
 		};
 
+		/** The absolute and additive values for all the non-decomposed channels */
 		FResult Result;
 
+		/**
+		 * Decomposed values for channels we're interested in. Note that decomposed additives-from-base are lumped with
+		 * decomposed additives.
+		 */
 		TArray<TTuple<FMovieSceneEntityID, FWeightedValue>> DecomposedAbsolutes;
 		TArray<TTuple<FMovieSceneEntityID, FWeightedValue>> DecomposedAdditives;
-		TArray<TTuple<FMovieSceneEntityID, FWeightedValue>> DecomposedAdditivesFromBase;
 
+		/**
+		 * Get the value that the channel behind the given entity should have in order for the combined blended values
+		 * of all known channels to produce the desired current value.
+		 *
+		 * @param EntityID The entity matching the channel we want to determine the value of
+		 * @param CurrentValue The desired current value for the overall property animated by the channel (among other
+		 *					   channels)
+		 * @param InitialValue An optional pointer to the initial value of the property
+		 * @return The value that the entity's channel should have to produce the desired current value
+		 */
 		MOVIESCENE_API double Recompose(FMovieSceneEntityID EntityID, double CurrentValue, const double* InitialValue) const;
 
 	private:
@@ -84,9 +144,15 @@ namespace MovieScene
 			AdditiveFromBase
 		};
 
+		/**
+		 * The current structure keeps track of more than one channel/entity we're interested in. This method will
+		 * reduce it to just one particular channel/entity, throwing the other ones in the per-blend-type outputs.
+		 */
 		void Decompose(
-				FMovieSceneEntityID EntityID, FWeightedValue& ThisValue, EDecomposedValueBlendType& OutBlendType, 
-				FWeightedValue& Absolutes, FWeightedValue& Additives, FWeightedValue& AdditivesFromBase) const;
+				FMovieSceneEntityID EntityID, 
+				FWeightedValue& ThisValue, EDecomposedValueBlendType& OutBlendType, 
+				FAccumulatedWeightedValue& Absolutes, 
+				FAccumulatedWeightedValue& Additives) const;
 	};
 
 	// Align results to cache lines so there's no contention between cores
diff --git a/Engine/Source/Runtime/MovieSceneTracks/Private/Systems/MovieScenePiecewiseDoubleBlenderSystem.cpp b/Engine/Source/Runtime/MovieSceneTracks/Private/Systems/MovieScenePiecewiseDoubleBlenderSystem.cpp
--- a/Engine/Source/Runtime/MovieSceneTracks/Private/Systems/MovieScenePiecewiseDoubleBlenderSystem.cpp
+++ b/Engine/Source/Runtime/MovieSceneTracks/Private/Systems/MovieScenePiecewiseDoubleBlenderSystem.cpp
@@ -736,7 +736,7 @@ FGraphEventRef UMovieScenePiecewiseDoubleBlenderSystem::DispatchDecomposeTask(co
 					}
 					else if (bAdditiveFromBase)
 					{
-						Result->Value.DecomposedAdditivesFromBase.Add(MakeTuple(EntityToDecompose, FWeightedValue{ ValueResult, Weight, BaseValue }));
+						Result->Value.DecomposedAdditives.Add(MakeTuple(EntityToDecompose, FWeightedValue{ ValueResult, Weight, BaseValue }));
 					}
 					else
 					{
@@ -753,8 +753,8 @@ FGraphEventRef UMovieScenePiecewiseDoubleBlenderSystem::DispatchDecomposeTask(co
 				}
 				else
 				{
-					Result->Value.Result.Absolute.Value  += ValueResult * Weight;
-					Result->Value.Result.Absolute.Weight += Weight;
+					Result->Value.Result.Absolute.Total += ValueResult * Weight;
+					Result->Value.Result.Absolute.TotalWeight += Weight;
 				}
 			}
 		}
diff --git a/Engine/Source/Runtime/MovieSceneTracks/Private/Tests/MovieSceneDecomposerTests.cpp b/Engine/Source/Runtime/MovieSceneTracks/Private/Tests/MovieSceneDecomposerTests.cpp
new file mode 100644
--- /dev/null
+++ b/Engine/Source/Runtime/MovieSceneTracks/Private/Tests/MovieSceneDecomposerTests.cpp
@@ -0,0 +1,621 @@
+// Copyright Epic Games, Inc. All Rights Reserved.
+
+#include "Tests/MovieSceneDecomposerTests.h"
+#include "Channels/MovieSceneFloatChannel.h"
+#include "EntitySystem/Interrogation/MovieSceneInterrogatedPropertyInstantiator.h"
+#include "EntitySystem/Interrogation/MovieSceneInterrogationLinker.h"
+#include "EntitySystem/MovieSceneDecompositionQuery.h"
+#include "EntitySystem/MovieSceneEntityIDs.h"
+#include "EntitySystem/MovieSceneEntitySystemRunner.h"
+#include "Evaluation/MovieSceneEvaluationTemplateInstance.h"
+#include "IMovieScenePlayer.h"
+#include "Misc/AutomationTest.h"
+#include "Sections/MovieSceneFloatSection.h"
+#include "Tests/MovieSceneTestDataBuilders.h"
+#include "Tracks/MovieSceneFloatTrack.h"
+#include "Tracks/MovieSceneCinematicShotTrack.h"
+
+#define LOCTEXT_NAMESPACE "MovieSceneDecomposerTests"
+
+namespace UE::MovieScene::Test
+{
+
+template<typename PropertyTraits>
+TRecompositionResult<typename PropertyTraits::StorageType>
+RecomposeBlendOperational(
+		UObject* InObject,
+		TArray<UMovieSceneSection*> InPropertySections,
+		const UE::MovieScene::TPropertyComponents<PropertyTraits>& InPropertyComponents,
+		const typename PropertyTraits::StorageType& InCurrentValue,
+		FFrameTime InCurrentTime = FFrameTime(0))
+{
+	using namespace UE::MovieScene;
+
+	// Create the interrogator.
+	FSystemInterrogator Interrogator;
+	Interrogator.TrackImportedEntities(true);
+	
+	TGuardValue<FEntityManager*> DebugVizGuard(GEntityManagerForDebuggingVisualizers, &Interrogator.GetLinker()->EntityManager);
+
+	// All given sections should belong to the same track.
+	check(InPropertySections.Num() > 0);
+	UMovieScenePropertyTrack* PropertyTrack = InPropertySections[0]->GetTypedOuter<UMovieScenePropertyTrack>();
+	check(PropertyTrack);
+
+	// Run an interrogation on the track at the specified time.
+	FInterrogationKey InterrogationKey(FInterrogationKey::Default());
+	FInterrogationChannel InterrogationChannel = Interrogator.AllocateChannel(InObject, PropertyTrack->GetPropertyBinding());
+	InterrogationKey.Channel = InterrogationChannel;
+	Interrogator.ImportTrack(PropertyTrack, InterrogationChannel);
+
+	Interrogator.AddInterrogation(InCurrentTime);
+	Interrogator.Update();
+
+	// Gather the entity IDs for the sections we're interested in.
+	TArray<FMovieSceneEntityID> EntityIDs;
+	for (UMovieSceneSection* PropertySection : InPropertySections)
+	{
+		const FMovieSceneEntityID EntityID = Interrogator.FindEntityFromOwner(InterrogationKey, PropertySection, 0);
+		if (ensure(EntityID.IsValid()))
+		{
+			EntityIDs.Add(EntityID);
+		}
+	}
+	if (EntityIDs.IsEmpty())
+	{
+		return TRecompositionResult<typename PropertyTraits::StorageType>(0, 1);
+	}
+
+	// Run the recomposition on those sections' entities.
+	UMovieSceneInterrogatedPropertyInstantiatorSystem* System = Interrogator.GetLinker()->FindSystem<UMovieSceneInterrogatedPropertyInstantiatorSystem>();
+
+	FDecompositionQuery Query;
+	Query.Entities = EntityIDs;
+	Query.bConvertFromSourceEntityIDs = false;
+	Query.Object = InObject;
+	TRecompositionResult<typename PropertyTraits::StorageType> RecomposeResult = System->RecomposeBlendOperational(
+			InPropertyComponents, Query, InCurrentValue);
+
+	return RecomposeResult;
+}
+
+template<typename PropertyTraits>
+TRecompositionResult<typename PropertyTraits::StorageType>
+RecomposeBlendOperational(
+	UObject* InObject,
+	UMovieSceneSection* InPropertySection,
+	const UE::MovieScene::TPropertyComponents<PropertyTraits>& InPropertyComponents,
+	const typename PropertyTraits::StorageType& InCurrentValue,
+	FFrameTime InCurrentTime = FFrameTime(0))
+{
+	return RecomposeBlendOperational(InObject, TArray<UMovieSceneSection*>({ InPropertySection }), InPropertyComponents, InCurrentValue, InCurrentTime);
+}
+
+} // namespace UE::MovieScene::Test
+
+IMPLEMENT_SIMPLE_AUTOMATION_TEST(
+		FMovieSceneDecomposerAbsoluteTest, 
+		"System.Engine.Sequencer.Decomposer.Absolute", 
+		EAutomationTestFlags::EditorContext | EAutomationTestFlags::EngineFilter)
+bool FMovieSceneDecomposerAbsoluteTest::RunTest(const FString& Parameters)
+{
+	using namespace UE::MovieScene;
+	using namespace UE::MovieScene::Test;
+
+	UMovieSceneSection* TestObjectFloatSection;
+	UMovieSceneDecomposerTestObject* TestObject = NewObject<UMovieSceneDecomposerTestObject>();
+
+	// Absolute section setting the value at 100. Find what it would take to make it 120. It's simple, it should be 120.
+	FSequenceBuilder()
+		.AddObjectBinding(TestObject)
+		.AddPropertyTrack<UMovieSceneFloatTrack>(GET_MEMBER_NAME_CHECKED(UMovieSceneDecomposerTestObject, FloatProperty))
+			.AddSection(0, 5000)
+				.Assign(TestObjectFloatSection)
+				.AddKey<FMovieSceneFloatChannel, float>(0, 0, 100.f)
+			.Pop()
+		.Pop();
+
+	FMovieSceneTracksComponentTypes* ComponentTypes = FMovieSceneTracksComponentTypes::Get();
+	TRecompositionResult<double> Result = RecomposeBlendOperational(TestObject, TestObjectFloatSection, ComponentTypes->Float, 120.f);
+
+	UTEST_EQUAL("Recomposed values", Result.Values.Num(), 1);
+	UTEST_EQUAL("Value to key", Result.Values[0], 120.0);
+
+	return true;
+}
+
+IMPLEMENT_SIMPLE_AUTOMATION_TEST(
+	FMovieSceneDecomposerAbsoluteWithEasingTest,
+	"System.Engine.Sequencer.Decomposer.AbsoluteWithEasing",
+	EAutomationTestFlags::EditorContext | EAutomationTestFlags::EngineFilter)
+	bool FMovieSceneDecomposerAbsoluteWithEasingTest::RunTest(const FString& Parameters)
+{
+	using namespace UE::MovieScene;
+	using namespace UE::MovieScene::Test;
+
+	UMovieSceneSection* TestObjectFloatSection;
+	UMovieSceneDecomposerTestObject* TestObject = NewObject<UMovieSceneDecomposerTestObject>();
+	FMovieSceneTracksComponentTypes* ComponentTypes = FMovieSceneTracksComponentTypes::Get();
+
+	// We need to set the initial value cache manually for now, because of limitations of the interrogator.
+	TestObject->FloatProperty = 20.f;
+	TSharedPtr<FInitialValueCache> InitialValueCache = FInitialValueCache::GetGlobalInitialValues();
+	InitialValueCache->GetStorage<FFloatPropertyTraits>(ComponentTypes->Float.InitialValue)->AddInitialValue(TestObject, 20.0, offsetof(UMovieSceneDecomposerTestObject, FloatProperty));
+
+	// Absolute section setting the value at 100. But it has linear ease-in and we consider the time at which the easing is at 50%.
+	// This means it brings in 50% of the default value, which we set to 20. The result is (20/2 + 50/2) = 35.
+	FSequenceBuilder()
+		.AddObjectBinding(TestObject)
+		.AddPropertyTrack<UMovieSceneFloatTrack>(GET_MEMBER_NAME_CHECKED(UMovieSceneDecomposerTestObject, FloatProperty))
+			.AddSection(0, 5000)
+				.Assign(TestObjectFloatSection)
+				.SetEaseIn(1000)
+				.AddKey<FMovieSceneFloatChannel, float>(0, 0, 100.f)
+			.Pop()
+		.Pop();
+
+	// Find what it would take to make the result 120. It should be 220, so that (20/2 + 220/2) = 120.
+	TRecompositionResult<double> Result = RecomposeBlendOperational(TestObject, TestObjectFloatSection, ComponentTypes->Float, 120.f, 500);
+
+	UTEST_EQUAL("Recomposed values", Result.Values.Num(), 1);
+	UTEST_EQUAL("Value to key", Result.Values[0], 220.0);
+
+	// Run the same but at time 0, when the weight is 0%. This should hit the special case where we just key the value.
+	Result = RecomposeBlendOperational(TestObject, TestObjectFloatSection, ComponentTypes->Float, 120.f, 0);
+
+	UTEST_EQUAL("Recomposed values", Result.Values.Num(), 1);
+	UTEST_EQUAL("Value to key", Result.Values[0], 120.0);
+
+	return true;
+}
+
+IMPLEMENT_SIMPLE_AUTOMATION_TEST(
+		FMovieSceneDecomposerMultipleAbsolutesTest,
+		"System.Engine.Sequencer.Decomposer.MultipleAbsolutes",
+		EAutomationTestFlags::EditorContext | EAutomationTestFlags::EngineFilter)
+bool FMovieSceneDecomposerMultipleAbsolutesTest::RunTest(const FString& Parameters)
+{
+	using namespace UE::MovieScene;
+	using namespace UE::MovieScene::Test;
+
+	UMovieSceneSection* TestObjectFloatSection1;
+	UMovieSceneSection* TestObjectFloatSection2;
+	UMovieSceneDecomposerTestObject* TestObject = NewObject<UMovieSceneDecomposerTestObject>();
+
+	// One absolute section sets the value to 50, and the other to 150. When they're both at full weight, it gives
+	// a result of 100. We add some ease in to the second section in order to test how weight also factors into it.
+	FSequenceBuilder()
+		.AddObjectBinding(TestObject)
+		.AddPropertyTrack<UMovieSceneFloatTrack>(GET_MEMBER_NAME_CHECKED(UMovieSceneDecomposerTestObject, FloatProperty))
+			.AddSection(0, 5000)
+				.Assign(TestObjectFloatSection1)
+				.AddKey<FMovieSceneFloatChannel, float>(0, 0, 50.f)
+			.Pop()
+			.AddSection(0, 5000)
+				.Assign(TestObjectFloatSection2)
+				.SetEaseIn(1000)
+				.AddKey<FMovieSceneFloatChannel, float>(0, 0, 150.f)
+			.Pop()
+		.Pop();
+
+	// To get 120, we need the second section to be at 190, so that (50 + 190)/2 = 120.
+	FMovieSceneTracksComponentTypes* ComponentTypes = FMovieSceneTracksComponentTypes::Get();
+	TRecompositionResult<double> Result = RecomposeBlendOperational(TestObject, TestObjectFloatSection2, ComponentTypes->Float, 120.f, 2000);
+
+	UTEST_EQUAL("Recomposed values", Result.Values.Num(), 1);
+	UTEST_EQUAL("Value to key", Result.Values[0], 190.0);
+
+	// Now try again at a time when the ease-in on the second section is active and at 50%. To get to 120, we need the
+	// first section to be at 105, so that (105 + 150/2)/1.5 = 120.
+	Result = RecomposeBlendOperational(TestObject, TestObjectFloatSection1, ComponentTypes->Float, 120.f, 500);
+
+	UTEST_EQUAL("Recomposed values", Result.Values.Num(), 1);
+	UTEST_EQUAL("Value to key", Result.Values[0], 105.0);
+
+	// If we recompose the second section, it needs to be at 260, so that (50 + 260/2)/1.5 = 120.
+	Result = RecomposeBlendOperational(TestObject, TestObjectFloatSection2, ComponentTypes->Float, 120.f, 500);
+
+	UTEST_EQUAL("Recomposed values", Result.Values.Num(), 1);
+	UTEST_EQUAL("Value to key", Result.Values[0], 260.0);
+
+	// And now we recompose both at the same time. When the ease-in is at 50%, the result is 83.3333 so getting to
+	// 120 means we are missing +36.6666, to share among both sections.
+	// This sharing is done based on the weights of the sections at the given time, which is 100% and 50% respectively.
+	// So the first section will get two thirds of this, and the second will get one third. Of course, that one third
+	// gets in practice bumped up to two thirds since that second section has 50% weight, which means we need to double
+	// the value to get the desired effect. Each section gets +36.666 to reach the desired value.
+	// Bottom line, we get (86.666 + 186.666/2)/1.5 = 120 (give or take).
+	Result = RecomposeBlendOperational(TestObject, { TestObjectFloatSection1, TestObjectFloatSection2 }, ComponentTypes->Float, 120.f, 500);
+
+	UTEST_EQUAL("Recomposed values", Result.Values.Num(), 2);
+	UTEST_EQUAL_TOLERANCE("Value to key", Result.Values[0], 86.66666667, 0.0001);
+	UTEST_EQUAL_TOLERANCE("Value to key", Result.Values[1], 186.66666667, 0.0001);
+
+	return true;
+}
+
+IMPLEMENT_SIMPLE_AUTOMATION_TEST(
+		FMovieSceneDecomposerAdditiveTest,
+		"System.Engine.Sequencer.Decomposer.Additive", 
+		EAutomationTestFlags::EditorContext | EAutomationTestFlags::EngineFilter)
+bool FMovieSceneDecomposerAdditiveTest::RunTest(const FString& Parameters)
+{
+	using namespace UE::MovieScene;
+	using namespace UE::MovieScene::Test;
+
+	UMovieSceneSection* TestObjectFloatSection1;
+	UMovieSceneSection* TestObjectFloatSection2;
+	UMovieSceneDecomposerTestObject* TestObject = NewObject<UMovieSceneDecomposerTestObject>();
+
+	// Absolute section setting the value to 100, with an additive on top worth +10. Figure out what it would take for the additive
+	// to give us a final value of 120. It should make the additive +20.
+	FSequenceBuilder()
+		.AddObjectBinding(TestObject)
+		.AddPropertyTrack<UMovieSceneFloatTrack>(GET_MEMBER_NAME_CHECKED(UMovieSceneDecomposerTestObject, FloatProperty))
+			.AddSection(0, 5000)
+				.Assign(TestObjectFloatSection1)
+				.AddKey<FMovieSceneFloatChannel, float>(0, 0, 100.f)
+			.Pop()
+			.AddSection(0, 5000, 1)
+				.Assign(TestObjectFloatSection2)
+				.SetBlendType(EMovieSceneBlendType::Additive)
+				.AddKey<FMovieSceneFloatChannel, float>(0, 0, 10.f)
+			.Pop()
+		.Pop();
+
+	FMovieSceneTracksComponentTypes* ComponentTypes = FMovieSceneTracksComponentTypes::Get();
+	TRecompositionResult<double> Result = RecomposeBlendOperational(TestObject, TestObjectFloatSection2, ComponentTypes->Float, 120.f);
+
+	UTEST_EQUAL("Recomposed values", Result.Values.Num(), 1);
+	UTEST_EQUAL("Value to key", Result.Values[0], 20.0);
+
+	// Now let's decompose both sections. It should put the whole difference on the additive section, leaving the
+	// absolute section as-is.
+	Result = RecomposeBlendOperational(TestObject, { TestObjectFloatSection1, TestObjectFloatSection2 }, ComponentTypes->Float, 120.f);
+
+	UTEST_EQUAL("Recomposed values", Result.Values.Num(), 2);
+	UTEST_EQUAL("Value to key", Result.Values[0], 100.0);
+	UTEST_EQUAL("Value to key", Result.Values[1], 20.0);
+
+	return true;
+}
+
+IMPLEMENT_SIMPLE_AUTOMATION_TEST(
+		FMovieSceneDecomposerAdditiveWithEasingTest,
+		"System.Engine.Sequencer.Decomposer.AdditiveWithEasing", 
+		EAutomationTestFlags::EditorContext | EAutomationTestFlags::EngineFilter)
+bool FMovieSceneDecomposerAdditiveWithEasingTest::RunTest(const FString& Parameters)
+{
+	using namespace UE::MovieScene;
+	using namespace UE::MovieScene::Test;
+
+	UMovieSceneSection* TestObjectFloatSection1;
+	UMovieSceneSection* TestObjectFloatSection2;
+	UMovieSceneDecomposerTestObject* TestObject = NewObject<UMovieSceneDecomposerTestObject>();
+
+	// Absolute section setting the value to 100, with an additive on top worth +10. There is ease-in on the additive,
+	// and we'll be considering the time at which the weight is 50% and the additive only contributes +5.
+	FSequenceBuilder()
+		.AddObjectBinding(TestObject)
+		.AddPropertyTrack<UMovieSceneFloatTrack>(GET_MEMBER_NAME_CHECKED(UMovieSceneDecomposerTestObject, FloatProperty))
+			.AddSection(0, 5000)
+				.Assign(TestObjectFloatSection1)
+				.AddKey<FMovieSceneFloatChannel, float>(0, 0, 100.f)
+			.Pop()
+			.AddSection(0, 5000, 1)
+				.Assign(TestObjectFloatSection2)
+				.SetBlendType(EMovieSceneBlendType::Additive)
+				.SetEaseIn(2000)
+				.AddKey<FMovieSceneFloatChannel, float>(0, 0, 10.f)
+			.Pop()
+		.Pop();
+
+	// Figure out what it would take for the additive to give us a final value of 120. It should make the additive +40,
+	// since it only has a weight of 50% at the evaluation time.
+	FMovieSceneTracksComponentTypes* ComponentTypes = FMovieSceneTracksComponentTypes::Get();
+	TRecompositionResult<double> Result = RecomposeBlendOperational(TestObject, TestObjectFloatSection2, ComponentTypes->Float, 120.f, 1000);
+
+	UTEST_EQUAL("Recomposed values", Result.Values.Num(), 1);
+	UTEST_EQUAL("Value to key", Result.Values[0], 40.0);
+
+	// Run the decomposition for both sections at the same time. The absolute section shouldn't change, as we put all
+	// the difference on the additive in this sort of case.
+	Result = RecomposeBlendOperational(TestObject, { TestObjectFloatSection1, TestObjectFloatSection2 }, ComponentTypes->Float, 120.f, 1000);
+
+	UTEST_EQUAL("Recomposed values", Result.Values.Num(), 2);
+	UTEST_EQUAL("Value to key", Result.Values[0], 100.0);
+	UTEST_EQUAL("Value to key", Result.Values[1], 40.0);
+
+	// Run the same but at time 0, where the weight is 0%. This should hit the special case where we just key the
+	// difference in value between the absolute and the desired value.
+	Result = RecomposeBlendOperational(TestObject, TestObjectFloatSection2, ComponentTypes->Float, 120.f, 0);
+
+	UTEST_EQUAL("Recomposed values", Result.Values.Num(), 1);
+	UTEST_EQUAL("Value to key", Result.Values[0], 20.0);
+
+	return true;
+}
+
+IMPLEMENT_SIMPLE_AUTOMATION_TEST(
+		FMovieSceneDecomposerAdditiveWithMultipleEasingTest,
+		"System.Engine.Sequencer.Decomposer.AdditiveWithMultipleEasing", 
+		EAutomationTestFlags::EditorContext | EAutomationTestFlags::EngineFilter)
+bool FMovieSceneDecomposerAdditiveWithMultipleEasingTest::RunTest(const FString& Parameters)
+{
+	using namespace UE::MovieScene;
+	using namespace UE::MovieScene::Test;
+
+	UMovieSceneSection* TestObjectFloatSection;
+	UMovieSceneDecomposerTestObject* TestObject = NewObject<UMovieSceneDecomposerTestObject>();
+
+	// Absolute section setting the value to 100, with an additive on top worth +10. Both sections have ease-in and
+	// ease-out in order to test various combinations and edge cases.
+	// Absolute section: ease-in [0, 1000[, ease-out [4500, 5000[
+	// Additive section: ease-in [0, 500[,  ease-out [4000, 5000[
+	FSequenceBuilder()
+		.AddObjectBinding(TestObject)
+		.AddPropertyTrack<UMovieSceneFloatTrack>(GET_MEMBER_NAME_CHECKED(UMovieSceneDecomposerTestObject, FloatProperty))
+			.AddSection(0, 5000)
+				.SetEaseIn(1000).SetEaseOut(500)
+				.AddKey<FMovieSceneFloatChannel, float>(0, 0, 100.f)
+			.Pop()
+			.AddSection(0, 5000, 1)
+				.Assign(TestObjectFloatSection)
+				.SetBlendType(EMovieSceneBlendType::Additive)
+				.SetEaseIn(500).SetEaseOut(1000)
+				.AddKey<FMovieSceneFloatChannel, float>(0, 0, 10.f)
+			.Pop()
+		.Pop();
+
+	// At time 250, absolute is 25% (25) and additive is 50% (+5), for a total of 30. To get to 50 we need another 20.
+	// The additive being at 50%, this means we need 40, so additive should be +50.
+	FMovieSceneTracksComponentTypes* ComponentTypes = FMovieSceneTracksComponentTypes::Get();
+	TRecompositionResult<double> Result = RecomposeBlendOperational(TestObject, TestObjectFloatSection, ComponentTypes->Float, 50.f, 250);
+
+	UTEST_EQUAL("Recomposed values", Result.Values.Num(), 1);
+	UTEST_EQUAL("Value to key", Result.Values[0], 50.0);
+
+	// At time 750, absolute is 75% (75) and additive is 100% (+10), for a total of 85. To get to 100 we need another
+	// 15. So the additive should be +25.
+	Result = RecomposeBlendOperational(TestObject, TestObjectFloatSection, ComponentTypes->Float, 100.f, 750);
+
+	UTEST_EQUAL("Recomposed values", Result.Values.Num(), 1);
+	UTEST_EQUAL("Value to key", Result.Values[0], 25.0);
+
+	// At time 0, both sections have 0% weight. We should hit the fallback code which just keys the desired value.
+	// It's not a super awesome value in this case but no value is really awesome anyway.
+	Result = RecomposeBlendOperational(TestObject, TestObjectFloatSection, ComponentTypes->Float, 120.f, 0);
+
+	UTEST_EQUAL("Recomposed values", Result.Values.Num(), 1);
+	UTEST_EQUAL("Value to key", Result.Values[0], 120.0);
+
+	return true;
+}
+
+IMPLEMENT_SIMPLE_AUTOMATION_TEST(
+		FMovieSceneDecomposerAdditiveFromBaseTest,
+		"System.Engine.Sequencer.Decomposer.AdditiveFromBase",
+		EAutomationTestFlags::EditorContext | EAutomationTestFlags::EngineFilter)
+bool FMovieSceneDecomposerAdditiveFromBaseTest::RunTest(const FString& Parameters)
+{
+	using namespace UE::MovieScene;
+	using namespace UE::MovieScene::Test;
+
+	UMovieSceneSection* TestObjectFloatSection;
+	UMovieSceneDecomposerTestObject* TestObject = NewObject<UMovieSceneDecomposerTestObject>();
+
+	// Absolute section setting a value of 100, with an additive that goes from +0 to +5 over time using an
+	// additive-from-base blend with a base value of 50. In order to reach a final value of 120, we need to
+	// set the additive to 70, so that (70-50) = +20.
+	FSequenceBuilder()
+		.AddObjectBinding(TestObject)
+		.AddPropertyTrack<UMovieSceneFloatTrack>(GET_MEMBER_NAME_CHECKED(UMovieSceneDecomposerTestObject, FloatProperty))
+			.AddSection(0, 5000)
+				.AddKey<FMovieSceneFloatChannel, float>(0, 0, 100.f)
+			.Pop()
+			.AddSection(0, 5000, 1)
+				.Assign(TestObjectFloatSection)
+				.SetBlendType(EMovieSceneBlendType::AdditiveFromBase)
+				.AddKeys<FMovieSceneFloatChannel, float>(0, { 0, 3000 }, { 50.f, 55.f })
+			.Pop()
+		.Pop();
+
+	FMovieSceneTracksComponentTypes* ComponentTypes = FMovieSceneTracksComponentTypes::Get();
+	TRecompositionResult<double> Result = RecomposeBlendOperational(TestObject, TestObjectFloatSection, ComponentTypes->Float, 120.f, FFrameTime(4000));
+
+	UTEST_EQUAL("Recomposed values", Result.Values.Num(), 1);
+	UTEST_EQUAL("Value to key", Result.Values[0], 70.0);
+
+	return true;
+}
+
+IMPLEMENT_SIMPLE_AUTOMATION_TEST(
+		FMovieSceneDecomposerAdditiveFromBaseWithEasingTest,
+		"System.Engine.Sequencer.Decomposer.AdditiveFromBaseWithEasing",
+		EAutomationTestFlags::EditorContext | EAutomationTestFlags::EngineFilter)
+bool FMovieSceneDecomposerAdditiveFromBaseWithEasingTest::RunTest(const FString& Parameters)
+{
+	using namespace UE::MovieScene;
+	using namespace UE::MovieScene::Test;
+
+	UMovieSceneSection* TestObjectFloatSection;
+	UMovieSceneDecomposerTestObject* TestObject = NewObject<UMovieSceneDecomposerTestObject>();
+	FMovieSceneTracksComponentTypes* ComponentTypes = FMovieSceneTracksComponentTypes::Get();
+
+	// We need to set the initial value cache manually for now, because of limitations of the interrogator.
+	TestObject->FloatProperty = 20.f;
+	TSharedPtr<FInitialValueCache> InitialValueCache = FInitialValueCache::GetGlobalInitialValues();
+	InitialValueCache->GetStorage<FFloatPropertyTraits>(ComponentTypes->Float.InitialValue)->AddInitialValue(TestObject, 20.0, offsetof(UMovieSceneDecomposerTestObject, FloatProperty));
+
+	// Absolute section setting a value of 100, with an additive that goes from +0 to +10 over time using an
+	// additive-from-base blend with a base value of 50.
+	FSequenceBuilder()
+		.AddObjectBinding(TestObject)
+		.AddPropertyTrack<UMovieSceneFloatTrack>(GET_MEMBER_NAME_CHECKED(UMovieSceneDecomposerTestObject, FloatProperty))
+			.AddSection(0, 5000)
+				.SetEaseOut(1000)
+				.AddKey<FMovieSceneFloatChannel, float>(0, 0, 100.f)
+			.Pop()
+			.AddSection(0, 5000, 1)
+				.Assign(TestObjectFloatSection)
+				.SetBlendType(EMovieSceneBlendType::AdditiveFromBase)
+				.SetEaseIn(2000)
+				.AddKeys<FMovieSceneFloatChannel, float>(0, { 0, 1000 }, { 50.f, 60.f })
+			.Pop()
+		.Pop();
+
+	// By the time the additive is at +10, its ease-in keeps it at 50%. If we recompose at that time for a target value
+	// of 120, we need the additive to be at 90, so that (90-50) = +40 x 50% = +20.
+	TRecompositionResult<double> Result = RecomposeBlendOperational(TestObject, TestObjectFloatSection, ComponentTypes->Float, 120.f, FFrameTime(1000));
+
+	UTEST_EQUAL("Recomposed values", Result.Values.Num(), 1);
+	UTEST_EQUAL("Value to key", Result.Values[0], 90.0);
+
+	// Try again, but at a later time, when the additive is at 100% weight, but the absolute is at 50% with its
+	// ease-out. Since the initial value is 20, the full absolute value is (100 x 0.5 + 20 x 0.5) = 60. The additive
+	// needs to add +60 to reach 120, so we should get 110.
+	Result = RecomposeBlendOperational(TestObject, TestObjectFloatSection, ComponentTypes->Float, 120.f, FFrameTime(4500));
+
+	UTEST_EQUAL("Recomposed values", Result.Values.Num(), 1);
+	UTEST_EQUAL("Value to key", Result.Values[0], 110.0);
+
+	return true;
+}
+
+IMPLEMENT_SIMPLE_AUTOMATION_TEST(
+		FMovieSceneDecomposerComplexTest,
+		"System.Engine.Sequencer.Decomposer.Complex",
+		EAutomationTestFlags::EditorContext | EAutomationTestFlags::EngineFilter)
+bool FMovieSceneDecomposerComplexTest::RunTest(const FString& Parameters)
+{
+	using namespace UE::MovieScene;
+	using namespace UE::MovieScene::Test;
+
+	UMovieSceneSection *TestObjectFloatSectionAbsolute1, *TestObjectFloatSectionAbsolute2;
+	UMovieSceneSection *TestObjectFloatSectionAdditive1, *TestObjectFloatSectionAdditive2;
+	UMovieSceneSection *TestObjectFloatSectionAdditiveFromBase1, *TestObjectFloatSectionAdditiveFromBase2;
+	UMovieSceneDecomposerTestObject* TestObject = NewObject<UMovieSceneDecomposerTestObject>();
+	FMovieSceneTracksComponentTypes* ComponentTypes = FMovieSceneTracksComponentTypes::Get();
+
+	// We need to set the initial value cache manually for now, because of limitations of the interrogator.
+	TestObject->FloatProperty = 20.f;
+	TSharedPtr<FInitialValueCache> InitialValueCache = FInitialValueCache::GetGlobalInitialValues();
+	InitialValueCache->GetStorage<FFloatPropertyTraits>(ComponentTypes->Float.InitialValue)->AddInitialValue(TestObject, 20.0, offsetof(UMovieSceneDecomposerTestObject, FloatProperty));
+
+	// Ok let's pull all the stops here. It's like a bad joke: three absolute sections, three additive sections, and
+	// three additive-from-base sections walk into a bar... anyway, we're going to decompose two of each, leaving the
+	// third as non-decomposed. We also add ease-in on all of them, so we can also test with weights.
+	FSequenceBuilder()
+		.AddObjectBinding(TestObject)
+		.AddPropertyTrack<UMovieSceneFloatTrack>(GET_MEMBER_NAME_CHECKED(UMovieSceneDecomposerTestObject, FloatProperty))
+			// Absolute sections are at 50, 60, 70
+			.AddSection(0, 5000)
+				.SetEaseIn(1000)
+				.AddKey<FMovieSceneFloatChannel, float>(0, 0, 50.f)
+			.Pop()
+			.AddSection(0, 5000, 1)
+				.Assign(TestObjectFloatSectionAbsolute1)
+				.SetEaseIn(1000)
+				.AddKey<FMovieSceneFloatChannel, float>(0, 0, 60.f)
+			.Pop()
+			.AddSection(0, 5000, 2)
+				.Assign(TestObjectFloatSectionAbsolute2)
+				.SetEaseIn(1000)
+				.AddKey<FMovieSceneFloatChannel, float>(0, 0, 70.f)
+			.Pop()
+			// Additive sections are at +10, +30, +40
+			.AddSection(0, 5000, 2)
+				.SetBlendType(EMovieSceneBlendType::Additive)
+				.SetEaseIn(1000)
+				.AddKey<FMovieSceneFloatChannel, float>(0, 0, 10.f)
+			.Pop()
+			.AddSection(0, 5000, 3)
+				.Assign(TestObjectFloatSectionAdditive1)
+				.SetBlendType(EMovieSceneBlendType::Additive)
+				.SetEaseIn(1000)
+				.AddKey<FMovieSceneFloatChannel, float>(0, 0, 30.f)
+			.Pop()
+			.AddSection(0, 5000, 4)
+				.Assign(TestObjectFloatSectionAdditive2)
+				.SetBlendType(EMovieSceneBlendType::Additive)
+				.SetEaseIn(1000)
+				.AddKey<FMovieSceneFloatChannel, float>(0, 0, 40.f)
+			.Pop()
+			// Additive-from-base sections are at +5, +10, +20
+			.AddSection(0, 5000, 5)
+				.SetBlendType(EMovieSceneBlendType::AdditiveFromBase)
+				.SetEaseIn(1000)
+				.AddKeys<FMovieSceneFloatChannel, float>(0, { 0, 500 }, { 50.f, 55.f })
+			.Pop()
+			.AddSection(0, 5000, 6)
+				.Assign(TestObjectFloatSectionAdditiveFromBase1)
+				.SetBlendType(EMovieSceneBlendType::AdditiveFromBase)
+				.SetEaseIn(1000)
+				.AddKeys<FMovieSceneFloatChannel, float>(0, { 0, 500 }, { 50.f, 60.f })
+			.Pop()
+			.AddSection(0, 5000, 7)
+				.Assign(TestObjectFloatSectionAdditiveFromBase2)
+				.SetBlendType(EMovieSceneBlendType::AdditiveFromBase)
+				.SetEaseIn(1000)
+				.AddKeys<FMovieSceneFloatChannel, float>(0, { 0, 500 }, { 50.f, 70.f })
+			.Pop()
+		.Pop();
+
+	// Decompose all the values at time 2000, when all weights are 100%. At that time, we have:
+	//
+	// (50 + 60 + 70)/3 + (10 + 30 + 40) + ((55 - 50) + (60 - 50) + (70 - 50)) = 175
+	//
+	// If we want to reach 275, we will leave all the absolutes alone and distribute the difference across the
+	// additives. The absolutes give out (50 + 60 + 70)/3 = 60, so we need +215. We have 4 decomposed additives, and
+	// two that we don't decompose. These two contribute +15, so there's +200 left to redistribute between our additives.
+	// These 4 decomposed additives originally contirbuted +100, which makes it easy to compute their contribution
+	// factors: it's their value in percentage! So Additive1 will get 30% of +200, Additive2 will get 40% of +200, and
+	// so on.
+	TRecompositionResult<double> Result = RecomposeBlendOperational(
+			TestObject, 
+			{ TestObjectFloatSectionAbsolute1, TestObjectFloatSectionAbsolute2,
+			  TestObjectFloatSectionAdditive1, TestObjectFloatSectionAdditive2,
+			  TestObjectFloatSectionAdditiveFromBase1, TestObjectFloatSectionAdditiveFromBase2 },
+			ComponentTypes->Float, 
+			275.f,
+			FFrameTime(1000));
+
+	UTEST_EQUAL("Recomposed values", Result.Values.Num(), 6);
+	UTEST_EQUAL("Value to key (absolute 1)", Result.Values[0], 60.0);
+	UTEST_EQUAL("Value to key (absolute 2)", Result.Values[1], 70.0);
+	UTEST_EQUAL("Value to key (additive 1)", Result.Values[2], 60.0);
+	UTEST_EQUAL("Value to key (additive 2)", Result.Values[3], 80.0);
+	UTEST_EQUAL("Value to key (additive from base 1)", Result.Values[4], 70.0);
+	UTEST_EQUAL("Value to key (additive from base 2)", Result.Values[5], 90.0);
+
+	// Now let's try again at a time when all the weights are at 50%. We have:
+	//
+	// (50/2 + 60/2 + 70/2)/1.5 + (10/2 + 30/2 + 40/2) + ((55 - 50)/2 + (60 - 50)/2 + (70 - 50)/2) = 117.5
+	//
+	// If we want to get 275, the same happens as previously, only the differences applied to the additives get doubled
+	// to account for the half-weights. The absolutes give out 60, and the two non-decomposed additives give +7.5. So
+	// we're missing (275 - 60 - 7.5) = 207.5 which will be proportionally spread out between the decomposed additives
+	// (since, once again, we don't change the decomposed absolutes). Additive1 gets 30% of 207.5, Additive2 gets 40% of
+	// 207.5, and so on. Of course, these values get doubled to account for these additives being at half-weight.
+	Result = RecomposeBlendOperational(
+			TestObject, 
+			{ TestObjectFloatSectionAbsolute1, TestObjectFloatSectionAbsolute2,
+			  TestObjectFloatSectionAdditive1, TestObjectFloatSectionAdditive2,
+			  TestObjectFloatSectionAdditiveFromBase1, TestObjectFloatSectionAdditiveFromBase2 },
+			ComponentTypes->Float, 
+			275.f,
+			FFrameTime(500));
+
+	UTEST_EQUAL("Recomposed values", Result.Values.Num(), 6);
+	UTEST_EQUAL("Value to key (absolute 1)", Result.Values[0], 60.0);
+	UTEST_EQUAL("Value to key (absolute 2)", Result.Values[1], 70.0);
+	UTEST_EQUAL("Value to key (additive 1)", Result.Values[2], 124.5);
+	UTEST_EQUAL("Value to key (additive 2)", Result.Values[3], 166.0);
+	UTEST_EQUAL("Value to key (additive from base 1)", Result.Values[4], 91.5);
+	UTEST_EQUAL("Value to key (additive from base 2)", Result.Values[5], 133.0);
+	
+	return true;
+}
+
+#undef LOCTEXT_NAMESPACE
+
diff --git a/Engine/Source/Runtime/MovieSceneTracks/Private/Tests/MovieSceneDecomposerTests.h b/Engine/Source/Runtime/MovieSceneTracks/Private/Tests/MovieSceneDecomposerTests.h
new file mode 100644
--- /dev/null
+++ b/Engine/Source/Runtime/MovieSceneTracks/Private/Tests/MovieSceneDecomposerTests.h
@@ -0,0 +1,18 @@
+// Copyright Epic Games, Inc. All Rights Reserved.
+
+#pragma once
+
+#include "CoreTypes.h"
+
+#include "MovieSceneDecomposerTests.generated.h"
+
+UCLASS()
+class MOVIESCENETRACKS_API UMovieSceneDecomposerTestObject : public UObject
+{
+	GENERATED_BODY()
+
+public:
+	UPROPERTY()
+	float FloatProperty = 0.f;
+};
+
