Sequencer: When keep state, allow audio to continue and stop on its own when it reaches the end

Note, this could cause audio to continue indefinitely if it's set to Keep State and Looping. This was thought to be a bug in UE-63501

diff --git a/Engine/Source/Runtime/MovieSceneTracks/Private/Evaluation/MovieSceneAudioTemplate.cpp b/Engine/Source/Runtime/MovieSceneTracks/Private/Evaluation/MovieSceneAudioTemplate.cpp
--- a/Engine/Source/Runtime/MovieSceneTracks/Private/Evaluation/MovieSceneAudioTemplate.cpp
+++ b/Engine/Source/Runtime/MovieSceneTracks/Private/Evaluation/MovieSceneAudioTemplate.cpp
@@ -23,7 +23,6 @@
 
 
 DECLARE_CYCLE_STAT(TEXT("Audio Track Evaluate"), MovieSceneEval_AudioTrack_Evaluate, STATGROUP_MovieSceneEval);
-DECLARE_CYCLE_STAT(TEXT("Audio Track Tear Down"), MovieSceneEval_AudioTrack_TearDown, STATGROUP_MovieSceneEval);
 DECLARE_CYCLE_STAT(TEXT("Audio Track Token Execute"), MovieSceneEval_AudioTrack_TokenExecute, STATGROUP_MovieSceneEval);
 
 static float MaxSequenceAudioDesyncToleranceCVar = 0.5f;
@@ -249,17 +248,6 @@ struct FCachedAudioTrackData : IPersistentEvaluationData
 			}
 		}
 	}
-
-	void StopSoundsOnSection(FObjectKey ObjectKey)
-	{
-		for (TPair<FObjectKey, TMap<FObjectKey, TWeakObjectPtr<UAudioComponent>>>& Pair : AudioComponentsByActorKey)
-		{
-			if (UAudioComponent* AudioComponent = Pair.Value.FindRef(ObjectKey).Get())
-			{
-				AudioComponent->Stop();
-			}
-		}
-	}
 };
 
 
@@ -619,17 +607,4 @@ void FMovieSceneAudioSectionTemplate::Evaluate(const FMovieSceneEvaluationOperan
 	{
 		ExecutionTokens.Add(FAudioSectionExecutionToken(AudioSection));
 	}
-}
-
-void FMovieSceneAudioSectionTemplate::TearDown(FPersistentEvaluationData& PersistentData, IMovieScenePlayer& Player) const
-{
-	MOVIESCENE_DETAILED_SCOPE_CYCLE_COUNTER(MovieSceneEval_AudioTrack_TearDown)
-
-	if (GEngine && GEngine->UseSound())
-	{
-		FCachedAudioTrackData& TrackData = PersistentData.GetOrAddTrackData<FCachedAudioTrackData>();
-
-		TrackData.StopSoundsOnSection(AudioSection);
-	}
-}
-
+}
\ No newline at end of file
diff --git a/Engine/Source/Runtime/MovieSceneTracks/Private/Evaluation/MovieSceneAudioTemplate.h b/Engine/Source/Runtime/MovieSceneTracks/Private/Evaluation/MovieSceneAudioTemplate.h
--- a/Engine/Source/Runtime/MovieSceneTracks/Private/Evaluation/MovieSceneAudioTemplate.h
+++ b/Engine/Source/Runtime/MovieSceneTracks/Private/Evaluation/MovieSceneAudioTemplate.h
@@ -29,9 +29,4 @@ private:
 
 	virtual UScriptStruct& GetScriptStructImpl() const override { return *StaticStruct(); }
 	virtual void Evaluate(const FMovieSceneEvaluationOperand& Operand, const FMovieSceneContext& Context, const FPersistentEvaluationData& PersistentData, FMovieSceneExecutionTokens& ExecutionTokens) const override;
-	virtual void SetupOverrides() override
-	{
-		EnableOverrides(RequiresTearDownFlag);
-	}
-	virtual void TearDown(FPersistentEvaluationData& PersistentData, IMovieScenePlayer& Player) const override;
 };
