Sequencer: Fixed invalid entity IDs remaining cached in material parameters and blendchannel/entity leak on garbage collection

diff --git a/Engine/Source/Runtime/MovieScene/Public/EntitySystem/MovieSceneOverlappingEntityTracker.h b/Engine/Source/Runtime/MovieScene/Public/EntitySystem/MovieSceneOverlappingEntityTracker.h
--- a/Engine/Source/Runtime/MovieScene/Public/EntitySystem/MovieSceneOverlappingEntityTracker.h
+++ b/Engine/Source/Runtime/MovieScene/Public/EntitySystem/MovieSceneOverlappingEntityTracker.h
@@ -74,6 +74,10 @@ struct TGarbageTraitsImpl<TIntegerSequence<int, Indices...>, T...>
 	}
 };
 
+/** Override for OutputType* in order to provide custom garbage collection logic */
+inline void CollectGarbageForOutput(void*)
+{
+}
 
 template<typename... T>
 struct TOverlappingEntityInput
@@ -548,6 +552,8 @@ struct TOverlappingEntityTracker_WithGarbage : TOverlappingEntityTrackerImpl<Out
 			FOutput& Output = this->Outputs[Index];
 			if (KeyType::GarbageTraits::IsGarbage(Output.Key))
 			{
+				CollectGarbageForOutput(&this->Outputs[Index].OutputData);
+
 				this->Outputs.RemoveAt(Index, 1);
 
 				for (auto It = this->OutputToEntity.CreateKeyIterator(OutputIndex); It; ++It)
diff --git a/Engine/Source/Runtime/MovieSceneTracks/Private/Systems/MovieSceneMaterialParameterSystem.cpp b/Engine/Source/Runtime/MovieSceneTracks/Private/Systems/MovieSceneMaterialParameterSystem.cpp
--- a/Engine/Source/Runtime/MovieSceneTracks/Private/Systems/MovieSceneMaterialParameterSystem.cpp
+++ b/Engine/Source/Runtime/MovieSceneTracks/Private/Systems/MovieSceneMaterialParameterSystem.cpp
@@ -24,6 +24,36 @@ FAutoConsoleVariableRef CVarMaterialParameterBlending(
 	ECVF_Default
 );
 
+void CollectGarbageForOutput(FAnimatedMaterialParameterInfo* Output)
+{
+	// This should only happen during garbage collection
+	if (Output->OutputEntityID.IsValid())
+	{
+		UMovieSceneEntitySystemLinker* Linker = Output->WeakLinker.Get();
+		if (Linker)
+		{
+			Linker->EntityManager.AddComponent(Output->OutputEntityID, FBuiltInComponentTypes::Get()->Tags.NeedsUnlink);
+		}
+
+		Output->OutputEntityID = FMovieSceneEntityID();
+	}
+
+	if (Output->BlendChannelID.IsValid())
+	{
+		UMovieSceneBlenderSystem* BlenderSystem = Output->WeakBlenderSystem.Get();
+		if (BlenderSystem)
+		{
+			BlenderSystem->ReleaseBlendChannel(Output->BlendChannelID);
+		}
+		Output->BlendChannelID = FMovieSceneBlendChannelID();
+	}
+}
+
+FAnimatedMaterialParameterInfo::~FAnimatedMaterialParameterInfo()
+{
+	ensureAlways(!OutputEntityID.IsValid() && !BlendChannelID.IsValid());
+}
+
 /** Apply scalar material parameters */
 struct FApplyScalarParameters
 {
@@ -119,6 +149,9 @@ struct FOverlappingMaterialParameterHandler
 					Linker->SystemGraph.AddReference(System, System->DoubleBlenderSystem);
 				}
 
+				Output->WeakLinker = Linker;
+				Output->WeakBlenderSystem = System->DoubleBlenderSystem;
+
 				// Initialize the blend channel ID
 				Output->BlendChannelID = System->DoubleBlenderSystem->AllocateBlendChannel();
 
@@ -265,16 +298,10 @@ void UMovieSceneMaterialParameterSystem::OnInstantiation()
 
 	if (GMaterialParameterBlending)
 	{
-		auto HandleUnlinkedAllocation = [this](const FEntityAllocation* Allocation, TRead<UObject*> InObjects, TReadOneOf<FName, FName, FName> ScalarVectorOrColorParameterNames)
+		auto HandleUnlinkedAllocation = [this](const FEntityAllocation* Allocation)
 		{
-			if (TComponentPtr<const FName> ScalarParameterNames = TComponentPtr<const FName>(ScalarVectorOrColorParameterNames.Get<0>()))
-			{
-				this->ScalarParameterTracker.VisitUnlinkedAllocation(Allocation);
-			}
-			else
-			{
-				this->VectorParameterTracker.VisitUnlinkedAllocation(Allocation);
-			}
+			this->ScalarParameterTracker.VisitUnlinkedAllocation(Allocation);
+			this->VectorParameterTracker.VisitUnlinkedAllocation(Allocation);
 		};
 
 		auto HandleUpdatedAllocation = [this](const FEntityAllocation* Allocation, TRead<UObject*> InObjects, TReadOneOf<FName, FName, FName> ScalarVectorOrColorParameterNames)
@@ -293,7 +320,7 @@ void UMovieSceneMaterialParameterSystem::OnInstantiation()
 			}
 		};
 
-		// Next handle any new or updated bound materials
+		// First step handle any new or updated bound materials
 		FEntityTaskBuilder()
 		.Read(TracksComponents->BoundMaterial)
 		.ReadOneOf(TracksComponents->ScalarParameterName, TracksComponents->VectorParameterName, TracksComponents->ColorParameterName)
@@ -301,10 +328,8 @@ void UMovieSceneMaterialParameterSystem::OnInstantiation()
 		.FilterNone({ BuiltInComponents->Tags.NeedsUnlink })
 		.Iterate_PerAllocation(&Linker->EntityManager, HandleUpdatedAllocation);
 
-		// First step, handle any entities that are going away
+		// Next handle any entities that are going away
 		FEntityTaskBuilder()
-		.Read(TracksComponents->BoundMaterial)
-		.ReadOneOf(TracksComponents->ScalarParameterName, TracksComponents->VectorParameterName, TracksComponents->ColorParameterName)
 		.FilterAll({ BuiltInComponents->Tags.NeedsUnlink })
 		.Iterate_PerAllocation(&Linker->EntityManager, HandleUnlinkedAllocation);
 
diff --git a/Engine/Source/Runtime/MovieSceneTracks/Public/Systems/MovieSceneMaterialParameterSystem.h b/Engine/Source/Runtime/MovieSceneTracks/Public/Systems/MovieSceneMaterialParameterSystem.h
--- a/Engine/Source/Runtime/MovieSceneTracks/Public/Systems/MovieSceneMaterialParameterSystem.h
+++ b/Engine/Source/Runtime/MovieSceneTracks/Public/Systems/MovieSceneMaterialParameterSystem.h
@@ -8,6 +8,7 @@
 
 #include "MovieSceneMaterialParameterSystem.generated.h"
 
+class UMovieSceneBlenderSystem;
 class UMovieScenePiecewiseDoubleBlenderSystem;
 
 namespace UE::MovieScene
@@ -15,11 +16,19 @@ namespace UE::MovieScene
 
 struct FAnimatedMaterialParameterInfo
 {
+	/** Weak linker ptr - only assigned if the output entity is ever allocated */
+	TWeakObjectPtr<UMovieSceneEntitySystemLinker> WeakLinker;
+	/** Weak blender system ptr - only assigned if the blend channel is ever allocated */
+	TWeakObjectPtr<UMovieSceneBlenderSystem> WeakBlenderSystem;
 	int32 NumContributors = 0;
 	FMovieSceneEntityID OutputEntityID;
 	FMovieSceneBlendChannelID BlendChannelID;
+
+	~FAnimatedMaterialParameterInfo();
 };
 
+MOVIESCENETRACKS_API void CollectGarbageForOutput(FAnimatedMaterialParameterInfo* Output);
+
 } // namespace UE::MovieScene
 
 
