Disable ensure by default since it can report false-positives

diff --git a/Engine/Source/Runtime/MovieSceneTracks/Private/Systems/MovieSceneMaterialParameterSystem.cpp b/Engine/Source/Runtime/MovieSceneTracks/Private/Systems/MovieSceneMaterialParameterSystem.cpp
--- a/Engine/Source/Runtime/MovieSceneTracks/Private/Systems/MovieSceneMaterialParameterSystem.cpp
+++ b/Engine/Source/Runtime/MovieSceneTracks/Private/Systems/MovieSceneMaterialParameterSystem.cpp
@@ -24,6 +24,15 @@ FAutoConsoleVariableRef CVarMaterialParameterBlending(
 	ECVF_Default
 );
 
+bool GMaterialParameterEntityLifetimeTracking = false;
+FAutoConsoleVariableRef CVarMaterialParameterEntityLifetimeTracking(
+	TEXT("Sequencer.MaterialParameterEntityLifetimeTracking"),
+	GMaterialParameterEntityLifetimeTracking,
+	TEXT("(Default: false) Ensure on destruction that all entities have been cleaned up. This can report false positives (when the linker and material system are both cleaned up together) so is not enabled by default.\n"),
+	ECVF_Default
+);
+
+
 void CollectGarbageForOutput(FAnimatedMaterialParameterInfo* Output)
 {
 	// This should only happen during garbage collection
@@ -51,7 +60,10 @@ void CollectGarbageForOutput(FAnimatedMaterialParameterInfo* Output)
 
 FAnimatedMaterialParameterInfo::~FAnimatedMaterialParameterInfo()
 {
-	ensureAlways(!OutputEntityID.IsValid() && !BlendChannelID.IsValid());
+	if (GMaterialParameterEntityLifetimeTracking)
+	{
+		ensureAlways(!OutputEntityID.IsValid() && !BlendChannelID.IsValid());
+	}
 }
 
 /** Apply scalar material parameters */
