CIS Fix for https://horde.devtools.epicgames.com/job/63617b308b0183c909bcfb79?step=cca1&issue=231083
Manual merged CL was missing unshelved fix.

diff --git a/Engine/Source/Runtime/MovieSceneTracks/Private/Sections/MovieSceneSkeletalAnimationSection.cpp b/Engine/Source/Runtime/MovieSceneTracks/Private/Sections/MovieSceneSkeletalAnimationSection.cpp
--- a/Engine/Source/Runtime/MovieSceneTracks/Private/Sections/MovieSceneSkeletalAnimationSection.cpp
+++ b/Engine/Source/Runtime/MovieSceneTracks/Private/Sections/MovieSceneSkeletalAnimationSection.cpp
@@ -455,12 +455,13 @@ int32 UMovieSceneSkeletalAnimationSection::SetBoneIndexForRootMotionCalculations
 		TempRootBoneIndex.Reset();
 		return 0;
 	}
-	else if(UAnimSequence* AnimSequence = Cast<UAnimSequence>(Params.Animation) )
-	{ 
-		if(TempRootBoneIndex.IsSet() == false || TempRootBoneIndex.GetValue() == INDEX_NONE)
+	else if (UAnimSequence* AnimSequence = Cast<UAnimSequence>(Params.Animation))
+	{
+		if (TempRootBoneIndex.IsSet() == false || TempRootBoneIndex.GetValue() == INDEX_NONE)
 		{
 			//but if not first find first
 			int32 RootIndex = INDEX_NONE;
+#if WITH_EDITOR
 			const IAnimationDataModel* DataModel = AnimSequence->GetDataModel();
 			const TArray<FBoneAnimationTrack>& BoneAnimationTracks = DataModel->GetBoneAnimationTracks();
 			const int32 NumTracks = BoneAnimationTracks.Num();
@@ -471,7 +472,7 @@ int32 UMovieSceneSkeletalAnimationSection::SetBoneIndexForRootMotionCalculations
 				const int32 BoneTreeIndex = AnimationTrack.BoneTreeIndex;
 				if (BoneTreeIndex != INDEX_NONE)
 				{
-					
+
 					int32 ParentIndex = AnimSequence->GetSkeleton()->GetReferenceSkeleton().GetParentIndex(BoneTreeIndex);
 					if (ParentIndex == INDEX_NONE)
 					{
@@ -507,9 +508,53 @@ int32 UMovieSceneSkeletalAnimationSection::SetBoneIndexForRootMotionCalculations
 							break;
 						}
 					}
-				
+
+				}
+			}
+
+#else
+			const TArray<FTrackToSkeletonMap>& BoneMappings = AnimSequence->GetCompressedTrackToSkeletonMapTable();
+			const int32 NumTracks = BoneMappings.Num();
+			for (int32 TrackIndex = 0; TrackIndex < NumTracks; ++TrackIndex)
+			{
+				const FTrackToSkeletonMap& Mapping = BoneMappings[TrackIndex];
+				// verify if this bone exists in skeleton
+				const int32 BoneTreeIndex = Mapping.BoneTreeIndex;
+				if (BoneTreeIndex != INDEX_NONE)
+				{
+					int32 ParentIndex = AnimSequence->GetSkeleton()->GetReferenceSkeleton().GetParentIndex(BoneTreeIndex);
+					if (ParentIndex == INDEX_NONE)
+					{
+						RootIndex = TrackIndex;
+					}
+					else if (ParentIndex == RootIndex)
+					{
+						FTransform Transform;
+						const int32 NumFrames = AnimSequence->GetNumberOfSampledKeys();
+						for (int32 Index = 0; Index < NumFrames; ++Index)
+						{
+							float Pos = FMath::Clamp((float)AnimSequence->GetSamplingFrameRate().AsSeconds(Index), 0.f, AnimSequence->GetPlayLength());
+PRAGMA_DISABLE_DEPRECATION_WARNINGS
+							//jurre this is what I am talking  about
+							// can't create tkhis since GetRetargetTransformsSourceName isn't  public.
+							//FAnimSequenceDecompressionContext DecompContext(AnimSequence->GetSamplingFrameRate(), AnimSequence->GetSamplingFrameRate().AsFrameNumber(AnimSequence->GetPlayLength()).Value, AnimSequence->Interpolation, AnimSequence->GetRetargetTransformsSourceName(), *AnimSequence->CompressedData.CompressedDataStructure, AnimSequence->GetSkeleton()->GetRefLocalPoses(), AnimSequence->CompressedData.CompressedTrackToSkeletonMapTable, AnimSequence->GetSkeleton(), AnimSequence->IsValidAdditive());
+
+							AnimSequence->GetBoneTransform(Transform, TrackIndex, Pos, false);
+PRAGMA_ENABLE_DEPRECATION_WARNINGS
+							if (Transform.Equals(FTransform::Identity) == false)
+							{
+								TempRootBoneIndex = BoneTreeIndex;
+								break;
+							}
+						}
+						if (TempRootBoneIndex.IsSet())
+						{
+							break;
+						}
+					}
 				}
 			}
+#endif
 		}
 	}
 	return TempRootBoneIndex.GetValue();
