Sequencer: Hierarchical bias is now constrained to operate only within the specific root sequence

Previously we would gather and compute winning hierarchical bias for every blend channel globally - this meant if you had 2 sequences simultaneously animating the same object, one of them could get ignored by hbias even if it resides in a completely different root sequence. We still collect hbias by channel ID, but now we compute a winning hbias for each root sequence to ensure that a higher bias from one sequence doesn't affect the others.

diff --git a/Engine/Source/Runtime/MovieSceneTracks/Private/Systems/MovieSceneHierarchicalBiasSystem.cpp b/Engine/Source/Runtime/MovieSceneTracks/Private/Systems/MovieSceneHierarchicalBiasSystem.cpp
--- a/Engine/Source/Runtime/MovieSceneTracks/Private/Systems/MovieSceneHierarchicalBiasSystem.cpp
+++ b/Engine/Source/Runtime/MovieSceneTracks/Private/Systems/MovieSceneHierarchicalBiasSystem.cpp
@@ -14,15 +14,29 @@ namespace UE
 namespace MovieScene
 {
 
+struct FBlendChannelSequenceKey
+{
+	FRootInstanceHandle RootInstance;
+	FMovieSceneBlendChannelID BlendChannelID;
+	friend uint32 GetTypeHash(const FBlendChannelSequenceKey& In)
+	{
+		return HashCombine(GetTypeHash(In.RootInstance), GetTypeHash(In.BlendChannelID));
+	}
+	friend bool operator==(const FBlendChannelSequenceKey& A, const FBlendChannelSequenceKey& B)
+	{
+		return A.RootInstance == B.RootInstance && A.BlendChannelID == B.BlendChannelID;
+	}
+};
+
 struct FHierarchicalBiasTask
 {
 	explicit FHierarchicalBiasTask(UMovieSceneEntitySystemLinker* InLinker)
 		: Linker(InLinker)
 	{}
 
-	void InitializeChannel(FMovieSceneBlendChannelID BlendChannel)
+	void InitializeChannel(FRootInstanceHandle RootInstanceHandle, FMovieSceneBlendChannelID BlendChannel)
 	{
-		MaxBiasByChannel.FindOrAdd(BlendChannel, MIN_int16);
+		MaxBiasByChannel.FindOrAdd(FBlendChannelSequenceKey{ RootInstanceHandle, BlendChannel }, MIN_int16);
 	}
 
 	bool HasAnyWork() const
@@ -30,21 +44,21 @@ struct FHierarchicalBiasTask
 		return MaxBiasByChannel.Num() != 0;
 	}
 
-	void ForEachAllocation(const FEntityAllocation* Allocation, TRead<FMovieSceneEntityID> EntityIDs, TRead<FMovieSceneBlendChannelID> BlendChannels, TReadOptional<int16> OptHBiases)
+	void ForEachAllocation(const FEntityAllocation* Allocation, TRead<FMovieSceneEntityID> EntityIDs, TRead<FRootInstanceHandle> RootInstanceHandles, TRead<FMovieSceneBlendChannelID> BlendChannels, TReadOptional<int16> OptHBiases)
 	{
 		const int32 Num = Allocation->Num();
 		if (OptHBiases)
 		{
 			for (int32 Index = 0; Index < Num; ++Index)
 			{
-				VisitChannel(EntityIDs[Index], BlendChannels[Index], OptHBiases[Index]);
+				VisitChannel(EntityIDs[Index], RootInstanceHandles[Index], BlendChannels[Index], OptHBiases[Index]);
 			}
 		}
 		else
 		{
 			for (int32 Index = 0; Index < Num; ++Index)
 			{
-				VisitChannel(EntityIDs[Index], BlendChannels[Index], 0);
+				VisitChannel(EntityIDs[Index], RootInstanceHandles[Index], BlendChannels[Index], 0);
 			}
 		}
 	}
@@ -66,38 +80,40 @@ struct FHierarchicalBiasTask
 
 private:
 
-	void VisitChannel(FMovieSceneEntityID EntityID, FMovieSceneBlendChannelID BlendChannel, int16 HBias)
+	void VisitChannel(FMovieSceneEntityID EntityID, FRootInstanceHandle RootInstanceHandle, FMovieSceneBlendChannelID BlendChannel, int16 HBias)
 	{
+		FBlendChannelSequenceKey Key{ RootInstanceHandle, BlendChannel };
+
 		// If this channel hasn't changed at all (ie InitializeChannel was not called for it) do nothing
-		if (int16* ExistingBias = MaxBiasByChannel.Find(BlendChannel))
+		if (int16* ExistingBias = MaxBiasByChannel.Find(Key))
 		{
 			if (HBias > *ExistingBias)
 			{
-				for (auto It = ActiveContributorsByChannel.CreateKeyIterator(BlendChannel); It; ++It)
+				for (auto It = ActiveContributorsByChannel.CreateKeyIterator(Key); It; ++It)
 				{
-					InactiveContributorsByChannel.Add(BlendChannel, It.Value());
+					InactiveContributorsByChannel.Add(Key, It.Value());
 					It.RemoveCurrent();
 				}
 
 				*ExistingBias = HBias;
-				ActiveContributorsByChannel.Add(BlendChannel, EntityID);
+				ActiveContributorsByChannel.Add(Key, EntityID);
 			}
 			else if (HBias == *ExistingBias)
 			{
-				ActiveContributorsByChannel.Add(BlendChannel, EntityID);
+				ActiveContributorsByChannel.Add(Key, EntityID);
 			}
 			else
 			{
-				InactiveContributorsByChannel.Add(BlendChannel, EntityID);
+				InactiveContributorsByChannel.Add(Key, EntityID);
 			}
 		}
 	}
 
-	TMap<FMovieSceneBlendChannelID, int16> MaxBiasByChannel;
+	TMap<FBlendChannelSequenceKey, int16> MaxBiasByChannel;
 
-	TMultiMap<FMovieSceneBlendChannelID, FMovieSceneEntityID> InactiveContributorsByChannel;
+	TMultiMap<FBlendChannelSequenceKey, FMovieSceneEntityID> InactiveContributorsByChannel;
 
-	TMultiMap<FMovieSceneBlendChannelID, FMovieSceneEntityID> ActiveContributorsByChannel;
+	TMultiMap<FBlendChannelSequenceKey, FMovieSceneEntityID> ActiveContributorsByChannel;
 
 	UMovieSceneEntitySystemLinker* Linker;
 };
@@ -137,14 +153,16 @@ void UMovieSceneHierarchicalBiasSystem::OnRun(FSystemTaskPrerequisites& InPrereq
 
 	// First, add all the channels that have changed to the map
 	FEntityTaskBuilder()
+	.Read(Components->RootInstanceHandle)
 	.Read(Components->BlendChannelInput)
 	.FilterAny({ Components->Tags.NeedsLink, Components->Tags.NeedsUnlink })
-	.Iterate_PerEntity(&Linker->EntityManager, [&Task](FMovieSceneBlendChannelID BlendChannel){ Task.InitializeChannel(BlendChannel); });
+	.Iterate_PerEntity(&Linker->EntityManager, [&Task](FRootInstanceHandle RootInstanceHandle, FMovieSceneBlendChannelID BlendChannel){ Task.InitializeChannel(RootInstanceHandle, BlendChannel); });
 
 	if (Task.HasAnyWork())
 	{
 		FEntityTaskBuilder()
 		.ReadEntityIDs()
+		.Read(Components->RootInstanceHandle)
 		.Read(Components->BlendChannelInput)
 		.ReadOptional(Components->HierarchicalBias)
 		.FilterNone({ Components->Tags.NeedsUnlink })
