Sequencer: Added the ability to assign console variable collections to console variable tracks

diff --git a/Engine/Plugins/Editor/ConsoleVariablesEditor/Source/ConsoleVariablesEditorRuntime/ConsoleVariablesEditorRuntime.build.cs b/Engine/Plugins/Editor/ConsoleVariablesEditor/Source/ConsoleVariablesEditorRuntime/ConsoleVariablesEditorRuntime.build.cs
--- a/Engine/Plugins/Editor/ConsoleVariablesEditor/Source/ConsoleVariablesEditorRuntime/ConsoleVariablesEditorRuntime.build.cs
+++ b/Engine/Plugins/Editor/ConsoleVariablesEditor/Source/ConsoleVariablesEditorRuntime/ConsoleVariablesEditorRuntime.build.cs
@@ -13,7 +13,8 @@ public class ConsoleVariablesEditorRuntime : ModuleRules
 			new string[]
 			{
 				"Core",
-				"Engine"
+				"Engine",
+				"MovieSceneTracks",
 			}
 		);
 			
diff --git a/Engine/Plugins/Editor/ConsoleVariablesEditor/Source/ConsoleVariablesEditorRuntime/Private/ConsoleVariablesAsset.cpp b/Engine/Plugins/Editor/ConsoleVariablesEditor/Source/ConsoleVariablesEditorRuntime/Private/ConsoleVariablesAsset.cpp
--- a/Engine/Plugins/Editor/ConsoleVariablesEditor/Source/ConsoleVariablesEditorRuntime/Private/ConsoleVariablesAsset.cpp
+++ b/Engine/Plugins/Editor/ConsoleVariablesEditor/Source/ConsoleVariablesEditorRuntime/Private/ConsoleVariablesAsset.cpp
@@ -1,6 +1,7 @@
 // Copyright Epic Games, Inc. All Rights Reserved.
 
 #include "ConsoleVariablesAsset.h"
+#include "Templates/Tuple.h"
 
 #include "ConsoleVariablesEditorRuntimeLog.h"
 #include "Engine/Engine.h"
@@ -36,6 +37,19 @@ FString UConsoleVariablesAsset::GetSavedCommandsAsCommaSeparatedString(bool bOnl
 	return FString::Join(GetSavedCommandsAsStringArray(bOnlyIncludeChecked), TEXT(","));
 }
 
+void UConsoleVariablesAsset::GetConsoleVariablesForTrack(bool bOnlyIncludeChecked, TArray<TTuple<FString, FString>>& OutVariables)
+{
+	for (const FConsoleVariablesEditorAssetSaveData& Command : SavedCommands)
+	{
+		if (bOnlyIncludeChecked && Command.CheckedState != ECheckBoxState::Checked)
+		{
+			continue;
+		}
+		
+		OutVariables.Add(MakeTuple(Command.CommandName, Command.CommandValueAsString));
+	}
+}
+
 void UConsoleVariablesAsset::ExecuteSavedCommands(UObject* WorldContextObject, bool bOnlyIncludeChecked) const
 {
 	if (const TObjectPtr<UWorld> World = GEngine->GetWorldFromContextObject(WorldContextObject, EGetWorldErrorMode::LogAndReturnNull))
diff --git a/Engine/Plugins/Editor/ConsoleVariablesEditor/Source/ConsoleVariablesEditorRuntime/Public/ConsoleVariablesAsset.h b/Engine/Plugins/Editor/ConsoleVariablesEditor/Source/ConsoleVariablesEditorRuntime/Public/ConsoleVariablesAsset.h
--- a/Engine/Plugins/Editor/ConsoleVariablesEditor/Source/ConsoleVariablesEditorRuntime/Public/ConsoleVariablesAsset.h
+++ b/Engine/Plugins/Editor/ConsoleVariablesEditor/Source/ConsoleVariablesEditorRuntime/Public/ConsoleVariablesAsset.h
@@ -4,6 +4,7 @@
 
 #include "Containers/UnrealString.h"
 #include "Styling/SlateTypes.h"
+#include "Sections/MovieSceneConsoleVariableTrackInterface.h"
 
 #include "ConsoleVariablesAsset.generated.h"
 
@@ -31,7 +32,7 @@ struct FConsoleVariablesEditorAssetSaveData
 
 /** An asset used to track collections of console variables that can be recalled and edited using the Console Variables Editor. */
 UCLASS(BlueprintType)
-class CONSOLEVARIABLESEDITORRUNTIME_API UConsoleVariablesAsset : public UObject
+class CONSOLEVARIABLESEDITORRUNTIME_API UConsoleVariablesAsset : public UObject, public IMovieSceneConsoleVariableTrackInterface
 {
 	GENERATED_BODY()
 public:
@@ -100,7 +101,12 @@ public:
 	void CopyFrom(const UConsoleVariablesAsset* InAssetToCopy);
 	
 private:
-	
+
+	//~ IMovieSceneConsoleVariableTrackInterface overrides
+	virtual void GetConsoleVariablesForTrack(bool bOnlyIncludeChecked, TArray<TPair<FString, FString>>& OutVariables) override;
+
+private:
+
 	/* User defined description of the variable collection */
 	UPROPERTY(AssetRegistrySearchable, BlueprintGetter = "GetVariableCollectionDescription", EditAnywhere, Category = "Console Variables Editor")
 	FString VariableCollectionDescription;
diff --git a/Engine/Source/Runtime/MovieSceneTracks/Private/Sections/MovieSceneCVarSection.cpp b/Engine/Source/Runtime/MovieSceneTracks/Private/Sections/MovieSceneCVarSection.cpp
--- a/Engine/Source/Runtime/MovieSceneTracks/Private/Sections/MovieSceneCVarSection.cpp
+++ b/Engine/Source/Runtime/MovieSceneTracks/Private/Sections/MovieSceneCVarSection.cpp
@@ -44,20 +44,25 @@ void FMovieSceneCVarOverrides::SetFromString(const FString& InString)
 	}
 }
 
-FString FMovieSceneCVarOverrides::GetString() const
+void FMovieSceneCVarOverrides::GetString(TStringBuilder<256>& OutBuilder) const
 {
-	TStringBuilder<256> CVarString;
 	for (const TPair<FString, FString>& Pair : ValuesByCVar)
 	{
-		if (CVarString.Len() != 0)
+		if (OutBuilder.Len() != 0)
 		{
-			CVarString += TEXT(", ");
+			OutBuilder += TEXT(", ");
 		}
 
-		CVarString += Pair.Key;
-		CVarString += TEXT(" ");
-		CVarString += Pair.Value;
+		OutBuilder += Pair.Key;
+		OutBuilder += TEXT(" ");
+		OutBuilder += Pair.Value;
 	}
+}
+
+FString FMovieSceneCVarOverrides::GetString() const
+{
+	TStringBuilder<256> CVarString;
+	GetString(CVarString);
 	return CVarString.ToString();
 }
 
@@ -89,6 +94,21 @@ void UMovieSceneCVarSection::SetFromString(const FString& InString)
 
 FString UMovieSceneCVarSection::GetString() const
 {
-	return ConsoleVariables.GetString();
+	TStringBuilder<256> ResultBuilder;
+
+	for (const FMovieSceneConsoleVariableCollection& Collection : ConsoleVariableCollections)
+	{
+		if (UObject* Object = Collection.Interface.GetObject())
+		{
+			if (ResultBuilder.Len() != 0)
+			{
+				ResultBuilder += TEXT(", ");
+			}
+			ResultBuilder += Object->GetName();
+		}
+	}
+
+	ConsoleVariables.GetString(ResultBuilder);
+	return ResultBuilder.ToString();
 }
 
diff --git a/Engine/Source/Runtime/MovieSceneTracks/Private/TrackInstances/MovieSceneCVarTrackInstance.cpp b/Engine/Source/Runtime/MovieSceneTracks/Private/TrackInstances/MovieSceneCVarTrackInstance.cpp
--- a/Engine/Source/Runtime/MovieSceneTracks/Private/TrackInstances/MovieSceneCVarTrackInstance.cpp
+++ b/Engine/Source/Runtime/MovieSceneTracks/Private/TrackInstances/MovieSceneCVarTrackInstance.cpp
@@ -2,6 +2,7 @@
 
 #include "TrackInstances/MovieSceneCVarTrackInstance.h"
 #include "Sections/MovieSceneCVarSection.h"
+#include "Sections/MovieSceneConsoleVariableTrackInterface.h"
 #include "Tracks/MovieSceneCVarTrack.h"
 #include "EntitySystem/MovieSceneEntitySystemLinker.h"
 #include "EntitySystem/MovieSceneInstanceRegistry.h"
@@ -77,7 +78,17 @@ void UMovieSceneCVarTrackInstance::OnInputAdded(const FMovieSceneTrackInstanceIn
 
 		TSharedPtr<FPreAnimatedCVarStorage> CVarStorage = Linker->PreAnimatedState.GetOrCreateStorage<FPreAnimatedCVarStorage>();
 
-		for (const TPair<FString, FString>& Pair : Section->ConsoleVariables.ValuesByCVar)
+		// Gather all CVars
+		TArray<TPair<FString, FString>> CommandsAndValues = Section->ConsoleVariables.ValuesByCVar.Array();
+		for (FMovieSceneConsoleVariableCollection& Collection : Section->ConsoleVariableCollections)
+		{
+			if (Collection.Interface)
+			{
+				Collection.Interface->GetConsoleVariablesForTrack(Collection.bOnlyIncludeChecked, CommandsAndValues);
+			}
+		}
+
+		for (const TPair<FString, FString>& Pair : CommandsAndValues)
 		{
 			CVarStorage->SavePreAnimatedState(UngroupedHandle, Pair.Key, FPreAnimatedCVarTraits::CachePreAnimatedValue);
 		}
@@ -119,7 +130,18 @@ void UMovieSceneCVarTrackInstance::OnEndUpdateInputs()
 		const int32 HierarchicalBias = InstanceRegistry->GetInstance(Input.InstanceHandle).GetContext().GetHierarchicalBias();
 
 		UMovieSceneCVarSection* Section = CastChecked<UMovieSceneCVarSection>(Input.Section);
-		for (const TPair<FString, FString>& Pair : Section->ConsoleVariables.ValuesByCVar)
+
+		// Gather all CVars
+		TArray<TPair<FString, FString>> CommandsAndValues = Section->ConsoleVariables.ValuesByCVar.Array();
+		for (FMovieSceneConsoleVariableCollection& Collection : Section->ConsoleVariableCollections)
+		{
+			if (Collection.Interface)
+			{
+				Collection.Interface->GetConsoleVariablesForTrack(Collection.bOnlyIncludeChecked, CommandsAndValues);
+			}
+		}
+
+		for (const TPair<FString, FString>& Pair : CommandsAndValues)
 		{
 			// If CVarsNeedingUpdate is empty that implies that a section was changed so we need to reapply everything
 			if (CVarsNeedingUpdate.Num() == 0 || CVarsNeedingUpdate.Contains(Pair.Key))
diff --git a/Engine/Source/Runtime/MovieSceneTracks/Public/Sections/MovieSceneCVarSection.h b/Engine/Source/Runtime/MovieSceneTracks/Public/Sections/MovieSceneCVarSection.h
--- a/Engine/Source/Runtime/MovieSceneTracks/Public/Sections/MovieSceneCVarSection.h
+++ b/Engine/Source/Runtime/MovieSceneTracks/Public/Sections/MovieSceneCVarSection.h
@@ -5,6 +5,7 @@
 #include "Containers/Map.h"
 #include "Containers/UnrealString.h"
 #include "CoreMinimal.h"
+#include "Containers/StringFwd.h"
 #include "EntitySystem/IMovieSceneEntityProvider.h"
 #include "EntitySystem/MovieSceneEntityIDs.h"
 #include "MovieSceneSection.h"
@@ -13,10 +14,26 @@
 
 #include "MovieSceneCVarSection.generated.h"
 
+class IMovieSceneConsoleVariableTrackInterface;
 class UMovieSceneEntitySystemLinker;
 class UObject;
 struct FFrame;
 
+
+USTRUCT()
+struct FMovieSceneConsoleVariableCollection
+{
+	GENERATED_BODY()
+
+	/** Array of console variable preset assets that this track should operate on */
+	UPROPERTY(EditAnywhere, DisplayName="Collection Asset", Category="Collection")
+	TScriptInterface<IMovieSceneConsoleVariableTrackInterface> Interface;
+
+	/** Whether to include checked consolve variables or not */
+	UPROPERTY(EditAnywhere, Category="Collection")
+	bool bOnlyIncludeChecked = false;
+};
+
 USTRUCT()
 struct FMovieSceneCVarOverrides
 {
@@ -24,6 +41,7 @@ struct FMovieSceneCVarOverrides
 
 	MOVIESCENETRACKS_API void SetFromString(const FString& InString);
 
+	MOVIESCENETRACKS_API void GetString(TStringBuilder<256>& OutBuilder) const;
 	MOVIESCENETRACKS_API FString GetString() const;
 
 	/** The name of the console variable and the value, separated by ' ' or '=', ie: "foo.bar=1" or "foo.bar 1". */
@@ -57,9 +75,11 @@ private:
 
 public:
 
+	/** Array of console variable preset assets that this track should operate on */
+	UPROPERTY(EditAnywhere, Category="Console Variable Collections")
+	TArray<FMovieSceneConsoleVariableCollection> ConsoleVariableCollections;
+
 	/** The name of the console variable and the value, separated by ' ' or '=', ie: "foo.bar=1" or "foo.bar 1". */
-	UPROPERTY(EditAnywhere,  meta=(MultiLine=true), Category="Console Variables")
+	UPROPERTY(EditAnywhere, Category="Console Variables")
 	FMovieSceneCVarOverrides ConsoleVariables;
-
-	friend class UMovieSceneCVarTrackInstance;
 };
diff --git a/Engine/Source/Runtime/MovieSceneTracks/Public/Sections/MovieSceneConsoleVariableTrackInterface.h b/Engine/Source/Runtime/MovieSceneTracks/Public/Sections/MovieSceneConsoleVariableTrackInterface.h
new file mode 100644
--- /dev/null
+++ b/Engine/Source/Runtime/MovieSceneTracks/Public/Sections/MovieSceneConsoleVariableTrackInterface.h
@@ -0,0 +1,27 @@
+// Copyright Epic Games, Inc. All Rights Reserved.
+
+#pragma once
+
+#include "Containers/ContainersFwd.h"
+#include "UObject/ObjectMacros.h"
+#include "UObject/UObjectGlobals.h"
+#include "UObject/Interface.h"
+
+#include "MovieSceneConsoleVariableTrackInterface.generated.h"
+
+template<typename...> struct TTuple;
+class FString;
+
+UINTERFACE(MinimalAPI)
+class UMovieSceneConsoleVariableTrackInterface : public UInterface
+{
+	GENERATED_BODY()
+};
+
+class IMovieSceneConsoleVariableTrackInterface
+{
+public:
+	GENERATED_BODY()
+
+	virtual void GetConsoleVariablesForTrack(bool bOnlyIncludeChecked, TArray<TTuple<FString, FString>>& OutVariables) = 0;
+};
