Sequencer: Reworked hierarchical weights and easing to support easing on any sub section, and dynamic weights at runtime

The presence of Hierarchical Easing is no longer computed on sequence compilation, and instead is detected automatically by the weight and easing evaluator system.

New and expired HierarchicalEasingProvider components will now trigger the allocation of an easing channel for a given sequence ID, and this will also be propagated to future entities, as well as added to existing entities by way of a new mutation method on the entity ledger. This allows sub-sequence easing to simply create the provider entity without any of its inner entities needing to be concerned with adding easing channels. It also allows us to dynamically create providers for runtime weighting and have it work exactly the same as static easing.

A new flag exists on the sequence playback settings as well as the asset itself to indicate that a sequence will have dynamic weighting. This allows the system to correctly cache initial values for blending into and out of a sequence without taking that CPU and memory hit for every sequence.

Added initial value caching to material parameters, allowing them to be blended into and out of, whilst also allowing additive material parameter tracks

diff --git a/Engine/Plugins/MovieScene/TemplateSequence/Source/TemplateSequence/Private/Section/TemplateSequenceSection.cpp b/Engine/Plugins/MovieScene/TemplateSequence/Source/TemplateSequence/Private/Section/TemplateSequenceSection.cpp
--- a/Engine/Plugins/MovieScene/TemplateSequence/Source/TemplateSequence/Private/Section/TemplateSequenceSection.cpp
+++ b/Engine/Plugins/MovieScene/TemplateSequence/Source/TemplateSequence/Private/Section/TemplateSequenceSection.cpp
@@ -104,7 +104,7 @@ bool UTemplateSequenceSection::PopulateEvaluationFieldImpl(const TRange<FFrameNu
 	return true;
 }
 
-void UTemplateSequenceSection::ImportEntityImpl(UMovieSceneEntitySystemLinker* EntityLinker, const FEntityImportParams& Params, FImportedEntity* OutImportedEntity)
+void UTemplateSequenceSection::ImportEntityImpl(UMovieSceneEntitySystemLinker* EntityLinker, const UE::MovieScene::FEntityImportParams& Params, UE::MovieScene::FImportedEntity* OutImportedEntity)
 {
 	using namespace UE::MovieScene;
 
diff --git a/Engine/Plugins/MovieScene/TemplateSequence/Source/TemplateSequence/Public/Sections/TemplateSequenceSection.h b/Engine/Plugins/MovieScene/TemplateSequence/Source/TemplateSequence/Public/Sections/TemplateSequenceSection.h
--- a/Engine/Plugins/MovieScene/TemplateSequence/Source/TemplateSequence/Public/Sections/TemplateSequenceSection.h
+++ b/Engine/Plugins/MovieScene/TemplateSequence/Source/TemplateSequence/Public/Sections/TemplateSequenceSection.h
@@ -53,7 +53,6 @@ struct FTemplateSectionPropertyScale
 UCLASS()
 class TEMPLATESEQUENCE_API UTemplateSequenceSection 
 	: public UMovieSceneSubSection
-	, public IMovieSceneEntityProvider
 {
 public:
 
@@ -73,7 +72,7 @@ protected:
 	virtual EMovieSceneChannelProxyType CacheChannelProxy() override;
 
 	virtual bool PopulateEvaluationFieldImpl(const TRange<FFrameNumber>& EffectiveRange, const FMovieSceneEvaluationFieldEntityMetaData& InMetaData, FMovieSceneEntityComponentFieldBuilder* OutFieldBuilder) override;
-	virtual void ImportEntityImpl(UMovieSceneEntitySystemLinker* EntityLinker, const FEntityImportParams& Params, FImportedEntity* OutImportedEntity) override;
+	virtual void ImportEntityImpl(UMovieSceneEntitySystemLinker* EntityLinker, const UE::MovieScene::FEntityImportParams& Params, UE::MovieScene::FImportedEntity* OutImportedEntity) override;
 
 public:
 
diff --git a/Engine/Source/Editor/Sequencer/Private/SSequencer.cpp b/Engine/Source/Editor/Sequencer/Private/SSequencer.cpp
--- a/Engine/Source/Editor/Sequencer/Private/SSequencer.cpp
+++ b/Engine/Source/Editor/Sequencer/Private/SSequencer.cpp
@@ -2425,6 +2425,7 @@ TSharedRef<SWidget> SSequencer::MakePlaybackMenu()
 		}
 
 		MenuBuilder.AddMenuEntry( FSequencerCommands::Get().ToggleAsyncEvaluation );
+		MenuBuilder.AddMenuEntry( FSequencerCommands::Get().ToggleDynamicWeighting );
 	}
 	MenuBuilder.EndSection();
 
diff --git a/Engine/Source/Editor/Sequencer/Private/Sequencer.cpp b/Engine/Source/Editor/Sequencer/Private/Sequencer.cpp
--- a/Engine/Source/Editor/Sequencer/Private/Sequencer.cpp
+++ b/Engine/Source/Editor/Sequencer/Private/Sequencer.cpp
@@ -11045,6 +11045,12 @@ void FSequencer::BindCommands()
 		FCanExecuteAction::CreateLambda( []{ return true; } ),
 		FIsActionChecked::CreateLambda( [this]{ return this->UsesAsyncEvaluation(); } ) );
 
+	SequencerCommandBindings->MapAction(
+		Commands.ToggleDynamicWeighting,
+		FExecuteAction::CreateLambda( [this]{ this->ToggleDynamicWeighting(); } ),
+		FCanExecuteAction::CreateLambda( []{ return true; } ),
+		FIsActionChecked::CreateLambda( [this]{ return this->UsesDynamicWeighting(); } ) );
+
 	SequencerCommandBindings->MapAction(
 		Commands.ToggleKeepCursorInPlaybackRangeWhileScrubbing,
 		FExecuteAction::CreateLambda([this] { Settings->SetKeepCursorInPlayRangeWhileScrubbing(!Settings->ShouldKeepCursorInPlayRangeWhileScrubbing()); }),
@@ -11745,5 +11751,23 @@ bool FSequencer::UsesAsyncEvaluation()
 	return !EnumHasAnyFlags(GetRootMovieSceneSequence()->GetFlags(), EMovieSceneSequenceFlags::BlockingEvaluation);
 }
 
+void FSequencer::ToggleDynamicWeighting()
+{
+	UMovieSceneSequence* Sequence = GetFocusedMovieSceneSequence();
+
+	EMovieSceneSequenceFlags NewFlags = Sequence->GetFlags();
+	NewFlags ^= EMovieSceneSequenceFlags::DynamicWeighting;
+
+	FScopedTransaction Transaction(EnumHasAnyFlags(NewFlags, EMovieSceneSequenceFlags::DynamicWeighting) ? LOCTEXT("DisableDynamicWeighting", "Disable Dynamic Weighting") : LOCTEXT("EnableDynamicWeighting", "Enable Dynamic Weighting"));
+
+	Sequence->Modify();
+	Sequence->SetSequenceFlags(NewFlags);
+}
+
+bool FSequencer::UsesDynamicWeighting()
+{
+	return EnumHasAnyFlags(GetFocusedMovieSceneSequence()->GetFlags(), EMovieSceneSequenceFlags::DynamicWeighting);
+}
+
 #undef LOCTEXT_NAMESPACE
 
diff --git a/Engine/Source/Editor/Sequencer/Private/Sequencer.h b/Engine/Source/Editor/Sequencer/Private/Sequencer.h
--- a/Engine/Source/Editor/Sequencer/Private/Sequencer.h
+++ b/Engine/Source/Editor/Sequencer/Private/Sequencer.h
@@ -1106,6 +1106,9 @@ private:
 	void ToggleAsyncEvaluation();
 	bool UsesAsyncEvaluation();
 
+	void ToggleDynamicWeighting();
+	bool UsesDynamicWeighting();
+
 	void UpdateCachedPlaybackContextAndClient();
 
 	void UpdateCachedCameraActors();
diff --git a/Engine/Source/Editor/Sequencer/Private/SequencerCommands.cpp b/Engine/Source/Editor/Sequencer/Private/SequencerCommands.cpp
--- a/Engine/Source/Editor/Sequencer/Private/SequencerCommands.cpp
+++ b/Engine/Source/Editor/Sequencer/Private/SequencerCommands.cpp
@@ -52,6 +52,7 @@ void FSequencerCommands::RegisterCommands()
 	UI_COMMAND( ToggleCleanPlaybackMode, "Game View (Clean Playback Mode)", "Enable game view and hide viewport buttons while playing.", EUserInterfaceActionType::ToggleButton, FInputChord() );
 	UI_COMMAND( ToggleRerunConstructionScripts, "Rerun Construction Scripts", "Rerun construction scripts on bound actors every frame.", EUserInterfaceActionType::ToggleButton, FInputChord() );
 	UI_COMMAND( ToggleAsyncEvaluation, "Async Evaluation", "When enabled, enables a single asynchronous evaluation once per-frame. When disabled, forces a full blocking evaluation every time this sequence is evaluated (should be avoided for real-time content).", EUserInterfaceActionType::ToggleButton, FInputChord() );
+	UI_COMMAND( ToggleDynamicWeighting, "Dynamic Weighting", "When enabled, all blendable tracks will cache their initial values to ensure that they are able to correctly blend in/out when dynamic weights are being used.", EUserInterfaceActionType::ToggleButton, FInputChord() );
 
 	UI_COMMAND( ToggleKeepCursorInPlaybackRangeWhileScrubbing, "Keep Playhead in Playback Range While Scrubbing", "When checked, the playhead will be constrained to the current playback range while scrubbing", EUserInterfaceActionType::ToggleButton, FInputChord() );
 	UI_COMMAND( ToggleKeepPlaybackRangeInSectionBounds, "Keep Playback Range in Section Bounds", "When checked, the playback range will be synchronized to the section bounds", EUserInterfaceActionType::ToggleButton, FInputChord() );
diff --git a/Engine/Source/Editor/Sequencer/Public/SequencerCommands.h b/Engine/Source/Editor/Sequencer/Public/SequencerCommands.h
--- a/Engine/Source/Editor/Sequencer/Public/SequencerCommands.h
+++ b/Engine/Source/Editor/Sequencer/Public/SequencerCommands.h
@@ -142,6 +142,9 @@ public:
 	/** When enabled, enables a single asynchronous evaluation once per-frame. When disabled, forces a full blocking evaluation every time this sequence is evaluated (should be avoided for real-time content). */
 	TSharedPtr< FUICommandInfo > ToggleAsyncEvaluation;
 
+	/** When enabled, all blendable tracks will always cache their initial values to ensure that they are able to correctly blend in/out when dynamic weights are being used. */
+	TSharedPtr< FUICommandInfo > ToggleDynamicWeighting;
+
 	/** Toggle constraining the time cursor to the playback range while scrubbing */
 	TSharedPtr< FUICommandInfo > ToggleKeepCursorInPlaybackRangeWhileScrubbing;
 
diff --git a/Engine/Source/Runtime/MovieScene/Private/Compilation/MovieSceneCompiledDataManager.cpp b/Engine/Source/Runtime/MovieScene/Private/Compilation/MovieSceneCompiledDataManager.cpp
--- a/Engine/Source/Runtime/MovieScene/Private/Compilation/MovieSceneCompiledDataManager.cpp
+++ b/Engine/Source/Runtime/MovieScene/Private/Compilation/MovieSceneCompiledDataManager.cpp
@@ -142,7 +142,6 @@ struct FGatherParameters
 		, LocalClampRange(RootClampRange)
 		, Flags(ESectionEvaluationFlags::None)
 		, HierarchicalBias(0)
-		, bHasHierarchicalEasing(false)
 	{}
 
 	FGatherParameters CreateForSubData(const FMovieSceneSubSequenceData& SubData, FMovieSceneSequenceID InSubSequenceID) const
@@ -156,7 +155,6 @@ struct FGatherParameters
 
 		SubParams.RootToSequenceTransform   = SubData.RootToSequenceTransform;
 		SubParams.HierarchicalBias          = SubData.HierarchicalBias;
-		SubParams.bHasHierarchicalEasing    = SubData.bHasHierarchicalEasing;
 		SubParams.SequenceID                = InSubSequenceID;
 		SubParams.RootToSequenceWarpCounter = WarpCounter;
 
@@ -197,9 +195,6 @@ struct FGatherParameters
 	/** Current accumulated hierarchical bias */
 	int16 HierarchicalBias;
 
-	/** Whether the current sequence is receiving hierarchical easing from some parent sequence */
-	bool bHasHierarchicalEasing;
-
 	EMovieSceneServerClientMask NetworkMask;
 };
 
@@ -1463,7 +1458,6 @@ bool UMovieSceneCompiledDataManager::GenerateSubSequenceData(UMovieSceneSubTrack
 		NewSubData.PlayRange               = TRange<FFrameNumber>::Intersection(InnerClampRange, NewSubData.PlayRange.Value);
 		NewSubData.RootToSequenceTransform = NewSubData.RootToSequenceTransform * Params.RootToSequenceTransform;
 		NewSubData.HierarchicalBias        = Params.HierarchicalBias + NewSubData.HierarchicalBias;
-		NewSubData.bHasHierarchicalEasing  = Params.bHasHierarchicalEasing || NewSubData.bHasHierarchicalEasing;
 
 		// Add the sub data to the root hierarchy
 		InOutHierarchy->Add(NewSubData, InnerSequenceID, ParentSequenceID);
diff --git a/Engine/Source/Runtime/MovieScene/Private/EntitySystem/BuiltInComponentTypes.cpp b/Engine/Source/Runtime/MovieScene/Private/EntitySystem/BuiltInComponentTypes.cpp
--- a/Engine/Source/Runtime/MovieScene/Private/EntitySystem/BuiltInComponentTypes.cpp
+++ b/Engine/Source/Runtime/MovieScene/Private/EntitySystem/BuiltInComponentTypes.cpp
@@ -26,6 +26,7 @@ FBuiltInComponentTypes::FBuiltInComponentTypes()
 	FComponentRegistry* ComponentRegistry = UMovieSceneEntitySystemLinker::GetComponents();
 
 	ComponentRegistry->NewComponentType(&ParentEntity,          TEXT("Parent Entity"));
+	ComponentRegistry->NewComponentType(&SequenceID,            TEXT("Sequence ID"));
 	ComponentRegistry->NewComponentType(&InstanceHandle,        TEXT("Instance Handle"));
 	ComponentRegistry->NewComponentType(&RootInstanceHandle,    TEXT("Root Instance Handle"));
 
@@ -83,6 +84,7 @@ FBuiltInComponentTypes::FBuiltInComponentTypes()
 	ComponentRegistry->NewComponentType(&WeightChannelFlags,      TEXT("Weight Channel Flags"));
 
 	ComponentRegistry->NewComponentType(&Easing,                  TEXT("Easing"));
+	ComponentRegistry->NewComponentType(&EasingResult,            TEXT("Easing Result"));
 	ComponentRegistry->NewComponentType(&HierarchicalEasingChannel, TEXT("Hierarchical Easing Channel"));
 	ComponentRegistry->NewComponentType(&HierarchicalEasingProvider, TEXT("Hierarchical Easing Provider"));
 
@@ -151,6 +153,7 @@ FBuiltInComponentTypes::FBuiltInComponentTypes()
 	Tags.FixedTime               = ComponentRegistry->NewTag(TEXT("Fixed Time"));
 	Tags.PreRoll                 = ComponentRegistry->NewTag(TEXT("Pre Roll"));
 	Tags.SectionPreRoll          = ComponentRegistry->NewTag(TEXT("Section Pre Roll"));
+	Tags.AlwaysCacheInitialValue = ComponentRegistry->NewTag(TEXT("Always Cache Initial Value"));
 
 	SymbolicTags.CreatesEntities = ComponentRegistry->NewTag(TEXT("~~ SYMBOLIC ~~ Creates Entities"));
 
@@ -168,12 +171,14 @@ FBuiltInComponentTypes::FBuiltInComponentTypes()
 	ComponentRegistry->Factories.DefineChildComponent(Tags.FixedTime,     Tags.FixedTime);
 	ComponentRegistry->Factories.DefineChildComponent(Tags.PreRoll,       Tags.PreRoll);
 	ComponentRegistry->Factories.DefineChildComponent(Tags.SectionPreRoll,Tags.SectionPreRoll);
+	ComponentRegistry->Factories.DefineChildComponent(Tags.AlwaysCacheInitialValue,Tags.AlwaysCacheInitialValue);
 
 	ComponentRegistry->Factories.DefineMutuallyInclusiveComponent(Tags.SectionPreRoll,Tags.PreRoll);
 
 	ComponentRegistry->Factories.DuplicateChildComponent(EvalTime);
 	ComponentRegistry->Factories.DuplicateChildComponent(BaseValueEvalTime);
 
+	ComponentRegistry->Factories.DuplicateChildComponent(SequenceID);
 	ComponentRegistry->Factories.DuplicateChildComponent(InstanceHandle);
 	ComponentRegistry->Factories.DuplicateChildComponent(RootInstanceHandle);
 	ComponentRegistry->Factories.DuplicateChildComponent(PropertyBinding);
@@ -283,6 +288,7 @@ FBuiltInComponentTypes::FBuiltInComponentTypes()
 
 		// Easing needs a time to evaluate
 		ComponentRegistry->Factories.DefineMutuallyInclusiveComponent(Easing, EvalTime);
+		ComponentRegistry->Factories.DefineMutuallyInclusiveComponent(Easing, EasingResult);
 	}
 
 	// Weight channel relationships
@@ -299,8 +305,8 @@ FBuiltInComponentTypes::FBuiltInComponentTypes()
 
 	// Weight and easing result component relationship
 	{
-		ComponentRegistry->Factories.DefineMutuallyInclusiveComponent(Easing, WeightAndEasingResult);
 		ComponentRegistry->Factories.DefineMutuallyInclusiveComponent(HierarchicalEasingChannel, WeightAndEasingResult);
+		ComponentRegistry->Factories.DefineMutuallyInclusiveComponent(EasingResult, WeightAndEasingResult);
 		ComponentRegistry->Factories.DefineMutuallyInclusiveComponent(WeightResult, WeightAndEasingResult);
 	}
 
diff --git a/Engine/Source/Runtime/MovieScene/Private/EntitySystem/IMovieSceneEntityProvider.cpp b/Engine/Source/Runtime/MovieScene/Private/EntitySystem/IMovieSceneEntityProvider.cpp
--- a/Engine/Source/Runtime/MovieScene/Private/EntitySystem/IMovieSceneEntityProvider.cpp
+++ b/Engine/Source/Runtime/MovieScene/Private/EntitySystem/IMovieSceneEntityProvider.cpp
@@ -30,6 +30,8 @@ FMovieSceneEntityID FImportedEntity::Manufacture(const FEntityImportParams& Para
 	auto BaseBuilder = FEntityBuilder()
 	.AddTag(Components->Tags.NeedsLink)
 	.AddTag(Components->Tags.ImportedEntity)
+	.AddTagConditional(Components->Tags.AlwaysCacheInitialValue, Params.Sequence.bDynamicWeighting)
+	.AddConditional(Components->SequenceID, Params.Sequence.SequenceID, Params.Sequence.SequenceID != MovieSceneSequenceID::Root)
 	.AddConditional(Components->RootInstanceHandle, Params.Sequence.RootInstanceHandle, Params.Sequence.RootInstanceHandle.IsValid())
 	.AddConditional(Components->InstanceHandle, Params.Sequence.InstanceHandle, Params.Sequence.InstanceHandle.IsValid());
 
diff --git a/Engine/Source/Runtime/MovieScene/Private/EntitySystem/MovieSceneEntityFactory.cpp b/Engine/Source/Runtime/MovieScene/Private/EntitySystem/MovieSceneEntityFactory.cpp
--- a/Engine/Source/Runtime/MovieScene/Private/EntitySystem/MovieSceneEntityFactory.cpp
+++ b/Engine/Source/Runtime/MovieScene/Private/EntitySystem/MovieSceneEntityFactory.cpp
@@ -375,25 +375,30 @@ int32 FEntityFactories::ComputeMutuallyInclusiveComponents(FComponentMask& Compo
 	return NumNewComponents;
 }
 
-void FEntityFactories::RunInitializers(const FComponentMask& ParentType, const FComponentMask& ChildType, const FEntityAllocation* ParentAllocation, TArrayView<const int32> ParentAllocationOffsets, const FEntityRange& InChildEntityRange)
+void FEntityFactories::RunInitializers(const FComponentMask& Type, const FEntityRange& InEntityRange)
 {
-	// First off, run child initializers
-	for (TInlineValue<FChildEntityInitializer>& ChildInit : ChildInitializers)
+	for (TInlineValue<FMutualEntityInitializer>& MutualInit : MutualInitializers)
 	{
-		if (ChildInit->IsRelevant(ParentType, ChildType))
+		if (MutualInit->IsRelevant(Type))
 		{
-			ChildInit->Run(InChildEntityRange, ParentAllocation, ParentAllocationOffsets);
+			MutualInit->Run(InEntityRange);
 		}
 	}
+}
 
+void FEntityFactories::RunInitializers(const FComponentMask& ParentType, const FComponentMask& ChildType, const FEntityAllocation* ParentAllocation, TArrayView<const int32> ParentAllocationOffsets, const FEntityRange& InChildEntityRange)
+{
 	// First off, run child initializers
-	for (TInlineValue<FMutualEntityInitializer>& MutualInit : MutualInitializers)
+	for (TInlineValue<FChildEntityInitializer>& ChildInit : ChildInitializers)
 	{
-		if (MutualInit->IsRelevant(ChildType))
+		if (ChildInit->IsRelevant(ParentType, ChildType))
 		{
-			MutualInit->Run(InChildEntityRange);
+			ChildInit->Run(InChildEntityRange, ParentAllocation, ParentAllocationOffsets);
 		}
 	}
+
+	// Run generic initializers
+	RunInitializers(ChildType, InChildEntityRange);
 }
 
 }	// using namespace MovieScene
diff --git a/Engine/Source/Runtime/MovieScene/Private/EntitySystem/MovieSceneEntityLedger.cpp b/Engine/Source/Runtime/MovieScene/Private/EntitySystem/MovieSceneEntityLedger.cpp
--- a/Engine/Source/Runtime/MovieScene/Private/EntitySystem/MovieSceneEntityLedger.cpp
+++ b/Engine/Source/Runtime/MovieScene/Private/EntitySystem/MovieSceneEntityLedger.cpp
@@ -4,6 +4,7 @@
 #include "EntitySystem/MovieSceneInstanceRegistry.h"
 #include "EntitySystem/MovieSceneEntitySystemLinker.h"
 #include "EntitySystem/BuiltInComponentTypes.h"
+#include "EntitySystem/MovieSceneEntityMutations.h"
 
 #include "Evaluation/MovieSceneEvaluationField.h"
 
@@ -249,5 +250,72 @@ void FEntityLedger::TagGarbage(UMovieSceneEntitySystemLinker* Linker)
 	}
 }
 
+bool FEntityLedger::Contains(UMovieSceneEntitySystemLinker* Linker, const FEntityComponentFilter& Filter) const
+{
+	bool bResult = false;
+
+	auto Visit = [&Filter, &bResult, Linker](FMovieSceneEntityID EntityID)
+	{
+		bResult = Filter.Match(Linker->EntityManager.GetEntityType(EntityID));
+	};
+
+	for (FMovieSceneEntityID EntityID : OneShotEntities)
+	{
+		Visit(EntityID);
+		Linker->EntityManager.IterateChildren_ParentFirst(EntityID, Visit);
+
+		if (bResult)
+		{
+			return true;
+		}
+	}
+
+	for (const TPair<FMovieSceneEvaluationFieldEntityKey, FImportedEntityData>& Pair : ImportedEntities)
+	{
+		Visit(Pair.Value.EntityID);
+		Linker->EntityManager.IterateChildren_ParentFirst(Pair.Value.EntityID, Visit);
+
+		if (bResult)
+		{
+			return true;
+		}
+	}
+
+	return bResult;
+}
+
+void FEntityLedger::MutateAll(UMovieSceneEntitySystemLinker* Linker, const FEntityComponentFilter& Filter, const IMovieScenePerEntityMutation& Mutation) const
+{
+	auto Visit = [&Filter, &Mutation, Linker](FMovieSceneEntityID EntityID)
+	{
+		const FComponentMask& ExistingType = Linker->EntityManager.GetEntityType(EntityID);
+		if (Filter.Match(ExistingType))
+		{
+			FComponentMask NewType = ExistingType;
+			Mutation.CreateMutation(&Linker->EntityManager, &NewType);
+
+			if (!NewType.CompareSetBits(ExistingType))
+			{
+				Linker->EntityManager.ChangeEntityType(EntityID, NewType);
+
+				FEntityInfo EntityInfo = Linker->EntityManager.GetEntity(EntityID);
+				Mutation.InitializeEntities(EntityInfo.Data.AsRange(), NewType);
+			}
+		}
+	};
+
+	for (FMovieSceneEntityID EntityID : OneShotEntities)
+	{
+		Visit(EntityID);
+		Linker->EntityManager.IterateChildren_ParentFirst(EntityID, Visit);
+	}
+
+	for (const TPair<FMovieSceneEvaluationFieldEntityKey, FImportedEntityData>& Pair : ImportedEntities)
+	{
+		Visit(Pair.Value.EntityID);
+		Linker->EntityManager.IterateChildren_ParentFirst(Pair.Value.EntityID, Visit);
+	}
+}
+
 } // namespace MovieScene
 } // namespace UE
diff --git a/Engine/Source/Runtime/MovieScene/Private/EntitySystem/MovieSceneInstanceRegistry.cpp b/Engine/Source/Runtime/MovieScene/Private/EntitySystem/MovieSceneInstanceRegistry.cpp
--- a/Engine/Source/Runtime/MovieScene/Private/EntitySystem/MovieSceneInstanceRegistry.cpp
+++ b/Engine/Source/Runtime/MovieScene/Private/EntitySystem/MovieSceneInstanceRegistry.cpp
@@ -68,26 +68,15 @@ FRootInstanceHandle FInstanceRegistry::AllocateRootInstance(IMovieScenePlayer* P
 	return InstanceHandle;
 }
 
-FInstanceHandle FInstanceRegistry::AllocateSubInstance(IMovieScenePlayer* Player, FMovieSceneSequenceID SequenceID, FRootInstanceHandle RootInstanceHandle)
+FInstanceHandle FInstanceRegistry::AllocateSubInstance(IMovieScenePlayer* Player, FMovieSceneSequenceID SequenceID, FRootInstanceHandle RootInstanceHandle, FInstanceHandle ParentInstanceHandle)
 {
-	check(Instances.Num() < 65535 && SequenceID != MovieSceneSequenceID::Root);
-
-	const FMovieSceneRootEvaluationTemplateInstance& Template = Player->GetEvaluationTemplate();
-	const FMovieSceneSequenceHierarchy*              Hierarchy = Template.GetCompiledDataManager()->FindHierarchy(Template.GetCompiledDataID());
-
-	checkf(Hierarchy, TEXT("Attempting to construct a new sub sequence instance without a hierarchy"));
-
-	const FMovieSceneSubSequenceData* SubData = Hierarchy->FindSubData(SequenceID);
-	checkf(SubData, TEXT("Attempting to construct a new sub sequence instance with a sub sequence ID that does not exist in the hierarchy"));
-
-	FMovieSceneCompiledDataID CompiledDataID = Template.GetCompiledDataManager()->GetDataID(SubData->GetSequence());
-
+	check(Instances.Num() < 65535 && SequenceID != MovieSceneSequenceID::Root && ParentInstanceHandle.IsValid());
 
 	const uint16 InstanceSerial = InstanceSerialNumber++;
 	FSparseArrayAllocationInfo NewAllocation = Instances.AddUninitialized();
 	FInstanceHandle InstanceHandle { (uint16)NewAllocation.Index, InstanceSerial };
 
-	new (NewAllocation) FSequenceInstance(Linker, Player, InstanceHandle, RootInstanceHandle, SequenceID, CompiledDataID);
+	new (NewAllocation) FSequenceInstance(Linker, Player, InstanceHandle, ParentInstanceHandle, RootInstanceHandle, SequenceID);
 
 	return InstanceHandle;
 }
diff --git a/Engine/Source/Runtime/MovieScene/Private/EntitySystem/MovieSceneSequenceInstance.cpp b/Engine/Source/Runtime/MovieScene/Private/EntitySystem/MovieSceneSequenceInstance.cpp
--- a/Engine/Source/Runtime/MovieScene/Private/EntitySystem/MovieSceneSequenceInstance.cpp
+++ b/Engine/Source/Runtime/MovieScene/Private/EntitySystem/MovieSceneSequenceInstance.cpp
@@ -85,20 +85,18 @@ FSequenceInstance::FSequenceInstance(UMovieSceneEntitySystemLinker* Linker, IMov
 	RootSequenceName = RootSequence->GetPathName();
 #endif
 
-	CompiledDataID = Player->GetEvaluationTemplate().GetCompiledDataID();
-
 	FMovieSceneObjectCache& ObjectCache = Player->State.GetObjectCache(SequenceID);
 	OnInvalidateObjectBindingHandle = ObjectCache.OnBindingInvalidated.AddUObject(Linker, &UMovieSceneEntitySystemLinker::InvalidateObjectBinding, InstanceHandle);
 
 	InvalidateCachedData(Linker);
 }
 
-FSequenceInstance::FSequenceInstance(UMovieSceneEntitySystemLinker* Linker, IMovieScenePlayer* Player, FInstanceHandle InInstanceHandle, FRootInstanceHandle InRootInstanceHandle, FMovieSceneSequenceID InSequenceID, FMovieSceneCompiledDataID InCompiledDataID)
-	: CompiledDataID(InCompiledDataID)
-	, SequenceID(InSequenceID)
+FSequenceInstance::FSequenceInstance(UMovieSceneEntitySystemLinker* Linker, IMovieScenePlayer* Player, FInstanceHandle InInstanceHandle, FInstanceHandle InParentInstanceHandle, FRootInstanceHandle InRootInstanceHandle, FMovieSceneSequenceID InSequenceID)
+	: SequenceID(InSequenceID)
 	, RootOverrideSequenceID(MovieSceneSequenceID::Invalid)
 	, PlayerIndex(Player->GetUniqueIndex())
 	, InstanceHandle(InInstanceHandle)
+	, ParentInstanceHandle(InParentInstanceHandle)
 	, RootInstanceHandle(InRootInstanceHandle)
 {
 	// Sub Sequence instances always start in a non-finished state because they will only ever
@@ -129,14 +127,15 @@ void FSequenceInstance::InitializeLegacyEvaluator(UMovieSceneEntitySystemLinker*
 	IMovieScenePlayer* Player = GetPlayer();
 	check(Player);
 
-	UMovieSceneCompiledDataManager*     CompiledDataManager = Player->GetEvaluationTemplate().GetCompiledDataManager();
-	const FMovieSceneCompiledDataEntry& CompiledEntry       = CompiledDataManager->GetEntryRef(CompiledDataID);
+	FMovieSceneRootEvaluationTemplateInstance& RootTemplate        = Player->GetEvaluationTemplate();
+	UMovieSceneCompiledDataManager*            CompiledDataManager = RootTemplate.GetCompiledDataManager();
+	const FMovieSceneCompiledDataEntry&        CompiledEntry       = CompiledDataManager->GetEntryRef(RootTemplate.GetCompiledDataID());
 
 	if (EnumHasAnyFlags(CompiledEntry.AccumulatedMask, EMovieSceneSequenceCompilerMask::EvaluationTemplate))
 	{
 		if (!LegacyEvaluator)
 		{
-			LegacyEvaluator = MakeUnique<FMovieSceneTrackEvaluator>(CompiledEntry.GetSequence(), CompiledDataID, CompiledDataManager);
+			LegacyEvaluator = MakeUnique<FMovieSceneTrackEvaluator>(CompiledEntry.GetSequence(), RootTemplate.GetCompiledDataID(), CompiledDataManager);
 		}
 	}
 	else if (LegacyEvaluator)
@@ -153,24 +152,25 @@ void FSequenceInstance::InvalidateCachedData(UMovieSceneEntitySystemLinker* Link
 	IMovieScenePlayer* Player = GetPlayer();
 	check(Player);
 
-	UMovieSceneCompiledDataManager* CompiledDataManager = Player->GetEvaluationTemplate().GetCompiledDataManager();
-
-	UMovieSceneSequence* Sequence = CompiledDataManager->GetEntryRef(CompiledDataID).GetSequence();
-	Player->State.AssignSequence(SequenceID, *Sequence, *Player);
+	FMovieSceneRootEvaluationTemplateInstance& RootTemplate        = Player->GetEvaluationTemplate();
+	UMovieSceneCompiledDataManager*            CompiledDataManager = RootTemplate.GetCompiledDataManager();
+	FMovieSceneCompiledDataID                  RootCompiledDataID  = RootTemplate.GetCompiledDataID();
 
 	if (SequenceID == MovieSceneSequenceID::Root)
 	{
+		Player->State.AssignSequence(SequenceID, *RootTemplate.GetRootSequence(), *Player);
+
 		// Try and recreate the volatility manager if this sequence is now volatile
 		if (!VolatilityManager)
 		{
-			VolatilityManager = FCompiledDataVolatilityManager::Construct(*Player, CompiledDataID, CompiledDataManager);
+			VolatilityManager = FCompiledDataVolatilityManager::Construct(*Player, RootCompiledDataID, CompiledDataManager);
 			if (VolatilityManager)
 			{
-				VolatilityManager->ConditionalRecompile(*Player, CompiledDataID, CompiledDataManager);
+				VolatilityManager->ConditionalRecompile(*Player, RootCompiledDataID, CompiledDataManager);
 			}
 		}
 
-		ISequenceUpdater::FactoryInstance(SequenceUpdater, CompiledDataManager, CompiledDataID);
+		ISequenceUpdater::FactoryInstance(SequenceUpdater, CompiledDataManager, RootCompiledDataID);
 
 		SequenceUpdater->InvalidateCachedData(Linker);
 
@@ -181,19 +181,25 @@ void FSequenceInstance::InvalidateCachedData(UMovieSceneEntitySystemLinker* Link
 
 		InitializeLegacyEvaluator(Linker);
 	}
+	else if (UMovieSceneSequence* SubSequence = RootTemplate.GetSequence(SequenceID))
+	{
+		Player->State.AssignSequence(SequenceID, *SubSequence, *Player);
+	}
 }
 
 void FSequenceInstance::DissectContext(UMovieSceneEntitySystemLinker* Linker, const FMovieSceneContext& InContext, TArray<TRange<FFrameTime>>& OutDissections)
 {
 	check(SequenceID == MovieSceneSequenceID::Root);
 
-
 	IMovieScenePlayer* Player = GetPlayer();
 
 	if (VolatilityManager)
 	{
-		UMovieSceneCompiledDataManager* CompiledDataManager = Player->GetEvaluationTemplate().GetCompiledDataManager();
-		if (VolatilityManager->ConditionalRecompile(*Player, CompiledDataID, CompiledDataManager))
+		FMovieSceneRootEvaluationTemplateInstance& RootTemplate        = Player->GetEvaluationTemplate();
+		UMovieSceneCompiledDataManager*            CompiledDataManager = RootTemplate.GetCompiledDataManager();
+		FMovieSceneCompiledDataID                  RootCompiledDataID  = RootTemplate.GetCompiledDataID();
+
+		if (VolatilityManager->ConditionalRecompile(*Player, RootCompiledDataID, CompiledDataManager))
 		{
 			InvalidateCachedData(Linker);
 		}
@@ -333,7 +339,7 @@ void FSequenceInstance::PostEvaluation(UMovieSceneEntitySystemLinker* Linker)
 			const bool bShouldPurgeTemplates = VolatilityManager && LegacyEvaluator;
 
 			UMovieSceneCompiledDataManager* LocalCompiledDataManager = bShouldPurgeTemplates ? Player->GetEvaluationTemplate().GetCompiledDataManager() : nullptr;
-			FMovieSceneCompiledDataID       LocalCompiledDataID      = CompiledDataID;
+			FMovieSceneCompiledDataID       LocalCompiledDataID      = bShouldPurgeTemplates ? Player->GetEvaluationTemplate().GetCompiledDataID()      : FMovieSceneCompiledDataID();
 
 			Player->PostEvaluation(Context);
 
diff --git a/Engine/Source/Runtime/MovieScene/Private/EntitySystem/MovieSceneSequenceUpdaters.cpp b/Engine/Source/Runtime/MovieScene/Private/EntitySystem/MovieSceneSequenceUpdaters.cpp
--- a/Engine/Source/Runtime/MovieScene/Private/EntitySystem/MovieSceneSequenceUpdaters.cpp
+++ b/Engine/Source/Runtime/MovieScene/Private/EntitySystem/MovieSceneSequenceUpdaters.cpp
@@ -51,6 +51,8 @@ private:
 
 	TOptional<TArray<FFrameTime>> CachedDeterminismFences;
 	FMovieSceneCompiledDataID CompiledDataID;
+
+	TOptional<bool> bDynamicWeighting;
 };
 
 /** Hierarchical sequence updater */
@@ -74,7 +76,7 @@ private:
 
 	TRange<FFrameNumber> UpdateEntitiesForSequence(const FMovieSceneEntityComponentField* ComponentField, FFrameTime SequenceTime, FMovieSceneEvaluationFieldEntitySet& OutEntities);
 
-	FInstanceHandle GetOrCreateSequenceInstance(IMovieScenePlayer* InPlayer, FInstanceRegistry* InstanceRegistry, FRootInstanceHandle RootInstanceHandle, FMovieSceneSequenceID SequenceID);
+	FInstanceHandle GetOrCreateSequenceInstance(IMovieScenePlayer* InPlayer, const FMovieSceneSequenceHierarchy* Hierarchy, FInstanceRegistry* InstanceRegistry, FRootInstanceHandle RootInstanceHandle, FMovieSceneSequenceID SequenceID);
 
 private:
 
@@ -90,6 +92,8 @@ private:
 	FMovieSceneCompiledDataID CompiledDataID;
 
 	FMovieSceneSequenceID RootOverrideSequenceID;
+
+	TOptional<bool> bDynamicWeighting;
 };
 
 void DissectRange(TArrayView<const FFrameTime> InDissectionTimes, const TRange<FFrameTime>& Bounds, TArray<TRange<FFrameTime>>& OutDissections)
@@ -229,6 +233,11 @@ void FSequenceUpdater_Flat::Update(UMovieSceneEntitySystemLinker* Linker, FRootI
 		return;
 	}
 
+	if (!bDynamicWeighting.IsSet())
+	{
+		bDynamicWeighting = InPlayer->HasDynamicWeighting() || EnumHasAnyFlags(Sequence->GetFlags(), EMovieSceneSequenceFlags::DynamicWeighting);
+	}
+
 	FMovieSceneEvaluationFieldEntitySet EntitiesScratch;
 
 	FFrameNumber ImportTime = Context.GetEvaluationFieldTime();
@@ -246,10 +255,12 @@ void FSequenceUpdater_Flat::Update(UMovieSceneEntitySystemLinker* Linker, FRootI
 		}
 
 		FEntityImportSequenceParams Params;
+		Params.SequenceID = MovieSceneSequenceID::Root;
 		Params.InstanceHandle = InstanceHandle;
 		Params.RootInstanceHandle = InstanceHandle;
 		Params.DefaultCompletionMode = Sequence->DefaultCompletionMode;
 		Params.HierarchicalBias = 0;
+		Params.bDynamicWeighting = bDynamicWeighting.Get(false);
 
 		SequenceInstance.Ledger.UpdateEntities(Linker, Params, ComponentField, EntitiesScratch);
 	}
@@ -263,10 +274,12 @@ void FSequenceUpdater_Flat::Update(UMovieSceneEntitySystemLinker* Linker, FRootI
 		if (EntitiesScratch.Num() != 0)
 		{
 			FEntityImportSequenceParams Params;
+			Params.SequenceID = MovieSceneSequenceID::Root;
 			Params.InstanceHandle = InstanceHandle;
 			Params.RootInstanceHandle = InstanceHandle;
 			Params.DefaultCompletionMode = Sequence->DefaultCompletionMode;
 			Params.HierarchicalBias = 0;
+			Params.bDynamicWeighting = bDynamicWeighting.Get(false);
 
 			SequenceInstance.Ledger.UpdateOneShotEntities(Linker, Params, ComponentField, EntitiesScratch);
 		}
@@ -287,6 +300,7 @@ void FSequenceUpdater_Flat::InvalidateCachedData(UMovieSceneEntitySystemLinker*
 {
 	CachedEntityRange = TRange<FFrameNumber>::Empty();
 	CachedDeterminismFences.Reset();
+	bDynamicWeighting.Reset();
 }
 
 
@@ -373,15 +387,32 @@ void FSequenceUpdater_Hierarchical::DissectContext(UMovieSceneEntitySystemLinker
 	}
 }
 
-FInstanceHandle FSequenceUpdater_Hierarchical::GetOrCreateSequenceInstance(IMovieScenePlayer* InPlayer, FInstanceRegistry* InstanceRegistry, FRootInstanceHandle RootInstanceHandle, FMovieSceneSequenceID SequenceID)
+FInstanceHandle FSequenceUpdater_Hierarchical::GetOrCreateSequenceInstance(IMovieScenePlayer* InPlayer, const FMovieSceneSequenceHierarchy* Hierarchy, FInstanceRegistry* InstanceRegistry, FRootInstanceHandle RootInstanceHandle, FMovieSceneSequenceID SequenceID)
 {
+	check(SequenceID != MovieSceneSequenceID::Root);
+
 	if (FSubInstanceData* Existing = SequenceInstances.Find(SequenceID))
 	{
 		Existing->bNeedsDestroy = false;
 		return Existing->Handle;
 	}
 
-	FInstanceHandle InstanceHandle = InstanceRegistry->AllocateSubInstance(InPlayer, SequenceID, RootInstanceHandle);
+	const FMovieSceneSequenceHierarchyNode* Node = Hierarchy->FindNode(SequenceID);
+
+	checkf(Node, TEXT("Attempting to construct a new sub sequence instance with a sub sequence ID that does not exist in the hierarchy"));
+	checkf(Node->ParentID != MovieSceneSequenceID::Invalid, TEXT("Parent should never be invalid for a non-root SequenceID"));
+
+	FInstanceHandle ParentInstance;
+	if (Node->ParentID == MovieSceneSequenceID::Root)
+	{
+		ParentInstance = RootInstanceHandle;
+	}
+	else
+	{
+		ParentInstance = GetOrCreateSequenceInstance(InPlayer, Hierarchy, InstanceRegistry, RootInstanceHandle, Node->ParentID);
+	}
+
+	FInstanceHandle InstanceHandle = InstanceRegistry->AllocateSubInstance(InPlayer, SequenceID, RootInstanceHandle, ParentInstance);
 	SequenceInstances.Add(SequenceID, FSubInstanceData { InstanceHandle });
 
 	return InstanceHandle;
@@ -424,14 +455,15 @@ void FSequenceUpdater_Hierarchical::Update(UMovieSceneEntitySystemLinker* Linker
 
 	TArray<FMovieSceneSequenceID, TInlineAllocator<16>> ActiveSequences;
 
+	const FMovieSceneSequenceHierarchy* RootHierarchy = CompiledDataManager->FindHierarchy(CompiledDataID);
+
 	if (RootOverrideSequenceID != MovieSceneSequenceID::Root)
 	{
-		const FMovieSceneSequenceHierarchy* RootHierarchy = CompiledDataManager->FindHierarchy(CompiledDataID);
 		const FMovieSceneSubSequenceData*   SubData         = RootHierarchy ? RootHierarchy->FindSubData(RootOverrideSequenceID) : nullptr;
 		UMovieSceneSequence*                RootSequence    = SubData ? SubData->GetSequence() : nullptr;
 		if (ensure(RootSequence))
 		{
-			RootInstanceHandle = GetOrCreateSequenceInstance(InPlayer, InstanceRegistry, InstanceHandle, RootOverrideSequenceID);
+			RootInstanceHandle = GetOrCreateSequenceInstance(InPlayer, RootHierarchy, InstanceRegistry, InstanceHandle, RootOverrideSequenceID);
 			RootCompiledDataID = CompiledDataManager->GetDataID(RootSequence);
 			RootContext        = Context.Transform(SubData->RootToSequenceTransform, SubData->TickResolution);
 
@@ -460,16 +492,23 @@ void FSequenceUpdater_Hierarchical::Update(UMovieSceneEntitySystemLinker* Linker
 		}
 		else
 		{
+			if (!bDynamicWeighting.IsSet())
+			{
+				bDynamicWeighting = InPlayer->HasDynamicWeighting() || EnumHasAnyFlags(CompiledDataManager->GetEntryRef(RootCompiledDataID).AccumulatedFlags, EMovieSceneSequenceFlags::DynamicWeighting);
+			}
+
 			// Update entities if necessary
 			if (bGatherEntities)
 			{
 				CachedEntityRange = UpdateEntitiesForSequence(RootComponentField, ImportTime, EntitiesScratch);
 
 				FEntityImportSequenceParams Params;
+				Params.SequenceID = MovieSceneSequenceID::Root;
 				Params.InstanceHandle = RootInstanceHandle;
 				Params.RootInstanceHandle = InstanceHandle;
 				Params.DefaultCompletionMode = RootSequence->DefaultCompletionMode;
 				Params.HierarchicalBias = 0;
+				Params.bDynamicWeighting = bDynamicWeighting.Get(false);
 
 				RootInstance.Ledger.UpdateEntities(Linker, Params, RootComponentField, EntitiesScratch);
 			}
@@ -483,10 +522,12 @@ void FSequenceUpdater_Hierarchical::Update(UMovieSceneEntitySystemLinker* Linker
 				if (EntitiesScratch.Num() != 0)
 				{
 					FEntityImportSequenceParams Params;
+					Params.SequenceID = MovieSceneSequenceID::Root;
 					Params.InstanceHandle = RootInstanceHandle;
 					Params.RootInstanceHandle = InstanceHandle;
 					Params.DefaultCompletionMode = RootSequence->DefaultCompletionMode;
 					Params.HierarchicalBias = 0;
+					Params.bDynamicWeighting = bDynamicWeighting.Get(false);
 
 					RootInstance.Ledger.UpdateOneShotEntities(Linker, Params, RootComponentField, EntitiesScratch);
 				}
@@ -531,7 +572,7 @@ void FSequenceUpdater_Hierarchical::Update(UMovieSceneEntitySystemLinker* Linker
 				FMovieSceneCompiledDataID SubDataID = CompiledDataManager->GetDataID(SubSequence);
 
 				// Set the context for the root sequence instance
-				FInstanceHandle    SubSequenceHandle = GetOrCreateSequenceInstance(InPlayer, InstanceRegistry, InstanceHandle, SequenceIDFromRoot);
+				FInstanceHandle    SubSequenceHandle = GetOrCreateSequenceInstance(InPlayer, RootHierarchy, InstanceRegistry, InstanceHandle, SequenceIDFromRoot);
 				FSequenceInstance& SubSequenceInstance = InstanceRegistry->MutateInstance(SubSequenceHandle);
 
 				// Update the sub sequence's context
@@ -566,13 +607,14 @@ void FSequenceUpdater_Hierarchical::Update(UMovieSceneEntitySystemLinker* Linker
 				const FFrameTime SubSequenceTime = SubContext.GetEvaluationFieldTime();
 
 				FEntityImportSequenceParams Params;
+				Params.SequenceID = SequenceIDFromRoot;
 				Params.InstanceHandle = SubSequenceHandle;
 				Params.RootInstanceHandle = InstanceHandle;
 				Params.DefaultCompletionMode = SubSequence->DefaultCompletionMode;
 				Params.HierarchicalBias = SubData->HierarchicalBias;
 				Params.bPreRoll  = bIsPreRoll;
 				Params.bPostRoll = bIsPostRoll;
-				Params.bHasHierarchicalEasing = SubData->bHasHierarchicalEasing;
+				Params.bDynamicWeighting = bDynamicWeighting.Get(false); // Always inherit dynamic weighting flags
 
 				if (bGatherEntities)
 				{
@@ -663,6 +705,7 @@ void FSequenceUpdater_Hierarchical::Destroy(UMovieSceneEntitySystemLinker* Linke
 
 void FSequenceUpdater_Hierarchical::InvalidateCachedData(UMovieSceneEntitySystemLinker* Linker)
 {
+	bDynamicWeighting.Reset();
 	CachedEntityRange = TRange<FFrameNumber>::Empty();
 
 	FInstanceRegistry* InstanceRegistry = Linker->GetInstanceRegistry();
diff --git a/Engine/Source/Runtime/MovieScene/Private/Evaluation/MovieSceneSequenceHierarchy.cpp b/Engine/Source/Runtime/MovieScene/Private/Evaluation/MovieSceneSequenceHierarchy.cpp
--- a/Engine/Source/Runtime/MovieScene/Private/Evaluation/MovieSceneSequenceHierarchy.cpp
+++ b/Engine/Source/Runtime/MovieScene/Private/Evaluation/MovieSceneSequenceHierarchy.cpp
@@ -15,7 +15,6 @@
 FMovieSceneSubSequenceData::FMovieSceneSubSequenceData()
 	: Sequence(nullptr)
 	, HierarchicalBias(0)
-	, bHasHierarchicalEasing(false)
 {}
 
 FMovieSceneSubSequenceData::FMovieSceneSubSequenceData(const UMovieSceneSubSection& InSubSection)
@@ -27,7 +26,6 @@ FMovieSceneSubSequenceData::FMovieSceneSubSequenceData(const UMovieSceneSubSecti
 	, ParentFirstLoopStartFrameOffset(InSubSection.Parameters.FirstLoopStartFrameOffset)
 	, bCanLoop(InSubSection.Parameters.bCanLoop)
 	, HierarchicalBias(InSubSection.Parameters.HierarchicalBias)
-	, bHasHierarchicalEasing(false)
 #if WITH_EDITORONLY_DATA
 	, SectionPath(*InSubSection.GetPathNameInMovieScene())
 #endif
@@ -77,10 +75,6 @@ FMovieSceneSubSequenceData::FMovieSceneSubSequenceData(const UMovieSceneSubSecti
 	{
 		PostRollRange = UE::MovieScene::MakeDiscreteRangeFromLower( TRangeBound<FFrameNumber>::FlipInclusion(SubSectionRange.GetUpperBound()), InSubSection.GetPostRollFrames() ) * RootToSequenceTransform.LinearTransform;
 	}
-
-	const bool bHasSubSectionEaseIn  = (InSubSection.Easing.GetEaseInDuration() > 0);
-	const bool bHasSubSectionEaseOut = (InSubSection.Easing.GetEaseOutDuration() > 0);
-	bHasHierarchicalEasing = (bHasSubSectionEaseIn || bHasSubSectionEaseOut);
 }
 
 UMovieSceneSequence* FMovieSceneSubSequenceData::GetSequence() const
@@ -246,10 +240,9 @@ void FMovieSceneSequenceHierarchy::LogHierarchy() const
 
 			FString Indent;
 			Indent.Append(TEXT(" "), CurDepth * 2);
-			UE_LOG(LogMovieScene, Log, TEXT("%s%s Loop=%s HEasing=%s HBias=%d UnwarpedRange=%s Transform=%s"), 
+			UE_LOG(LogMovieScene, Log, TEXT("%s%s Loop=%s HBias=%d UnwarpedRange=%s Transform=%s"), 
 					*Indent, *CurData->GetSequence()->GetName(),
 					*LexToString(CurData->bCanLoop),
-					*LexToString(CurData->bHasHierarchicalEasing),
 					CurData->HierarchicalBias,
 					*LexToString(CurData->UnwarpedPlayRange.Value),
 					*LexToString(CurData->RootToSequenceTransform));
diff --git a/Engine/Source/Runtime/MovieScene/Private/Evaluation/MovieSceneSequenceWeights.cpp b/Engine/Source/Runtime/MovieScene/Private/Evaluation/MovieSceneSequenceWeights.cpp
new file mode 100644
--- /dev/null
+++ b/Engine/Source/Runtime/MovieScene/Private/Evaluation/MovieSceneSequenceWeights.cpp
@@ -0,0 +1,109 @@
+// Copyright Epic Games, Inc. All Rights Reserved.
+
+#include "Evaluation/MovieSceneSequenceWeights.h"
+#include "EntitySystem/MovieSceneEntitySystemLinker.h"
+#include "EntitySystem/BuiltInComponentTypes.h"
+
+namespace UE::MovieScene
+{
+
+FSequenceWeights::FSequenceWeights(UMovieSceneEntitySystemLinker* InLinker, FRootInstanceHandle InRootInstanceHandle)
+	: WeakLinker(InLinker)
+	, RootInstanceHandle(InRootInstanceHandle)
+{
+}
+
+FSequenceWeights::~FSequenceWeights()
+{
+	UMovieSceneEntitySystemLinker* Linker = WeakLinker.Get();
+	if (Linker)
+	{
+		FBuiltInComponentTypes* BuiltInComponents = FBuiltInComponentTypes::Get();
+		for (TPair<FMovieSceneSequenceID, FMovieSceneEntityID> Pair : WeightEntitiesBySequenceID)
+		{
+			Linker->EntityManager.AddComponent(Pair.Value, BuiltInComponents->Tags.NeedsUnlink);
+		}
+	}
+}
+
+void FSequenceWeights::SetWeight(FMovieSceneSequenceID InSequenceID, double Weight)
+{
+	UMovieSceneEntitySystemLinker* Linker = WeakLinker.Get();
+	if (!ensureAlways(Linker))
+	{
+		return;
+	}
+
+	FBuiltInComponentTypes* BuiltInComponents = FBuiltInComponentTypes::Get();
+
+	FMovieSceneEntityID EntityID = WeightEntitiesBySequenceID.FindRef(InSequenceID);
+	if (!EntityID)
+	{
+		EntityID = FEntityBuilder()
+		.Add(BuiltInComponents->RootInstanceHandle, RootInstanceHandle)
+		.Add(BuiltInComponents->HierarchicalEasingProvider, InSequenceID)
+		.Add(BuiltInComponents->WeightResult, Weight)
+		.Add(BuiltInComponents->WeightAndEasingResult, Weight)
+		.AddTag(BuiltInComponents->Tags.NeedsLink)
+		.CreateEntity(&Linker->EntityManager);
+
+		WeightEntitiesBySequenceID.Add(InSequenceID, EntityID);
+	}
+	else
+	{
+		TOptionalComponentWriter<double> WeightComponent = Linker->EntityManager.WriteComponent(EntityID, BuiltInComponents->WeightResult);
+		if (WeightComponent)
+		{
+			*WeightComponent = Weight;
+		}
+	}
+}
+
+void FSequenceWeights::SetWeights(const TSortedMap<FMovieSceneSequenceID, double>& Weights)
+{
+	UMovieSceneEntitySystemLinker* Linker = WeakLinker.Get();
+	if (!ensureAlways(Linker))
+	{
+		return;
+	}
+
+	FBuiltInComponentTypes* BuiltInComponents = FBuiltInComponentTypes::Get();
+	for (TPair<FMovieSceneSequenceID, double> Pair : Weights)
+	{
+		FMovieSceneEntityID EntityID = WeightEntitiesBySequenceID.FindRef(Pair.Key);
+		if (!EntityID)
+		{
+			EntityID = FEntityBuilder()
+			.Add(BuiltInComponents->RootInstanceHandle, RootInstanceHandle)
+			.Add(BuiltInComponents->HierarchicalEasingProvider, Pair.Key)
+			.Add(BuiltInComponents->WeightResult, Pair.Value)
+			.Add(BuiltInComponents->WeightAndEasingResult, Pair.Value)
+			.AddTag(BuiltInComponents->Tags.NeedsLink)
+			.CreateEntity(&Linker->EntityManager);
+
+			WeightEntitiesBySequenceID.Add(Pair.Key, EntityID);
+		}
+		else
+		{
+			TOptionalComponentWriter<double> WeightComponent = Linker->EntityManager.WriteComponent(EntityID, BuiltInComponents->WeightResult);
+			if (WeightComponent)
+			{
+				*WeightComponent = Pair.Value;
+			}
+		}
+	}
+}
+
+void FSequenceWeights::RemoveWeight(FMovieSceneSequenceID InSequenceID)
+{
+	FMovieSceneEntityID            EntityID = WeightEntitiesBySequenceID.FindRef(InSequenceID);
+	UMovieSceneEntitySystemLinker* Linker   = WeakLinker.Get();
+
+	if (EntityID && ensureAlways(Linker))
+	{
+		Linker->EntityManager.AddComponent(EntityID, FBuiltInComponentTypes::Get()->Tags.NeedsUnlink);
+	}
+}
+
+
+} // namespace UE::MovieScene
\ No newline at end of file
diff --git a/Engine/Source/Runtime/MovieScene/Private/MovieSceneSection.cpp b/Engine/Source/Runtime/MovieScene/Private/MovieSceneSection.cpp
--- a/Engine/Source/Runtime/MovieScene/Private/MovieSceneSection.cpp
+++ b/Engine/Source/Runtime/MovieScene/Private/MovieSceneSection.cpp
@@ -322,7 +322,6 @@ void UMovieSceneSection::BuildDefaultComponents(UMovieSceneEntitySystemLinker* E
 	OutImportedEntity->AddBuilder(
 		FEntityBuilder()
 		.AddConditional(Components->Easing,                     FEasingComponentData{ decltype(FEasingComponentData::Section)(this) }, bHasEasing)
-		.AddConditional(Components->HierarchicalEasingChannel, uint16(-1), Params.Sequence.bHasHierarchicalEasing)
 		.AddConditional(Components->HierarchicalBias,           Params.Sequence.HierarchicalBias, Params.Sequence.HierarchicalBias != 0)
 		.AddConditional(Components->Interrogation.InputKey,     Params.InterrogationKey, Params.InterrogationKey.IsValid())
 		.AddConditional(Components->EvalTime,                   Params.EntityMetaData ? Params.EntityMetaData->ForcedTime : 0, bHasForcedTime)
diff --git a/Engine/Source/Runtime/MovieScene/Private/MovieSceneSequencePlayer.cpp b/Engine/Source/Runtime/MovieScene/Private/MovieSceneSequencePlayer.cpp
--- a/Engine/Source/Runtime/MovieScene/Private/MovieSceneSequencePlayer.cpp
+++ b/Engine/Source/Runtime/MovieScene/Private/MovieSceneSequencePlayer.cpp
@@ -7,6 +7,8 @@
 #include "MovieSceneSequence.h"
 #include "MovieSceneSequenceTickManager.h"
 #include "Engine/Engine.h"
+#include "UObject/Stack.h"
+#include "Internationalization/Text.h"
 #include "GameFramework/WorldSettings.h"
 #include "Misc/RuntimeErrors.h"
 #include "Net/UnrealNetwork.h"
@@ -14,10 +16,13 @@
 #include "Engine/NetConnection.h"
 #include "EntitySystem/MovieSceneEntitySystemLinker.h"
 #include "EntitySystem/MovieSceneEntitySystemRunner.h"
+#include "Compilation/MovieSceneCompiledDataManager.h"
+#include "Evaluation/MovieSceneSequenceWeights.h"
 #include "GameFramework/PlayerController.h"
 #include "GameFramework/PlayerState.h"
 #include "Algo/BinarySearch.h"
 
+
 #include UE_INLINE_GENERATED_CPP_BY_NAME(MovieSceneSequencePlayer)
 
 DEFINE_LOG_CATEGORY_STATIC(LogMovieSceneRepl, Log, All);
@@ -889,6 +894,16 @@ void UMovieSceneSequencePlayer::Initialize(UMovieSceneSequence* InSequence)
 	check(RunnerToUse);
 	RootTemplateInstance.Initialize(*Sequence, *this, nullptr, RunnerToUse);
 
+	if (!PlaybackSettings.bDynamicWeighting)
+	{
+		UMovieSceneCompiledDataManager* CompiledDataManager = RootTemplateInstance.GetCompiledDataManager();
+		FMovieSceneCompiledDataID       CompiledDataID      = RootTemplateInstance.GetCompiledDataID();
+		if (CompiledDataManager && CompiledDataID.IsValid())
+		{
+			PlaybackSettings.bDynamicWeighting = EnumHasAnyFlags(CompiledDataManager->GetEntryRef(CompiledDataID).AccumulatedFlags, EMovieSceneSequenceFlags::DynamicWeighting);
+		}
+	}
+
 	LatentActionManager.ClearLatentActions();
 
 	// Set up playback position (with offset) after Stop(), which will reset the starting time to StartTime
@@ -1223,6 +1238,11 @@ bool UMovieSceneSequencePlayer::IsDisablingEventTriggers(FFrameTime& DisabledUnt
 	return false;
 }
 
+bool UMovieSceneSequencePlayer::HasDynamicWeighting() const
+{
+	return PlaybackSettings.bDynamicWeighting;
+}
+
 void UMovieSceneSequencePlayer::PreEvaluation(const FMovieSceneContext& Context)
 {
 	RunPreEvaluationCallbacks();
@@ -1795,4 +1815,27 @@ void UMovieSceneSequencePlayer::RunLatentActions()
 	}
 }
 
+void UMovieSceneSequencePlayer::SetWeight(double InWeight)
+{
+	SetWeight(InWeight, MovieSceneSequenceID::Root);
+}
 
+void UMovieSceneSequencePlayer::SetWeight(double InWeight, FMovieSceneSequenceID SequenceID)
+{
+	UMovieSceneEntitySystemLinker* Linker = RootTemplateInstance.GetEntitySystemLinker();
+	if (Linker)
+	{
+		if (!SequenceWeights)
+		{
+			SequenceWeights = MakeUnique<UE::MovieScene::FSequenceWeights>(Linker, RootTemplateInstance.GetRootInstanceHandle());
+
+			if (!PlaybackSettings.bDynamicWeighting && Sequence)
+			{
+				FText Text = NSLOCTEXT("UMovieSceneSequencePlayer", "SetWeightWarning", "Attempting to set a weight on sequence {0} with PlaybackSettings.bDynamicWeighting disabled. This may lead to undesireable blending artifacts or broken in/out blends.");
+				FFrame::KismetExecutionMessage(*FText::Format(Text, FText::FromString(Sequence->GetName())).ToString(), ELogVerbosity::Warning);
+			}
+		}
+
+		SequenceWeights->SetWeight(SequenceID, InWeight);
+	}
+}
diff --git a/Engine/Source/Runtime/MovieScene/Private/Sections/MovieSceneSubSection.cpp b/Engine/Source/Runtime/MovieScene/Private/Sections/MovieSceneSubSection.cpp
--- a/Engine/Source/Runtime/MovieScene/Private/Sections/MovieSceneSubSection.cpp
+++ b/Engine/Source/Runtime/MovieScene/Private/Sections/MovieSceneSubSection.cpp
@@ -29,6 +29,8 @@ UMovieSceneSubSection::UMovieSceneSubSection(const FObjectInitializer& ObjInitia
 	, PrerollTime_DEPRECATED(DeprecatedMagicNumber)
 {
 	NetworkMask = (uint8)(EMovieSceneServerClientMask::Server | EMovieSceneServerClientMask::Client);
+
+	SetBlendType(EMovieSceneBlendType::Absolute);
 }
 
 FMovieSceneSequenceTransform UMovieSceneSubSection::OuterToInnerTransform() const
@@ -184,6 +186,25 @@ void UMovieSceneSubSection::PostLoad()
 	Super::PostLoad();
 }
 
+bool UMovieSceneSubSection::PopulateEvaluationFieldImpl(const TRange<FFrameNumber>& EffectiveRange, const FMovieSceneEvaluationFieldEntityMetaData& InMetaData, FMovieSceneEntityComponentFieldBuilder* OutFieldBuilder)
+{
+	const int32 EntityIndex   = OutFieldBuilder->FindOrAddEntity(this, 0);
+	const int32 MetaDataIndex = OutFieldBuilder->AddMetaData(InMetaData);
+	OutFieldBuilder->AddPersistentEntity(EffectiveRange, EntityIndex, MetaDataIndex);
+
+	return true;
+}
+
+void UMovieSceneSubSection::ImportEntityImpl(UMovieSceneEntitySystemLinker* EntityLinker, const FEntityImportParams& Params, FImportedEntity* OutImportedEntity)
+{
+	using namespace UE::MovieScene;
+
+	OutImportedEntity->AddBuilder(
+		FEntityBuilder().AddTag(FBuiltInComponentTypes::Get()->Tags.Root)
+	);
+	BuildDefaultSubSectionComponents(EntityLinker, Params, OutImportedEntity);
+}
+
 void UMovieSceneSubSection::SetSequence(UMovieSceneSequence* Sequence)
 {
 	SubSequence = Sequence;
diff --git a/Engine/Source/Runtime/MovieScene/Private/Tracks/MovieSceneSubTrack.cpp b/Engine/Source/Runtime/MovieScene/Private/Tracks/MovieSceneSubTrack.cpp
--- a/Engine/Source/Runtime/MovieScene/Private/Tracks/MovieSceneSubTrack.cpp
+++ b/Engine/Source/Runtime/MovieScene/Private/Tracks/MovieSceneSubTrack.cpp
@@ -25,6 +25,8 @@ UMovieSceneSubTrack::UMovieSceneSubTrack( const FObjectInitializer& ObjectInitia
 #endif
 
 	BuiltInTreePopulationMode = ETreePopulationMode::Blended;
+
+	SupportedBlendTypes.Add(EMovieSceneBlendType::Absolute);
 }
 
 UMovieSceneSubSection* UMovieSceneSubTrack::AddSequenceOnRow(UMovieSceneSequence* Sequence, FFrameNumber StartTime, int32 Duration, int32 RowIndex)
diff --git a/Engine/Source/Runtime/MovieScene/Public/EntitySystem/BuiltInComponentTypes.h b/Engine/Source/Runtime/MovieScene/Public/EntitySystem/BuiltInComponentTypes.h
--- a/Engine/Source/Runtime/MovieScene/Public/EntitySystem/BuiltInComponentTypes.h
+++ b/Engine/Source/Runtime/MovieScene/Public/EntitySystem/BuiltInComponentTypes.h
@@ -244,6 +244,8 @@ public:
 
 	TComponentTypeID<FRootInstanceHandle> RootInstanceHandle;
 
+	TComponentTypeID<FMovieSceneSequenceID> SequenceID;
+
 	TComponentTypeID<FFrameTime>          EvalTime;
 
 public:
@@ -318,6 +320,7 @@ public:
 
 	// A float representing the evaluated output of a weight channel
 	TComponentTypeID<double> WeightResult;
+	TComponentTypeID<double> EasingResult;
 
 	// The result of an evaluated FMovieSceneObjectPathChannel
 	TComponentTypeID<UObject*> ObjectResult;
@@ -391,6 +394,8 @@ public:
 
 		FComponentTypeID Ignored;
 
+		FComponentTypeID AlwaysCacheInitialValue;
+
 	} Tags;
 
 	struct
diff --git a/Engine/Source/Runtime/MovieScene/Public/EntitySystem/IMovieSceneEntityProvider.h b/Engine/Source/Runtime/MovieScene/Public/EntitySystem/IMovieSceneEntityProvider.h
--- a/Engine/Source/Runtime/MovieScene/Public/EntitySystem/IMovieSceneEntityProvider.h
+++ b/Engine/Source/Runtime/MovieScene/Public/EntitySystem/IMovieSceneEntityProvider.h
@@ -66,28 +66,30 @@ struct FEntityImportSequenceParams
 {
 	FEntityImportSequenceParams()
 		: HierarchicalBias(0)
+		, SequenceID(MovieSceneSequenceID::Root)
 		, DefaultCompletionMode(EMovieSceneCompletionMode::KeepState)
-		, bHasHierarchicalEasing(false)
 		, bPreRoll(false)
 		, bPostRoll(false)
+		, bDynamicWeighting(false)
 	{}
 
 	int32 HierarchicalBias;
 
+	FMovieSceneSequenceID SequenceID;
 	FInstanceHandle InstanceHandle;
 	FRootInstanceHandle RootInstanceHandle;
 
 	EMovieSceneCompletionMode DefaultCompletionMode;
 
-	bool bHasHierarchicalEasing : 1;
 	bool bPreRoll : 1;
 	bool bPostRoll : 1;
+	bool bDynamicWeighting : 1;
 };
 
 struct FEntityImportParams
 {
-	const FMovieSceneEvaluationFieldEntityMetaData* EntityMetaData;
-	const FMovieSceneEvaluationFieldSharedEntityMetaData* SharedMetaData;
+	const FMovieSceneEvaluationFieldEntityMetaData* EntityMetaData = nullptr;
+	const FMovieSceneEvaluationFieldSharedEntityMetaData* SharedMetaData = nullptr;
 
 	uint32 EntityID = 0;
 
diff --git a/Engine/Source/Runtime/MovieScene/Public/EntitySystem/MovieSceneEntityFactory.h b/Engine/Source/Runtime/MovieScene/Public/EntitySystem/MovieSceneEntityFactory.h
--- a/Engine/Source/Runtime/MovieScene/Public/EntitySystem/MovieSceneEntityFactory.h
+++ b/Engine/Source/Runtime/MovieScene/Public/EntitySystem/MovieSceneEntityFactory.h
@@ -173,7 +173,8 @@ struct FEntityFactories
 	 */
 	MOVIESCENE_API int32 ComputeMutuallyInclusiveComponents(FComponentMask& ComponentMask);
 
-	void RunInitializers(const FComponentMask& ParentType, const FComponentMask& ChildType, const FEntityAllocation* ParentAllocation, TArrayView<const int32> ParentAllocationOffsets, const FEntityRange& InChildEntityRange);
+	MOVIESCENE_API void RunInitializers(const FComponentMask& Type, const FEntityRange& InEntityRange);
+	MOVIESCENE_API void RunInitializers(const FComponentMask& ParentType, const FComponentMask& ChildType, const FEntityAllocation* ParentAllocation, TArrayView<const int32> ParentAllocationOffsets, const FEntityRange& InChildEntityRange);
 
 	TArray<TInlineValue<FChildEntityInitializer>> ChildInitializers;
 	TArray<TInlineValue<FMutualEntityInitializer>> MutualInitializers;
diff --git a/Engine/Source/Runtime/MovieScene/Public/EntitySystem/MovieSceneEntityLedger.h b/Engine/Source/Runtime/MovieScene/Public/EntitySystem/MovieSceneEntityLedger.h
--- a/Engine/Source/Runtime/MovieScene/Public/EntitySystem/MovieSceneEntityLedger.h
+++ b/Engine/Source/Runtime/MovieScene/Public/EntitySystem/MovieSceneEntityLedger.h
@@ -22,6 +22,7 @@ namespace UE
 namespace MovieScene
 {
 struct FEntityImportSequenceParams;
+struct IMovieScenePerEntityMutation;
 
 
 /**
@@ -101,6 +102,17 @@ public:
 	 */
 	void UnlinkOneShots(UMovieSceneEntitySystemLinker* Linker);
 
+
+	/**
+	 * Check whether any of the entities in this ledger or their children match the specified filter
+	 */
+	bool Contains(UMovieSceneEntitySystemLinker* Linker, const FEntityComponentFilter& Filter) const;
+
+	/**
+	 * Mutate all the entities within this ledger by using the specified filter
+	 */
+	void MutateAll(UMovieSceneEntitySystemLinker* Linker, const FEntityComponentFilter& Filter, const IMovieScenePerEntityMutation& Mutation) const;
+
 public:
 
 	/**
diff --git a/Engine/Source/Runtime/MovieScene/Public/EntitySystem/MovieSceneEntitySystemTypes.h b/Engine/Source/Runtime/MovieScene/Public/EntitySystem/MovieSceneEntitySystemTypes.h
--- a/Engine/Source/Runtime/MovieScene/Public/EntitySystem/MovieSceneEntitySystemTypes.h
+++ b/Engine/Source/Runtime/MovieScene/Public/EntitySystem/MovieSceneEntitySystemTypes.h
@@ -907,6 +907,11 @@ struct FEntityDataLocation
 {
 	FEntityAllocation* Allocation;
 	int32 ComponentOffset;
+
+	FEntityRange AsRange() const
+	{
+		return FEntityRange{ Allocation, ComponentOffset, 1 };
+	}
 };
 
 struct FEntityInfo
diff --git a/Engine/Source/Runtime/MovieScene/Public/EntitySystem/MovieSceneInstanceRegistry.h b/Engine/Source/Runtime/MovieScene/Public/EntitySystem/MovieSceneInstanceRegistry.h
--- a/Engine/Source/Runtime/MovieScene/Public/EntitySystem/MovieSceneInstanceRegistry.h
+++ b/Engine/Source/Runtime/MovieScene/Public/EntitySystem/MovieSceneInstanceRegistry.h
@@ -83,7 +83,7 @@ struct FInstanceRegistry
 
 	MOVIESCENE_API FRootInstanceHandle AllocateRootInstance(IMovieScenePlayer* Player);
 
-	MOVIESCENE_API FInstanceHandle AllocateSubInstance(IMovieScenePlayer* Player, FMovieSceneSequenceID SequenceID, FRootInstanceHandle RootInstance);
+	MOVIESCENE_API FInstanceHandle AllocateSubInstance(IMovieScenePlayer* Player, FMovieSceneSequenceID SequenceID, FRootInstanceHandle RootInstance, FInstanceHandle ParentInstanceHandle);
 
 	MOVIESCENE_API void DestroyInstance(FInstanceHandle InstanceHandle);
 
diff --git a/Engine/Source/Runtime/MovieScene/Public/EntitySystem/MovieScenePropertyTraits.h b/Engine/Source/Runtime/MovieScene/Public/EntitySystem/MovieScenePropertyTraits.h
--- a/Engine/Source/Runtime/MovieScene/Public/EntitySystem/MovieScenePropertyTraits.h
+++ b/Engine/Source/Runtime/MovieScene/Public/EntitySystem/MovieScenePropertyTraits.h
@@ -3,6 +3,7 @@
 #pragma once
 
 #include "EntitySystem/MovieScenePropertySystemTypes.h"
+#include "EntitySystem/MovieSceneOperationalTypeConversions.h"
 #include "MovieSceneCommonHelpers.h"
 
 
diff --git a/Engine/Source/Runtime/MovieScene/Public/EntitySystem/MovieSceneSequenceInstance.h b/Engine/Source/Runtime/MovieScene/Public/EntitySystem/MovieSceneSequenceInstance.h
--- a/Engine/Source/Runtime/MovieScene/Public/EntitySystem/MovieSceneSequenceInstance.h
+++ b/Engine/Source/Runtime/MovieScene/Public/EntitySystem/MovieSceneSequenceInstance.h
@@ -148,6 +148,16 @@ public:
 		return RootInstanceHandle;
 	}
 
+	/**
+	 * Get a handle to the parent instance for this sequence instance, or an invalid handle if it is a root instance
+	 *
+	 * @return A a handle to the parent instance for this sequence instance, or an invalid handle if it is a root instance
+	 */
+	FInstanceHandle GetParentInstanceHandle() const
+	{
+		return ParentInstanceHandle;
+	}
+
 	/**
 	 * Returns whether this instance is the root instance.
 	 */
@@ -261,7 +271,7 @@ public:
 	explicit FSequenceInstance(UMovieSceneEntitySystemLinker* Linker, IMovieScenePlayer* Player, FRootInstanceHandle ThisInstanceHandle);
 
 	/** Constructor for sub sequences */
-	explicit FSequenceInstance(UMovieSceneEntitySystemLinker* Linker, IMovieScenePlayer* Player, FInstanceHandle ThisInstanceHandle, FRootInstanceHandle RootInstanceHandle, FMovieSceneSequenceID InSequenceID, FMovieSceneCompiledDataID InCompiledDataID);
+	explicit FSequenceInstance(UMovieSceneEntitySystemLinker* Linker, IMovieScenePlayer* Player, FInstanceHandle ThisInstanceHandle, FInstanceHandle InParentInstanceHandle, FRootInstanceHandle RootInstanceHandle, FMovieSceneSequenceID InSequenceID);
 
 	/** Destructor */
 	~FSequenceInstance();
@@ -297,8 +307,6 @@ private:
 
 	/** Delegate Binding for when an object binding is invalidated in this instance . */
 	FDelegateHandle OnInvalidateObjectBindingHandle;
-	/** This sequence's compiled data ID. */
-	FMovieSceneCompiledDataID CompiledDataID;
 	/** This sequence instances sequence ID, or MovieSceneSequenceID::Root for top-level sequences. */
 	FMovieSceneSequenceID SequenceID;
 	/** When SequenceID != MovieSceneSequenceID::Root, specifies an ID to override as a simulated root. */
@@ -307,6 +315,8 @@ private:
 	int32 PlayerIndex;
 	/** This instance's handle. */
 	FInstanceHandle InstanceHandle;
+	/** This instance's parent handle. */
+	FInstanceHandle ParentInstanceHandle;
 	/** This instance's root handle, if it is a sub sequence. */
 	FRootInstanceHandle RootInstanceHandle;
 	/** Flag that is set when this sequence has or (will be) finished. */
diff --git a/Engine/Source/Runtime/MovieScene/Public/Evaluation/MovieSceneSequenceHierarchy.h b/Engine/Source/Runtime/MovieScene/Public/Evaluation/MovieSceneSequenceHierarchy.h
--- a/Engine/Source/Runtime/MovieScene/Public/Evaluation/MovieSceneSequenceHierarchy.h
+++ b/Engine/Source/Runtime/MovieScene/Public/Evaluation/MovieSceneSequenceHierarchy.h
@@ -146,10 +146,6 @@ struct FMovieSceneSubSequenceData
 	UPROPERTY()
 	int16 HierarchicalBias;
 
-	/** Whether this sub-sequence has hierarchical easing. */
-	UPROPERTY()
-	bool bHasHierarchicalEasing;
-
 	/** Instance data that should be used for any tracks contained immediately within this sub sequence */
 	UPROPERTY()
 	FMovieSceneSequenceInstanceDataPtr InstanceData;
diff --git a/Engine/Source/Runtime/MovieScene/Public/Evaluation/MovieSceneSequenceWeights.h b/Engine/Source/Runtime/MovieScene/Public/Evaluation/MovieSceneSequenceWeights.h
new file mode 100644
--- /dev/null
+++ b/Engine/Source/Runtime/MovieScene/Public/Evaluation/MovieSceneSequenceWeights.h
@@ -0,0 +1,60 @@
+// Copyright Epic Games, Inc. All Rights Reserved.
+
+#pragma once
+
+#include "CoreTypes.h"
+#include "Containers/SortedMap.h"
+#include "UObject/WeakObjectPtrTemplates.h"
+#include "MovieSceneSequenceID.h"
+#include "EntitySystem/MovieSceneSequenceInstanceHandle.h"
+
+class UMovieSceneEntitySystemLinker;
+
+namespace UE::MovieScene
+{
+
+/**
+ * Container class for keeping track of externally created entities for providing dynamic sequence weights
+ */
+class MOVIESCENE_API FSequenceWeights
+{
+public:
+	FSequenceWeights(UMovieSceneEntitySystemLinker* InLinker, FRootInstanceHandle InRootInstanceHandle);
+	~FSequenceWeights();
+
+	/**
+	 * Set the weight of the specifieid sequence ID by creating a new HierarchicalEasingProvider entity and assigning its weight
+	 *
+	 * @param InSequenceID    The sequence ID to assign a weight to
+	 * @param Weight          The weight to assign. This is multiplied with all other weights contributing to this sequence
+	 */
+	void SetWeight(FMovieSceneSequenceID InSequenceID, double Weight = 1.0);
+
+
+	/**
+	 * Set the weights of the specifieid sequence IDs by creating new HierarchicalEasingProvider entities and assigning their weights
+	 *
+	 * @param Weights         Map containing all desired sequence weights organized by their sequence ID
+	 */
+	void SetWeights(const TSortedMap<FMovieSceneSequenceID, double>& Weights);
+
+
+	/**
+	 * Remove a previously allocated weight for the specified sequence ID
+	 */
+	void RemoveWeight(FMovieSceneSequenceID InSequenceID);
+
+private:
+
+	/** The linker that contains the playback data for this container */
+	TWeakObjectPtr<UMovieSceneEntitySystemLinker> WeakLinker;
+
+	/** Map of allocated entity IDs organized by sequence ID */
+	TSortedMap<FMovieSceneSequenceID, FMovieSceneEntityID> WeightEntitiesBySequenceID;
+
+	/** The root instance handle for our main sequence */
+	FRootInstanceHandle RootInstanceHandle;
+};
+
+
+} // namespace UE::MovieScene
\ No newline at end of file
diff --git a/Engine/Source/Runtime/MovieScene/Public/IMovieScenePlayer.h b/Engine/Source/Runtime/MovieScene/Public/IMovieScenePlayer.h
--- a/Engine/Source/Runtime/MovieScene/Public/IMovieScenePlayer.h
+++ b/Engine/Source/Runtime/MovieScene/Public/IMovieScenePlayer.h
@@ -217,6 +217,11 @@ public:
 	 */
 	virtual bool IsPreview() const { return false; }
 
+	/**
+	 * Whether this player utilizes dynamic weighting
+	 */
+	virtual bool HasDynamicWeighting() const { return false; }
+
 	/**
 	 * Called by the evaluation system when evaluation has just started.
 	 **/
diff --git a/Engine/Source/Runtime/MovieScene/Public/MovieSceneFwd.h b/Engine/Source/Runtime/MovieScene/Public/MovieSceneFwd.h
--- a/Engine/Source/Runtime/MovieScene/Public/MovieSceneFwd.h
+++ b/Engine/Source/Runtime/MovieScene/Public/MovieSceneFwd.h
@@ -98,6 +98,11 @@ enum class EMovieSceneSequenceFlags : uint8
 	 */
 	BlockingEvaluation = 1 << 1,
 
+	/**
+	 * Indicates that a sequence will utilize dynamic weighting when it is played back. Setting this flag ensures that initial values are cached correctly so the sequence can be blended into and out of
+	 */
+	DynamicWeighting = 1 << 2,
+
 	/** Symbolic entry for all flags that should be inherited by parent sequences when present on a sub sequence */
 	InheritedFlags = Volatile UMETA(Hidden),
 };
diff --git a/Engine/Source/Runtime/MovieScene/Public/MovieSceneSequencePlaybackSettings.h b/Engine/Source/Runtime/MovieScene/Public/MovieSceneSequencePlaybackSettings.h
--- a/Engine/Source/Runtime/MovieScene/Public/MovieSceneSequencePlaybackSettings.h
+++ b/Engine/Source/Runtime/MovieScene/Public/MovieSceneSequencePlaybackSettings.h
@@ -52,6 +52,7 @@ struct FMovieSceneSequencePlaybackSettings
 		, bDisableCameraCuts(false)
 		, bPauseAtEnd(false)
 		, bInheritTickIntervalFromOwner(true)
+		, bDynamicWeighting(false)
 	{ }
 
 	GENERATED_BODY()
@@ -112,6 +113,10 @@ struct FMovieSceneSequencePlaybackSettings
 	UPROPERTY(EditAnywhere, Category="Playback", meta=(InlineEditConditionToggle))
 	uint32 bInheritTickIntervalFromOwner : 1;
 
+	/** Whether to enable dynamic weighting on this player or not */
+	UPROPERTY(EditAnywhere, Category="Playback")
+	uint32 bDynamicWeighting : 1;
+
 	MOVIESCENE_API bool SerializeFromMismatchedTag(const FPropertyTag& Tag, FStructuredArchive::FSlot Slot);
 };
 
diff --git a/Engine/Source/Runtime/MovieScene/Public/MovieSceneSequencePlayer.h b/Engine/Source/Runtime/MovieScene/Public/MovieSceneSequencePlayer.h
--- a/Engine/Source/Runtime/MovieScene/Public/MovieSceneSequencePlayer.h
+++ b/Engine/Source/Runtime/MovieScene/Public/MovieSceneSequencePlayer.h
@@ -27,6 +27,11 @@
 
 class UMovieSceneSequenceTickManager;
 
+namespace UE::MovieScene
+{
+	class FSequenceWeights;
+}
+
 DECLARE_DYNAMIC_MULTICAST_DELEGATE(FOnMovieSceneSequencePlayerEvent);
 
 /**
@@ -248,6 +253,25 @@ public:
 	UFUNCTION(BlueprintCallable, Category = "Sequencer|Player")
 	FQualifiedFrameTime GetEndTime() const { return FQualifiedFrameTime(StartTime + DurationFrames, PlayPosition.GetInputRate()); }
 
+	/**
+	 * Set a manual weight to be multiplied with all blendable elements within this sequence
+	 * @note: It is recommended that a weight between 0 and 1 is supplied, though this is not enforced
+	 * @note: It is recommended that either FMovieSceneSequencePlaybackSettings::DynamicWeighting should be true for this player or the asset it's playing back should be set to enable dynamic weight to avoid undesirable behavior
+	 *
+	 * @param InWeight    The weight to suuply to all elements in this sequence
+	 */
+	UFUNCTION(BlueprintCallable, Category = "Sequencer|Player")
+	void SetWeight(double InWeight);
+
+	/**
+	 * Set a manual weight to be multiplied with all blendable elements within the specified sequence
+	 * @note: It is recommended that a weight between 0 and 1 is supplied, though this is not enforced
+	 * @note: It is recommended that either FMovieSceneSequencePlaybackSettings::DynamicWeighting should be true for this player or the asset it's playing back should be set to enable dynamic weight to avoid undesirable behavior
+	 *
+	 * @param InWeight    The weight to suuply to all elements in this sequence
+	 */
+	void SetWeight(double InWeight, FMovieSceneSequenceID SequenceID);
+
 public:
 
 	/**
@@ -515,6 +539,7 @@ protected:
 	virtual void ResolveBoundObjects(const FGuid& InBindingId, FMovieSceneSequenceID SequenceID, UMovieSceneSequence& Sequence, UObject* ResolutionContext, TArray<UObject*, TInlineAllocator<1>>& OutObjects) const override;
 	virtual IMovieScenePlaybackClient* GetPlaybackClient() override { return PlaybackClient ? &*PlaybackClient : nullptr; }
 	virtual bool IsDisablingEventTriggers(FFrameTime& DisabledUntilTime) const override;
+	virtual bool HasDynamicWeighting() const override;
 	virtual void PreEvaluation(const FMovieSceneContext& Context) override;
 	virtual void PostEvaluation(const FMovieSceneContext& Context) override;
 
@@ -643,6 +668,9 @@ protected:
 	/** Spawn register */
 	TSharedPtr<FMovieSceneSpawnRegister> SpawnRegister;
 
+	/** Sequence Weights */
+	TUniquePtr<UE::MovieScene::FSequenceWeights> SequenceWeights;
+
 	struct FServerTimeSample
 	{
 		/** The actual server sequence time in seconds, with client ping at the time of the sample baked in */
diff --git a/Engine/Source/Runtime/MovieScene/Public/Sections/MovieSceneSubSection.h b/Engine/Source/Runtime/MovieScene/Public/Sections/MovieSceneSubSection.h
--- a/Engine/Source/Runtime/MovieScene/Public/Sections/MovieSceneSubSection.h
+++ b/Engine/Source/Runtime/MovieScene/Public/Sections/MovieSceneSubSection.h
@@ -12,6 +12,7 @@
 #include "Evaluation/MovieSceneSectionParameters.h"
 #include "Evaluation/MovieSceneSequenceHierarchy.h"
 #include "Evaluation/MovieSceneSequenceTransform.h"
+#include "EntitySystem/IMovieSceneEntityProvider.h"
 #include "GameFramework/Actor.h"
 #include "HAL/Platform.h"
 #include "HAL/PlatformCrt.h"
@@ -61,6 +62,7 @@ struct FSubSequenceInstanceDataParams
 UCLASS(BlueprintType, config = EditorPerProjectUserSettings)
 class MOVIESCENE_API UMovieSceneSubSection
 	: public UMovieSceneSection
+	, public IMovieSceneEntityProvider
 {
 	GENERATED_BODY()
 
@@ -155,6 +157,9 @@ protected:
 
 	void BuildDefaultSubSectionComponents(UMovieSceneEntitySystemLinker* EntityLinker, const UE::MovieScene::FEntityImportParams& Params, UE::MovieScene::FImportedEntity* OutImportedEntity) const;
 
+	virtual bool PopulateEvaluationFieldImpl(const TRange<FFrameNumber>& EffectiveRange, const FMovieSceneEvaluationFieldEntityMetaData& InMetaData, FMovieSceneEntityComponentFieldBuilder* OutFieldBuilder) override;
+	virtual void ImportEntityImpl(UMovieSceneEntitySystemLinker* EntityLinker, const FEntityImportParams& Params, FImportedEntity* OutImportedEntity) override;
+
 public:
 
 	UPROPERTY(config, BlueprintReadWrite, EditAnywhere, Category="General", meta=(ShowOnlyInnerProperties))
diff --git a/Engine/Source/Runtime/MovieSceneTracks/Private/MovieSceneTracksComponentTypes.cpp b/Engine/Source/Runtime/MovieSceneTracks/Private/MovieSceneTracksComponentTypes.cpp
--- a/Engine/Source/Runtime/MovieSceneTracks/Private/MovieSceneTracksComponentTypes.cpp
+++ b/Engine/Source/Runtime/MovieSceneTracks/Private/MovieSceneTracksComponentTypes.cpp
@@ -159,6 +159,15 @@ void ConvertOperationalProperty(const FVector4d& In, FDoubleIntermediateVector&
 	Out = FDoubleIntermediateVector(In.X, In.Y, In.Z, In.W);
 }
 
+void ConvertOperationalProperty(float In, double& Out)
+{
+	Out = static_cast<double>(In);
+}
+
+void ConvertOperationalProperty(double In, float& Out)
+{
+	Out = static_cast<float>(In);
+}
 
 FIntermediate3DTransform GetComponentTransform(const UObject* Object)
 {
@@ -431,6 +440,22 @@ struct FColorHandler : TPropertyComponentHandler<FColorPropertyTraits, double, d
 	}
 };
 
+struct FFloatParameterHandler : TPropertyComponentHandler<FFloatParameterTraits, double>
+{
+	virtual IInitialValueProcessor* GetInitialValueProcessor() override
+	{
+		return nullptr;
+	}
+};
+
+struct FColorParameterHandler : TPropertyComponentHandler<FColorParameterTraits, double, double, double, double>
+{
+	virtual IInitialValueProcessor* GetInitialValueProcessor() override
+	{
+		return nullptr;
+	}
+};
+
 
 struct FFloatVectorHandler : TPropertyComponentHandler<FFloatVectorPropertyTraits, double, double, double, double>
 {
@@ -538,6 +563,9 @@ FMovieSceneTracksComponentTypes::FMovieSceneTracksComponentTypes()
 	ComponentRegistry->NewPropertyType(EulerTransform, TEXT("FEulerTransform"));
 	ComponentRegistry->NewPropertyType(ComponentTransform, TEXT("Component Transform"));
 
+	ComponentRegistry->NewPropertyType(FloatParameter, TEXT("float parameter"));
+	ComponentRegistry->NewPropertyType(ColorParameter, TEXT("color parameter"));
+
 	Color.MetaDataComponents.Initialize(ComponentRegistry, TEXT("Color Type"));
 	FloatVector.MetaDataComponents.Initialize(ComponentRegistry, TEXT("Num Float Vector Channels"));
 	DoubleVector.MetaDataComponents.Initialize(ComponentRegistry, TEXT("Num Double Vector Channels"));
@@ -646,6 +674,23 @@ FMovieSceneTracksComponentTypes::FMovieSceneTracksComponentTypes()
 	.SetCustomAccessors(&Accessors.Color)
 	.Commit(FColorHandler());
 
+	// --------------------------------------------------------------------------------------------
+	// Set up float parameters
+	BuiltInComponents->PropertyRegistry.DefineProperty(FloatParameter, TEXT("Apply Float Parameters"))
+	.AddSoleChannel(BuiltInComponents->DoubleResult[0])
+	.SetBlenderSystem<UMovieScenePiecewiseDoubleBlenderSystem>()
+	.Commit(FFloatParameterHandler());
+
+	// --------------------------------------------------------------------------------------------
+	// Set up color parameters
+	BuiltInComponents->PropertyRegistry.DefineCompositeProperty(ColorParameter, TEXT("Apply Color Parameters"))
+	.AddComposite(BuiltInComponents->DoubleResult[0], &FIntermediateColor::R)
+	.AddComposite(BuiltInComponents->DoubleResult[1], &FIntermediateColor::G)
+	.AddComposite(BuiltInComponents->DoubleResult[2], &FIntermediateColor::B)
+	.AddComposite(BuiltInComponents->DoubleResult[3], &FIntermediateColor::A)
+	.SetBlenderSystem<UMovieScenePiecewiseDoubleBlenderSystem>()
+	.Commit(FColorParameterHandler());
+
 	// We have some custom accessors for well-known types.
 	Accessors.Color.Add(
 			ULightComponent::StaticClass(), GET_MEMBER_NAME_CHECKED(ULightComponent, LightColor), 
diff --git a/Engine/Source/Runtime/MovieSceneTracks/Private/Systems/MovieSceneInitialValueSystem.cpp b/Engine/Source/Runtime/MovieSceneTracks/Private/Systems/MovieSceneInitialValueSystem.cpp
--- a/Engine/Source/Runtime/MovieSceneTracks/Private/Systems/MovieSceneInitialValueSystem.cpp
+++ b/Engine/Source/Runtime/MovieSceneTracks/Private/Systems/MovieSceneInitialValueSystem.cpp
@@ -120,6 +120,13 @@ void UMovieSceneInitialValueSystem::OnRun(FSystemTaskPrerequisites& InPrerequisi
 	FBuiltInComponentTypes* BuiltInComponents = FBuiltInComponentTypes::Get();
 
 	FInitialValueMutation Mutation(Linker);
+
+	// If we don't have any initial value processors, we've no work to do
+	if (Mutation.AnyInitialValue.NumComponents() == 0)
+	{
+		return;
+	}
+
 	if (Mutation.IsCached() && Linker->FindExtension<IInterrogationExtension>() == nullptr)
 	{
 		// When there is an initial value cache extension, we mutate anything with an initial value component on it by
diff --git a/Engine/Source/Runtime/MovieSceneTracks/Private/Systems/MovieSceneMaterialParameterSystem.cpp b/Engine/Source/Runtime/MovieSceneTracks/Private/Systems/MovieSceneMaterialParameterSystem.cpp
--- a/Engine/Source/Runtime/MovieSceneTracks/Private/Systems/MovieSceneMaterialParameterSystem.cpp
+++ b/Engine/Source/Runtime/MovieSceneTracks/Private/Systems/MovieSceneMaterialParameterSystem.cpp
@@ -7,6 +7,7 @@
 #include "Systems/MovieScenePiecewiseDoubleBlenderSystem.h"
 #include "Systems/MovieSceneHierarchicalBiasSystem.h"
 #include "Systems/MovieSceneMaterialSystem.h"
+#include "Systems/MovieSceneInitialValueSystem.h"
 
 #include "Materials/MaterialParameterCollectionInstance.h"
 #include "Materials/MaterialInstanceDynamic.h"
@@ -126,14 +127,159 @@ struct FApplyVectorParameters
 	}
 };
 
+struct FScalarMixin
+{
+	void CreateEntity(UMovieSceneEntitySystemLinker* Linker, UObject* BoundMaterial, FName ParameterName, TArrayView<const FMovieSceneEntityID> Inputs, FAnimatedMaterialParameterInfo* Output)
+	{
+		FBuiltInComponentTypes* BuiltInComponents = FBuiltInComponentTypes::Get();
+		FMovieSceneTracksComponentTypes* TracksComponents = FMovieSceneTracksComponentTypes::Get();
+
+		bool bHasInitialValue = false;
+		float InitialValue = 0.f;
+
+		for (FMovieSceneEntityID Input : Inputs)
+		{
+			if (TOptionalComponentReader<double> ExistingInitialValue = Linker->EntityManager.ReadComponent(Input, TracksComponents->FloatParameter.InitialValue))
+			{
+				InitialValue = static_cast<float>(*ExistingInitialValue);
+				bHasInitialValue = true;
+				break;
+			}
+		}
+
+		if (!bHasInitialValue)
+		{
+			if (UMaterialInstanceDynamic* MID = Cast<UMaterialInstanceDynamic>(BoundMaterial))
+			{
+				bHasInitialValue = MID->GetScalarParameterValue(ParameterName, InitialValue);
+			}
+			else if (UMaterialParameterCollectionInstance* MPCI = Cast<UMaterialParameterCollectionInstance>(BoundMaterial))
+			{
+				bHasInitialValue = MPCI->GetScalarParameterValue(ParameterName, InitialValue);
+			}
+		}
+
+		Output->OutputEntityID = FEntityBuilder()
+		.Add(TracksComponents->BoundMaterial, BoundMaterial)
+		.Add(BuiltInComponents->BlendChannelOutput, Output->BlendChannelID)
+		.Add(BuiltInComponents->DoubleResult[0], 0.0)
+		.AddConditional(TracksComponents->FloatParameter.InitialValue, InitialValue, bHasInitialValue)
+		.AddTag(TracksComponents->FloatParameter.PropertyTag)
+		.AddTag(BuiltInComponents->Tags.NeedsLink)
+		.AddMutualComponents()
+		.CreateEntity(&Linker->EntityManager);
+
+		Linker->EntityManager.CopyComponents(Inputs[0], Output->OutputEntityID, Linker->EntityManager.GetComponents()->GetCopyAndMigrationMask());
+	}
+
+	void InitializeSoleInput(UMovieSceneEntitySystemLinker* Linker, UObject* BoundMaterial, FName ParameterName, FMovieSceneEntityID SoleContributor, FAnimatedMaterialParameterInfo* Output)
+	{
+		if (Linker->EntityManager.HasComponent(SoleContributor, FBuiltInComponentTypes::Get()->Tags.AlwaysCacheInitialValue))
+		{
+			FMovieSceneTracksComponentTypes* TracksComponents = FMovieSceneTracksComponentTypes::Get();
+
+			float InitialValue = 0.0;
+			if (UMaterialInstanceDynamic* MID = Cast<UMaterialInstanceDynamic>(BoundMaterial))
+			{
+				MID->GetScalarParameterValue(ParameterName, InitialValue);
+				Linker->EntityManager.AddComponent(SoleContributor, TracksComponents->FloatParameter.InitialValue, InitialValue);
+			}
+			else if (UMaterialParameterCollectionInstance* MPCI = Cast<UMaterialParameterCollectionInstance>(BoundMaterial))
+			{
+				MPCI->GetScalarParameterValue(ParameterName, InitialValue);
+				Linker->EntityManager.AddComponent(SoleContributor, TracksComponents->FloatParameter.InitialValue, InitialValue);
+			}
+		}
+	}
+};
+
+struct FVectorMixin
+{
+	void CreateEntity(UMovieSceneEntitySystemLinker* Linker, UObject* BoundMaterial, FName ParameterName, TArrayView<const FMovieSceneEntityID> Inputs, FAnimatedMaterialParameterInfo* Output)
+	{
+		FBuiltInComponentTypes* BuiltInComponents = FBuiltInComponentTypes::Get();
+		FMovieSceneTracksComponentTypes* TracksComponents = FMovieSceneTracksComponentTypes::Get();
+
+		bool bHasInitialValue = false;
+		FIntermediateColor InitialValue;
+
+		for (FMovieSceneEntityID Input : Inputs)
+		{
+			if (TOptionalComponentReader<FIntermediateColor> ExistingInitialValue = Linker->EntityManager.ReadComponent(Input, TracksComponents->ColorParameter.InitialValue))
+			{
+				InitialValue = *ExistingInitialValue;
+				bHasInitialValue = true;
+				break;
+			}
+		}
+
+		if (!bHasInitialValue)
+		{
+			FLinearColor ColorValue;
+			if (UMaterialInstanceDynamic* MID = Cast<UMaterialInstanceDynamic>(BoundMaterial))
+			{
+				bHasInitialValue = MID->GetVectorParameterValue(ParameterName, ColorValue);
+			}
+			else if (UMaterialParameterCollectionInstance* MPCI = Cast<UMaterialParameterCollectionInstance>(BoundMaterial))
+			{
+				bHasInitialValue = MPCI->GetVectorParameterValue(ParameterName, ColorValue);
+			}
+
+			if (bHasInitialValue)
+			{
+				InitialValue = FIntermediateColor(ColorValue.R, ColorValue.G, ColorValue.B, ColorValue.A);
+			}
+		}
+
+		Output->OutputEntityID = FEntityBuilder()
+		.Add(TracksComponents->BoundMaterial, BoundMaterial)
+		.Add(BuiltInComponents->BlendChannelOutput, Output->BlendChannelID)
+		.Add(BuiltInComponents->DoubleResult[0], 0.0)
+		.Add(BuiltInComponents->DoubleResult[1], 0.0)
+		.Add(BuiltInComponents->DoubleResult[2], 0.0)
+		.Add(BuiltInComponents->DoubleResult[3], 0.0)
+		.AddConditional(TracksComponents->ColorParameter.InitialValue, InitialValue, bHasInitialValue)
+		.AddTag(TracksComponents->ColorParameter.PropertyTag)
+		.AddTag(BuiltInComponents->Tags.NeedsLink)
+		.AddMutualComponents()
+		.CreateEntity(&Linker->EntityManager);
+
+		Linker->EntityManager.CopyComponents(Inputs[0], Output->OutputEntityID, Linker->EntityManager.GetComponents()->GetCopyAndMigrationMask());
+	}
+
+	void InitializeSoleInput(UMovieSceneEntitySystemLinker* Linker, UObject* BoundMaterial, FName ParameterName, FMovieSceneEntityID SoleContributor, FAnimatedMaterialParameterInfo* Output)
+	{
+		if (Linker->EntityManager.HasComponent(SoleContributor, FBuiltInComponentTypes::Get()->Tags.AlwaysCacheInitialValue))
+		{
+			FMovieSceneTracksComponentTypes* TracksComponents = FMovieSceneTracksComponentTypes::Get();
+
+			FLinearColor ColorValue = FLinearColor::White;
+			if (UMaterialInstanceDynamic* MID = Cast<UMaterialInstanceDynamic>(BoundMaterial))
+			{
+				if (MID->GetVectorParameterValue(ParameterName, ColorValue))
+				{
+					Linker->EntityManager.AddComponent(SoleContributor, TracksComponents->ColorParameter.InitialValue, FIntermediateColor(ColorValue));
+				}
+			}
+			else if (UMaterialParameterCollectionInstance* MPCI = Cast<UMaterialParameterCollectionInstance>(BoundMaterial))
+			{
+				if (MPCI->GetVectorParameterValue(ParameterName, ColorValue))
+				{
+					Linker->EntityManager.AddComponent(SoleContributor, TracksComponents->ColorParameter.InitialValue, FIntermediateColor(ColorValue));
+				}
+			}
+		}
+	}
+};
+
 /** Handler that manages creation of blend outputs where there are multiple contributors for the same material parameter */
-struct FOverlappingMaterialParameterHandler
+template<typename Mixin>
+struct TOverlappingMaterialParameterHandler : Mixin
 {
 	UMovieSceneEntitySystemLinker* Linker;
 	UMovieSceneMaterialParameterSystem* System;
-	FComponentMask DefaultComponentMask;
 
-	FOverlappingMaterialParameterHandler(UMovieSceneMaterialParameterSystem* InSystem)
+	TOverlappingMaterialParameterHandler(UMovieSceneMaterialParameterSystem* InSystem)
 		: Linker(InSystem->GetLinker())
 		, System(InSystem)
 	{}
@@ -151,7 +297,13 @@ struct FOverlappingMaterialParameterHandler
 		FMovieSceneTracksComponentTypes* TracksComponents = FMovieSceneTracksComponentTypes::Get();
 
 		const int32 NumContributors = Inputs.Num();
-		if (NumContributors > 1)
+		if (!ensure(NumContributors != 0))
+		{
+			return;
+		}
+
+		const bool bUseBlending = NumContributors > 1 || !Linker->EntityManager.HasComponent(Inputs[0], BuiltInComponents->Tags.AbsoluteBlend) || Linker->EntityManager.HasComponent(Inputs[0], BuiltInComponents->WeightAndEasingResult);
+		if (bUseBlending)
 		{
 			if (!Output->OutputEntityID)
 			{
@@ -168,14 +320,7 @@ struct FOverlappingMaterialParameterHandler
 				Output->BlendChannelID = System->DoubleBlenderSystem->AllocateBlendChannel();
 
 				// Needs blending
-				Output->OutputEntityID = FEntityBuilder()
-				.Add(TracksComponents->BoundMaterial, BoundMaterial)
-				.Add(BuiltInComponents->BlendChannelOutput, Output->BlendChannelID)
-				.AddTag(BuiltInComponents->Tags.NeedsLink)
-				.AddMutualComponents()
-				.CreateEntity(&Linker->EntityManager, DefaultComponentMask);
-
-				Linker->EntityManager.CopyComponents(Inputs[0], Output->OutputEntityID, Linker->EntityManager.GetComponents()->GetCopyAndMigrationMask());
+				Mixin::CreateEntity(Linker, BoundMaterial, ParameterName, Inputs, Output);
 			}
 
 			for (FMovieSceneEntityID Input : Inputs)
@@ -191,25 +336,9 @@ struct FOverlappingMaterialParameterHandler
 				}
 			}
 		}
-		else
+		else if (!Output->OutputEntityID && Inputs.Num() == 1)
 		{
-			// Previously blended, but is no more - remove the tag from the remaining input if necessary and delete the output entity
-			if (Output->OutputEntityID)
-			{
-				Linker->EntityManager.AddComponent(Output->OutputEntityID, BuiltInComponents->Tags.NeedsUnlink);
-				Output->OutputEntityID = FMovieSceneEntityID();
-
-				if (ensure(System->DoubleBlenderSystem))
-				{
-					System->DoubleBlenderSystem->ReleaseBlendChannel(Output->BlendChannelID);
-				}
-			}
-
-			for (FMovieSceneEntityID Input : Inputs)
-			{
-				Linker->EntityManager.RemoveComponent(Input, BuiltInComponents->BlendChannelInput);
-			}
-			Output->BlendChannelID = FMovieSceneBlendChannelID();
+			Mixin::InitializeSoleInput(Linker, BoundMaterial, ParameterName, Inputs[0], Output);
 		}
 
 		Output->NumContributors = NumContributors;
@@ -255,6 +384,7 @@ UMovieSceneMaterialParameterSystem::UMovieSceneMaterialParameterSystem(const FOb
 
 		DefineImplicitPrerequisite(UMovieScenePiecewiseDoubleBlenderSystem::StaticClass(), GetClass());
 		DefineImplicitPrerequisite(GetClass(), UMovieSceneHierarchicalBiasSystem::StaticClass());
+		DefineImplicitPrerequisite(GetClass(), UMovieSceneInitialValueSystem::StaticClass());
 	}
 }
 
@@ -273,10 +403,8 @@ void UMovieSceneMaterialParameterSystem::OnUnlink()
 	// Always reset the float blender system on link to ensure that recycled systems are correctly initialized.
 	DoubleBlenderSystem = nullptr;
 
-	FOverlappingMaterialParameterHandler Handler(this);
-
-	ScalarParameterTracker.Destroy(Handler);
-	VectorParameterTracker.Destroy(Handler);
+	ScalarParameterTracker.Destroy(TOverlappingMaterialParameterHandler<FScalarMixin>(this));
+	VectorParameterTracker.Destroy(TOverlappingMaterialParameterHandler<FVectorMixin>(this));
 }
 
 void UMovieSceneMaterialParameterSystem::OnRun(FSystemTaskPrerequisites& InPrerequisites, FSystemSubsequentTasks& Subsequents)
@@ -346,13 +474,11 @@ void UMovieSceneMaterialParameterSystem::OnInstantiation()
 		.Iterate_PerAllocation(&Linker->EntityManager, HandleUnlinkedAllocation);
 
 		// Process all blended scalar parameters
-		FOverlappingMaterialParameterHandler ScalarHandler(this);
-		ScalarHandler.DefaultComponentMask.Set(BuiltInComponents->DoubleResult[0]);
+		TOverlappingMaterialParameterHandler<FScalarMixin> ScalarHandler(this);
 		ScalarParameterTracker.ProcessInvalidatedOutputs(Linker, ScalarHandler);
 
 		// Process all blended vector parameters
-		FOverlappingMaterialParameterHandler VectorHandler(this);
-		VectorHandler.DefaultComponentMask.SetAll({ BuiltInComponents->DoubleResult[0], BuiltInComponents->DoubleResult[1], BuiltInComponents->DoubleResult[2], BuiltInComponents->DoubleResult[3] });
+		TOverlappingMaterialParameterHandler<FVectorMixin> VectorHandler(this);
 		VectorParameterTracker.ProcessInvalidatedOutputs(Linker, VectorHandler);
 	}
 }
diff --git a/Engine/Source/Runtime/MovieSceneTracks/Private/Systems/MovieScenePredictionSystem.cpp b/Engine/Source/Runtime/MovieSceneTracks/Private/Systems/MovieScenePredictionSystem.cpp
--- a/Engine/Source/Runtime/MovieSceneTracks/Private/Systems/MovieScenePredictionSystem.cpp
+++ b/Engine/Source/Runtime/MovieSceneTracks/Private/Systems/MovieScenePredictionSystem.cpp
@@ -310,8 +310,7 @@ int32 UMovieSceneAsyncAction_SequencePrediction::ImportTransformEntities(UObject
 								const FMovieSceneEntityComponentField* ComponentField = CompiledDataManager->FindEntityComponentField(SubDataID);
 								if (ComponentField)
 								{
-									ImportParams.HierarchicalBias       = SubData->HierarchicalBias;
-									ImportParams.bHasHierarchicalEasing = SubData->bHasHierarchicalEasing;
+									ImportParams.HierarchicalBias = SubData->HierarchicalBias;
 
 									FFrameTime SubPredictedTime = RootPredictedTime * SubData->RootToSequenceTransform;
 									TotalNumEntities += ImportTransformEntities(Linker, ImportParams, FoundID, SubPredictedTime, ComponentField, InterrogationKey);
diff --git a/Engine/Source/Runtime/MovieSceneTracks/Private/Systems/MovieScenePropertyInstantiator.cpp b/Engine/Source/Runtime/MovieSceneTracks/Private/Systems/MovieScenePropertyInstantiator.cpp
--- a/Engine/Source/Runtime/MovieSceneTracks/Private/Systems/MovieScenePropertyInstantiator.cpp
+++ b/Engine/Source/Runtime/MovieSceneTracks/Private/Systems/MovieScenePropertyInstantiator.cpp
@@ -543,6 +543,11 @@ void UMovieScenePropertyInstantiatorSystem::InitializeFastPath(const FPropertyPa
 		break;
 	}
 
+	if (Linker->EntityManager.HasComponent(SoleContributor, BuiltInComponents->Tags.AlwaysCacheInitialValue))
+	{
+		Linker->EntityManager.AddComponent(SoleContributor, Params.PropertyDefinition->InitialValueType);
+	}
+
 	if (Params.PropertyDefinition->MetaDataTypes.Num() > 0)
 	{
 		InitializePropertyMetaDataTasks.PadToNum(Params.PropertyInfo->PropertyDefinitionIndex+1, false);
diff --git a/Engine/Source/Runtime/MovieSceneTracks/Private/Systems/WeightAndEasingEvaluatorSystem.cpp b/Engine/Source/Runtime/MovieSceneTracks/Private/Systems/WeightAndEasingEvaluatorSystem.cpp
--- a/Engine/Source/Runtime/MovieSceneTracks/Private/Systems/WeightAndEasingEvaluatorSystem.cpp
+++ b/Engine/Source/Runtime/MovieSceneTracks/Private/Systems/WeightAndEasingEvaluatorSystem.cpp
@@ -7,9 +7,14 @@
 #include "EntitySystem/EntityAllocationIterator.h"
 #include "EntitySystem/MovieSceneEntitySystemTask.h"
 #include "EntitySystem/MovieSceneEntitySystemLinker.h"
+#include "EntitySystem/MovieSceneEntityMutations.h"
+#include "EntitySystem/MovieSceneRootInstantiatorSystem.h"
 #include "Evaluation/MovieSceneEvaluationTemplateInstance.h"
+#include "Systems/MovieScenePropertyInstantiator.h"
+#include "EntitySystem/Interrogation/MovieSceneInterrogatedPropertyInstantiator.h"
 #include "Sections/MovieSceneSubSection.h"
 #include "Systems/FloatChannelEvaluatorSystem.h"
+#include "Systems/DoubleChannelEvaluatorSystem.h"
 #include "IMovieScenePlayer.h"
 #include "MovieSceneSection.h"
 
@@ -20,266 +25,488 @@
 #include UE_INLINE_GENERATED_CPP_BY_NAME(WeightAndEasingEvaluatorSystem)
 
 DECLARE_CYCLE_STAT(TEXT("MovieScene: Evaluate easing"), MovieSceneEval_EvaluateEasingTask, STATGROUP_MovieSceneECS);
+DECLARE_CYCLE_STAT(TEXT("MovieScene: Harvest easing"), MovieSceneEval_HarvestEasingTask, STATGROUP_MovieSceneECS);
 
-UMovieSceneHierarchicalEasingInstantiatorSystem::UMovieSceneHierarchicalEasingInstantiatorSystem(const FObjectInitializer& ObjInit)
-	: Super(ObjInit)
+
+namespace UE::MovieScene
 {
-	using namespace UE::MovieScene;
 
-	FBuiltInComponentTypes* BuiltInComponents = FBuiltInComponentTypes::Get();
-	RelevantComponent = BuiltInComponents->HierarchicalEasingProvider;
-	SystemCategories = EEntitySystemCategory::Core;
-}
+static constexpr uint16 INVALID_EASING_CHANNEL = uint16(-1);
 
-void UMovieSceneHierarchicalEasingInstantiatorSystem::OnRun(FSystemTaskPrerequisites& InPrerequisites, FSystemSubsequentTasks& Subsequents)
+struct FAddEasingChannelToProviderMutation : IMovieSceneEntityMutation
 {
-	using namespace UE::MovieScene;
+	FAddEasingChannelToProviderMutation(UMovieSceneHierarchicalEasingInstantiatorSystem* InSystem)
+		: System(InSystem)
+		, BuiltInComponents(FBuiltInComponentTypes::Get())
+		, InstanceRegistry(InSystem->GetLinker()->GetInstanceRegistry())
+	{}
 
-	static constexpr uint16 INVALID_EASING_CHANNEL = uint16(-1);
+private:
 
-	const FBuiltInComponentTypes* const BuiltInComponents = FBuiltInComponentTypes::Get();
-	const FInstanceRegistry* const InstanceRegistry = Linker->GetInstanceRegistry();
-	UWeightAndEasingEvaluatorSystem* const EvaluatorSystem = Linker->LinkSystem<UWeightAndEasingEvaluatorSystem>();
+	void CreateMutation(FEntityManager* EntityManager, FComponentMask* InOutEntityComponentTypes) const override
+	{
+		InOutEntityComponentTypes->Set(BuiltInComponents->HierarchicalEasingChannel);
+	}
 
-	// Step 1: Visit any new hierarchical easing providers (i.e. entities created by sub-sections with easing on them)
-	//
-	// We allocate the hierarchical easing channel for their sub-sequence.
-	//
-	auto VisitNewEasingProviders = [this, InstanceRegistry, EvaluatorSystem](const FEntityAllocation* Allocation, const FInstanceHandle* InstanceHandles, const FMovieSceneSequenceID* HierarchicalEasingProviders)
+	void InitializeAllocation(FEntityAllocation* Allocation, const FComponentMask& AllocationType) const override
 	{
+		AllocateEasingChannelsForAllocation(Allocation);
+	}
+
+	void InitializeUnmodifiedAllocation(FEntityAllocation* Allocation, const FComponentMask& AllocationType) const override
+	{
+		AllocateEasingChannelsForAllocation(Allocation);
+	}
+
+	void AllocateEasingChannelsForAllocation(FEntityAllocation* Allocation) const
+	{
+		FEntityAllocationWriteContext NewAllocation = FEntityAllocationWriteContext::NewAllocation();
+
+		TComponentLock<TWrite<uint16>>               EasingChannels      = Allocation->WriteComponents(BuiltInComponents->HierarchicalEasingChannel, NewAllocation);
+		TComponentLock<TRead<FRootInstanceHandle>>   RootInstanceHandles = Allocation->ReadComponents(BuiltInComponents->RootInstanceHandle);
+		TComponentLock<TRead<FMovieSceneSequenceID>> SubSequenceIDs      = Allocation->ReadComponents(BuiltInComponents->HierarchicalEasingProvider);
+		
 		for (int32 Index = 0; Index < Allocation->Num(); ++Index)
 		{
-			const FMovieSceneSequenceID SequenceID = HierarchicalEasingProviders[Index];
-			const FInstanceHandle SubSequenceHandle = InstanceRegistry->FindRelatedInstanceHandle(InstanceHandles[Index], SequenceID);
-
-			// We use instance handles here because sequence IDs by themselves are only unique to a single hierarchy 
-			// of sequences. If a root sequence is playing twice at the same time, there will be 2 sequence instances
-			// for the same ID...
-			// 
-			// We allocate a new easing channel on the evaluator system, add this sub-section (provider) to the list of
-			// contributors to that channel, and add the channel ID to our own map for step 2 below (this is all done
-			// inside AllocateEasingChannel).
-			//
-			// It could happen that we already had an easing channel for this sub-sequence. This can happen in editor when
-			// the user forces a re-import of the sub-section (by resizing it or whatever).
-			uint16& EasingChannel = InstanceHandleToEasingChannel.FindOrAdd(SubSequenceHandle, INVALID_EASING_CHANNEL);
-			if (EasingChannel == INVALID_EASING_CHANNEL)
+			EasingChannels[Index] = System->AllocateEasingChannel(InstanceRegistry, RootInstanceHandles[Index], SubSequenceIDs[Index]).EasingChannelID;
+		}
+	}
+private:
+
+	UMovieSceneHierarchicalEasingInstantiatorSystem* System;
+	FBuiltInComponentTypes* BuiltInComponents;
+	FInstanceRegistry* InstanceRegistry;
+};
+
+
+
+struct FAddEasingChannelToConsumerMutation : IMovieSceneConditionalEntityMutation
+{
+	FEntityManager* EntityManager;
+	FAddEasingChannelToConsumerMutation(UMovieSceneHierarchicalEasingInstantiatorSystem* InSystem)
+		: System(InSystem)
+		, BuiltInComponents(FBuiltInComponentTypes::Get())
+		, InstanceRegistry(InSystem->GetLinker()->GetInstanceRegistry())
+	{
+		EntityManager = &InSystem->GetLinker()->EntityManager;
+	}
+
+private:
+
+	void MarkAllocation(FEntityAllocation* Allocation, TBitArray<>& OutEntitiesToMutate) const override
+	{
+		TComponentReader<FRootInstanceHandle> RootInstanceHandles = Allocation->ReadComponents(BuiltInComponents->RootInstanceHandle);
+		TOptionalComponentReader<FMovieSceneSequenceID> SubSequenceIDs = Allocation->TryReadComponents(BuiltInComponents->SequenceID);
+
+		if (SubSequenceIDs)
+		{
+			for (int32 Index = 0; Index < Allocation->Num(); ++Index)
 			{
-				EasingChannel = EvaluatorSystem->AllocateEasingChannel(SubSequenceHandle);
+				const uint16 EasingChannel = System->LocateEasingChannel(RootInstanceHandles[Index], SubSequenceIDs[Index]);
+				if (EasingChannel != uint16(-1))
+				{
+					OutEntitiesToMutate.Insert(true, Index);
+				}
 			}
-			else
+		}
+		else
+		{
+			for (int32 Index = 0; Index < Allocation->Num(); ++Index)
 			{
+				const uint16 EasingChannel = System->LocateEasingChannel(RootInstanceHandles[Index], MovieSceneSequenceID::Root);
+				if (EasingChannel != uint16(-1))
+				{
+					OutEntitiesToMutate.Insert(true, Index);
+				}
 			}
 		}
-	};
+	}
 
-	FEntityTaskBuilder()
-		.Read(BuiltInComponents->InstanceHandle)
-		.Read(BuiltInComponents->HierarchicalEasingProvider)
-		.FilterAll({ BuiltInComponents->Tags.NeedsLink })
-		.Iterate_PerAllocation(&Linker->EntityManager, VisitNewEasingProviders);
+	void CreateMutation(FEntityManager* InEntityManager, FComponentMask* InOutEntityComponentTypes) const override
+	{
+		InOutEntityComponentTypes->Set(BuiltInComponents->HierarchicalEasingChannel);
 
-	// Step 2: Visit any new entities that are inside an eased-in/out sub-sequence.
-	//
-	// We need to assign them to the appropriate hierarchical easing channel that we created in step 1.
-	//
-	auto VisitNewEasings = [this, InstanceRegistry](const FEntityAllocation* Allocation, const FInstanceHandle* InstanceHandles, uint16* HierarchicalEasings)
+		InEntityManager->GetComponents()->Factories.ComputeMutuallyInclusiveComponents(*InOutEntityComponentTypes);
+
+		InOutEntityComponentTypes->Set(BuiltInComponents->Tags.NeedsLink);
+	}
+
+	void InitializeEntities(const FEntityRange& EntityRange, const FComponentMask& AllocationType) const override
 	{
-		for (int32 Index = 0; Index < Allocation->Num(); ++Index)
+		FEntityAllocationWriteContext NewAllocation = FEntityAllocationWriteContext::NewAllocation();
+
+		EntityManager->GetComponents()->Factories.RunInitializers(AllocationType, EntityRange);
+
+		TComponentWriter<uint16>                        EasingChannels      = EntityRange.Allocation->WriteComponents(BuiltInComponents->HierarchicalEasingChannel, NewAllocation);
+		TComponentReader<FRootInstanceHandle>           RootInstanceHandles = EntityRange.Allocation->ReadComponents(BuiltInComponents->RootInstanceHandle);
+		TOptionalComponentReader<FMovieSceneSequenceID> OptSequenceIDs      = EntityRange.Allocation->TryReadComponents(BuiltInComponents->SequenceID);
+
+		for (int32 Index = 0; Index < EntityRange.Num; ++Index)
 		{
-			const FInstanceHandle& InstanceHandle = InstanceHandles[Index];
-			const uint16* EasingChannel = InstanceHandleToEasingChannel.Find(InstanceHandle);
-			if (ensure(EasingChannel))
+			const int32 Offset = EntityRange.ComponentStartOffset + Index;
+
+			FMovieSceneSequenceID SequenceID = OptSequenceIDs ? OptSequenceIDs[Offset] : MovieSceneSequenceID::Root;
+			const uint16 EasingChannel = System->LocateEasingChannel(RootInstanceHandles[Offset], SequenceID);
+			check(EasingChannel != uint16(-1));
+			EasingChannels[Offset] = EasingChannel;
+		}
+	}
+private:
+
+	UMovieSceneHierarchicalEasingInstantiatorSystem* System;
+	FBuiltInComponentTypes* BuiltInComponents;
+	FInstanceRegistry* InstanceRegistry;
+};
+
+
+struct FResetFinalWeightResults
+{
+
+	static void ForEachEntity(double& Result)
+	{
+		Result = 1.f;
+	}
+};
+
+struct FEvaluateEasings
+{
+	static void ForEachEntity(FFrameTime EvalTime, const FEasingComponentData& Easing, double& Result)
+	{
+		const double EasingWeight = Easing.Section->EvaluateEasing(EvalTime);
+		Result = FMath::Max(EasingWeight, 0.f);
+	}
+};
+
+struct FAccumulateManualWeights
+{
+	static void ForEachAllocation(const FEntityAllocation* Allocation, const TReadOneOrMoreOf<double, double>& Results, double* OutAccumulatedResults)
+	{
+		using NumericType = TDecay<decltype(OutAccumulatedResults)>::Type;
+		
+		const int32 Num = Allocation->Num();
+
+		const double* WeightResults = Results.Get<0>();
+		const double* EasingResults = Results.Get<1>();
+
+		check(WeightResults || EasingResults);
+
+		// Have to do math
+		if (WeightResults && EasingResults)
+		{
+			for (int32 Index = 0; Index < Num; ++Index)
 			{
-				HierarchicalEasings[Index] = *EasingChannel;
+				OutAccumulatedResults[Index] = WeightResults[Index] * EasingResults[Index];
 			}
 		}
-	};
+		else
+		{
+			FMemory::Memcpy(OutAccumulatedResults, WeightResults ? WeightResults : EasingResults, Num * sizeof(NumericType));
+		}
+	}
+};
 
-	FEntityTaskBuilder()
-		.Read(BuiltInComponents->InstanceHandle)
-		.Write(BuiltInComponents->HierarchicalEasingChannel)
-		.FilterAll({ BuiltInComponents->Tags.NeedsLink })
-		.Iterate_PerAllocation(&Linker->EntityManager, VisitNewEasings);
 
-	// Step 3: Visit removed hierarchical easing providers, so we can free up our channels.
-	//
-	auto VisitRemovedEasingProviders = [this, InstanceRegistry, EvaluatorSystem](const FEntityAllocation* Allocation, const FInstanceHandle* InstanceHandles, const FMovieSceneSequenceID* HierarchicalEasingProviders)
+struct FHarvestHierarchicalEasings
+{
+	TSparseArray<int32>* EasingChannelToIndex;
+	TArrayView<FHierarchicalEasingChannelData> ComputationData;
+
+	FHarvestHierarchicalEasings(TSparseArray<int32>* InEasingChannelToIndex, TArray<FHierarchicalEasingChannelData>* InComputationData)
+		: EasingChannelToIndex(InEasingChannelToIndex)
+		, ComputationData(*InComputationData)
 	{
-		for (int32 Index = 0; Index < Allocation->Num(); ++Index)
+		check(InComputationData);
+	}
+
+	// Before the task runs, initialize the results array
+	void PreTask()
+	{
+		for (FHierarchicalEasingChannelData& Data : ComputationData)
 		{
-			const FSequenceInstance* RootInstance = &InstanceRegistry->GetInstance(InstanceHandles[Index]);
-			if (!RootInstance->IsRootSequence())
+			Data.FinalResult = 1.0;
+		}
+	}
+
+	// Accumulate all entities that contribute to the channel
+	void ForEachEntity(double Result, uint16 EasingChannel)
+	{
+		const int32 ResultIndex = (*EasingChannelToIndex)[EasingChannel];
+		ComputationData[ResultIndex].FinalResult *= Result;
+	}
+
+	// Multiply hierarchical weights with sub sequences
+	void PostTask()
+	{
+		// Move forward through the results array, multiplying with parents
+		// This is possible because the results array is already sorted by depth
+		for (int32 Index = 0; Index < ComputationData.Num(); ++Index)
+		{
+			FHierarchicalEasingChannelData ChannelData = ComputationData[Index];
+			if (ChannelData.ParentEasingIndex != uint16(-1))
 			{
-				RootInstance = &InstanceRegistry->GetInstance(RootInstance->GetRootInstanceHandle());
+				// The parent result has already been multiplied by all its parent weights by this point
+				ComputationData[Index].FinalResult *= ComputationData[ChannelData.ParentEasingIndex].FinalResult;
 			}
+		}
+	}
+};
+
+struct FPropagateHierarchicalEasings
+{
+	TArrayView<const FHierarchicalEasingChannelData> ComputationData;
+
+	TArray<double> HierarchicalResultsByChannelID;
+	int32 MaxChannelNum;
+
+	FPropagateHierarchicalEasings(TArrayView<const FHierarchicalEasingChannelData> InComputationData, const int32 InMaxChannelNum)
+		: ComputationData(InComputationData)
+		, MaxChannelNum(InMaxChannelNum)
+	{}
+
+	// Before the task runs, initialize the results array to avoid a double indirection during the expansion in the actual task
+	void PreTask()
+	{
+		HierarchicalResultsByChannelID.SetNumZeroed(MaxChannelNum);
+
+		for (int32 Index = 0; Index < ComputationData.Num(); ++Index)
+		{
+			const int32 ChannelIdAsIndex = ComputationData[Index].ChannelID;
+			HierarchicalResultsByChannelID[ChannelIdAsIndex] = ComputationData[Index].FinalResult;
+		}
+	}
 
-			const FMovieSceneSequenceID SubSequenceID = HierarchicalEasingProviders[Index];
-			const FInstanceHandle SubSequenceHandle = RootInstance->FindSubInstance(SubSequenceID);
+	void ForEachAllocation(const FEntityAllocation* Allocation, TRead<uint16> HierarchicalEasingChannels, TWrite<double> WeightAndEasingResults)
+	{
+		const int32 Num = Allocation->Num();
+		for (int32 Index = 0; Index < Num; ++Index)
+		{
+			const uint16 HierarchicalEasingChannel = HierarchicalEasingChannels[Index];
 
-			uint16 OutEasingChannel;
-			const bool bRemoved = InstanceHandleToEasingChannel.RemoveAndCopyValue(SubSequenceHandle, OutEasingChannel);
-			if (ensure(bRemoved))
+			if (HierarchicalEasingChannel != INVALID_EASING_CHANNEL && ensure(HierarchicalResultsByChannelID.IsValidIndex(HierarchicalEasingChannel)))
 			{
-				EvaluatorSystem->ReleaseEasingChannel(OutEasingChannel);
+				WeightAndEasingResults[Index] *= HierarchicalResultsByChannelID[HierarchicalEasingChannel];
 			}
 		}
-	};
+	}
+};
 
-	FEntityTaskBuilder()
-		.Read(BuiltInComponents->InstanceHandle)
-		.Read(BuiltInComponents->HierarchicalEasingProvider)
-		.FilterAll({ BuiltInComponents->Tags.NeedsUnlink })
-		.FilterNone({ BuiltInComponents->ParentEntity })
-		.Iterate_PerAllocation(&Linker->EntityManager, VisitRemovedEasingProviders);
 
+} // namespace UE::MovieScene
+
+
+UMovieSceneHierarchicalEasingInstantiatorSystem::UMovieSceneHierarchicalEasingInstantiatorSystem(const FObjectInitializer& ObjInit)
+	: Super(ObjInit)
+{
+	using namespace UE::MovieScene;
+
+	FBuiltInComponentTypes* BuiltInComponents = FBuiltInComponentTypes::Get();
+	SystemCategories = EEntitySystemCategory::Core;
+
+	if (HasAnyFlags(RF_ClassDefaultObject))
+	{
+		DefineImplicitPrerequisite(UMovieSceneRootInstantiatorSystem::StaticClass(), GetClass());
+		DefineImplicitPrerequisite(GetClass(), UMovieScenePropertyInstantiatorSystem::StaticClass());
+		DefineImplicitPrerequisite(GetClass(), UMovieSceneInterrogatedPropertyInstantiatorSystem::StaticClass());
+		DefineComponentConsumer(GetClass(), BuiltInComponents->BoundObject);
+	}
+}
+
+bool UMovieSceneHierarchicalEasingInstantiatorSystem::IsRelevantImpl(UMovieSceneEntitySystemLinker* InLinker) const
+{
+	using namespace UE::MovieScene;
+
+	return PersistentHandleToEasingChannel.Num() || InLinker->EntityManager.ContainsComponent(FBuiltInComponentTypes::Get()->HierarchicalEasingProvider);
 }
 
-namespace UE
+void UMovieSceneHierarchicalEasingInstantiatorSystem::OnLink()
 {
-namespace MovieScene
+	EvaluatorSystem = Linker->LinkSystem<UWeightAndEasingEvaluatorSystem>();
+	// Keep the evaluator system alive as long as we are alive
+	Linker->SystemGraph.AddReference(this, EvaluatorSystem);
+}
+
+void UMovieSceneHierarchicalEasingInstantiatorSystem::OnUnlink()
 {
+	PersistentHandleToEasingChannel.Empty();
+}
 
-struct FEvaluateEasings
+void UMovieSceneHierarchicalEasingInstantiatorSystem::OnRun(FSystemTaskPrerequisites& InPrerequisites, FSystemSubsequentTasks& Subsequents)
 {
-	UWeightAndEasingEvaluatorSystem* EvaluatorSystem;
+	using namespace UE::MovieScene;
 
-	FEvaluateEasings(UWeightAndEasingEvaluatorSystem* InEvaluatorSystem)
-		: EvaluatorSystem(InEvaluatorSystem)
+	const FBuiltInComponentTypes* const BuiltInComponents = FBuiltInComponentTypes::Get();
+	const FInstanceRegistry* const InstanceRegistry = Linker->GetInstanceRegistry();
+
+	// Step 1: Create easing channels for newly created easing providers, and add a channel ID component
 	{
-		check(EvaluatorSystem);
+		FEntityComponentFilter Filter;
+		Filter.All({ BuiltInComponents->RootInstanceHandle, BuiltInComponents->HierarchicalEasingProvider, BuiltInComponents->Tags.NeedsLink });
+		Filter.Deny({ BuiltInComponents->Tags.ImportedEntity });
+		Linker->EntityManager.MutateAll(Filter, FAddEasingChannelToProviderMutation(this));
 	}
 
-	void ForEachAllocation(
-			const FEntityAllocation* InAllocation, 
-			TRead<FFrameTime> Times,
-			TReadOptional<FEasingComponentData> OptEasing,
-			TReadOptional<double> OptManualWeights,
-			TReadOptional<FInstanceHandle> OptInstanceHandles,
-			TReadOptional<FMovieSceneSequenceID> OptHierarchicalEasingProviders,
-			TWrite<double> Results)
+	FAddEasingChannelToConsumerMutation AddEasingChannelMutation(this);
+
+	// Step 2: Add easing ID components to any new entities that exist within a blended sequence
 	{
-		const int32 Num = InAllocation->Num();
+		FEntityComponentFilter Filter;
+		Filter.All({ BuiltInComponents->InstanceHandle, BuiltInComponents->Tags.NeedsLink });
+		Filter.None({ BuiltInComponents->HierarchicalEasingChannel, BuiltInComponents->Tags.NeedsUnlink, BuiltInComponents->Tags.ImportedEntity });
+
+		Linker->EntityManager.MutateConditional(Filter, AddEasingChannelMutation);
+	}
 
-		// Initialize our result array.
-		for (int32 Idx = 0; Idx < Num; ++Idx)
+	// Step 3: Add easing ID components to any pre-existing entities that exist within a sequence that just had a channel allocated
+	for (const FHierarchicalKey& Key : NewEasingChannelKeys)
+	{
+		const FSequenceInstance& RootInstance = InstanceRegistry->GetInstance(Key.RootInstanceHandle);
+
+		const FSequenceInstance* InstanceToCheck = nullptr;
+		if (Key.SequenceID == MovieSceneSequenceID::Root)
 		{
-			Results[Idx] = 1.f;
+			InstanceToCheck = &RootInstance;
 		}
-
-		// Compute and add easing weight.
-		if (OptEasing)
+		else
 		{
-			for (int32 Idx = 0; Idx < Num; ++Idx)
+			FInstanceHandle SubInstanceHandle = RootInstance.FindSubInstance(Key.SequenceID);
+			if (SubInstanceHandle.IsValid())
 			{
-				const double EasingWeight = OptEasing[Idx].Section->EvaluateEasing(Times[Idx]);
-				Results[Idx] *= FMath::Max(EasingWeight, 0.f);
+				InstanceToCheck = &InstanceRegistry->GetInstance(SubInstanceHandle);
 			}
 		}
 
-		// Manual weight has already been computed by the float channel evaluator system, so we
-		// just need to pick up the result and combine it.
-		if (OptManualWeights)
+		if (InstanceToCheck)
 		{
-			for (int32 Idx = 0; Idx < Num; ++Idx)
-			{
-				Results[Idx] *= FMath::Max(OptManualWeights[Idx], 0.f);
-			}
+			// Find any entities that have already been linked and add the easing channel to them
+			FEntityComponentFilter Filter;
+			Filter.None({ BuiltInComponents->Tags.NeedsLink, BuiltInComponents->Tags.NeedsUnlink, BuiltInComponents->HierarchicalEasingChannel, BuiltInComponents->Tags.ImportedEntity });
+
+			InstanceToCheck->Ledger.MutateAll(Linker, Filter, AddEasingChannelMutation);
 		}
+	}
 
-		// If this is an allocation for sub-sections that provide some ease-in/out to their child sub-sequence,
-		// we store the resulting weight/easing results in the corresponding hierarhical easing channel data.
-		// This will let us later apply those values onto all entities in the hierarchy below.
-		// Sadly, this goes into random data access.
-		//
-		// Note that we need to check for instance handles because in interrogation evaluations, there are no
-		// instance handles.
-		//
-		if (OptInstanceHandles && OptHierarchicalEasingProviders)
+	NewEasingChannelKeys.Empty();
+
+
+	RemoveUnlinkedHierarchicalEasingChannels(InPrerequisites, Subsequents);
+}
+
+void UMovieSceneHierarchicalEasingInstantiatorSystem::RemoveUnlinkedHierarchicalEasingChannels(FSystemTaskPrerequisites& InPrerequisites, FSystemSubsequentTasks& Subsequents)
+{
+	using namespace UE::MovieScene;
+
+	const FBuiltInComponentTypes* const BuiltInComponents = FBuiltInComponentTypes::Get();
+	const FInstanceRegistry* const InstanceRegistry = Linker->GetInstanceRegistry();
+
+	// Step 3: Visit removed hierarchical easing providers, so we can free up our channels.
+	//         @todo: There is a risk here that the easing provider is removed before the instance has been finished.
+	//                leaving some entities hanging around that still have this channel ID assigned. This would casue an ensure in the evaluator system.
+	//                This currently can't happen though because sub-sequence easing always lasts the duration of the sub-section, and dynamic weights
+	//                are never removed once they are created (until the root sequence ends)
+	auto VisitRemovedEasingProviders = [this, InstanceRegistry](const FEntityAllocation* Allocation, const FRootInstanceHandle* RootInstanceHandles, const FMovieSceneSequenceID* SubSequenceIDs)
+	{
+		for (int32 Index = 0; Index < Allocation->Num(); ++Index)
 		{
-			const FInstanceRegistry* InstanceRegistry = EvaluatorSystem->GetLinker()->GetInstanceRegistry();
+			FHierarchicalKey Key { RootInstanceHandles[Index], SubSequenceIDs[Index] };
 
-			for (int32 Idx = 0; Idx < Num; ++Idx)
+			if (FHierarchicalInstanceData* InstanceData = this->PersistentHandleToEasingChannel.Find(Key))
 			{
-				const FSequenceInstance* RootInstance = &InstanceRegistry->GetInstance(OptInstanceHandles[Idx]);
-				if (!RootInstance->IsRootSequence())
+				if (--InstanceData->RefCount == 0)
 				{
-					RootInstance = &InstanceRegistry->GetInstance(RootInstance->GetRootInstanceHandle());
+					this->EvaluatorSystem->ReleaseEasingChannel(InstanceData->EasingChannelID);
+					this->PersistentHandleToEasingChannel.Remove(Key);
 				}
-
-				const FInstanceHandle SubSequenceHandle = RootInstance->FindSubInstance(OptHierarchicalEasingProviders[Idx]);
-				EvaluatorSystem->SetSubSequenceEasing(SubSequenceHandle, Results[Idx]);
 			}
 		}
-	}
-};
+	};
+
+	FEntityTaskBuilder()
+	.Read(BuiltInComponents->RootInstanceHandle)
+	.Read(BuiltInComponents->HierarchicalEasingProvider)
+	.FilterAll({ BuiltInComponents->Tags.NeedsUnlink })
+	.FilterNone({ BuiltInComponents->ParentEntity })
+	.Iterate_PerAllocation(&Linker->EntityManager, VisitRemovedEasingProviders);
+}
 
-struct FAccumulateHierarchicalEasings
+UMovieSceneHierarchicalEasingInstantiatorSystem::FHierarchicalInstanceData UMovieSceneHierarchicalEasingInstantiatorSystem::AllocateEasingChannel(const UE::MovieScene::FInstanceRegistry* InstanceRegistry, UE::MovieScene::FRootInstanceHandle RootInstanceHandle, FMovieSceneSequenceID SequenceID)
 {
-	TSparseArray<FHierarchicalEasingChannelData>* EasingChannels;
-	
-	FAccumulateHierarchicalEasings(TSparseArray<FHierarchicalEasingChannelData>* InEasingChannels)
-		: EasingChannels(InEasingChannels)
-	{}
+	using namespace UE::MovieScene;
+
+	const FMovieSceneSequenceHierarchy* RootHierarchy = (SequenceID != MovieSceneSequenceID::Root)
+		? InstanceRegistry->GetInstance(RootInstanceHandle).GetPlayer()->GetEvaluationTemplate().GetHierarchy()
+		: nullptr;
+
+	return AllocateEasingChannelImpl(RootInstanceHandle, SequenceID, RootHierarchy);
+}
+
+UMovieSceneHierarchicalEasingInstantiatorSystem::FHierarchicalInstanceData UMovieSceneHierarchicalEasingInstantiatorSystem::AllocateEasingChannelImpl(UE::MovieScene::FRootInstanceHandle RootInstanceHandle, FMovieSceneSequenceID SequenceID, const FMovieSceneSequenceHierarchy* Hierarchy)
+{
+	using namespace UE::MovieScene;
 
-	FORCEINLINE TStatId           GetStatId() const    { return GET_STATID(MovieSceneEval_EvaluateEasingTask); }
-	static ENamedThreads::Type    GetDesiredThread()   { return ENamedThreads::AnyThread; }
-	static ESubsequentsMode::Type GetSubsequentsMode() { return ESubsequentsMode::TrackSubsequents; }
+	FHierarchicalKey Key{ RootInstanceHandle, SequenceID };
 
-	void DoTask(ENamedThreads::Type CurrentThread, const FGraphEventRef& MyCompletionGraphEvent)
+	FHierarchicalInstanceData* ExistingInstanceData = PersistentHandleToEasingChannel.Find(Key);
+	if (ExistingInstanceData != nullptr)
 	{
-		Run();
+		++ExistingInstanceData->RefCount;
+		return *ExistingInstanceData;
 	}
 
-	void Run()
+	FHierarchicalInstanceData ParentInstanceData;
+
+	if (SequenceID != MovieSceneSequenceID::Root)
 	{
-		for (FHierarchicalEasingChannelData& EasingChannel : *EasingChannels)
-		{
-			EasingChannel.FinalEasingResult = 1.f;
-			for (FHierarchicalEasingChannelContributorData& EasingChannelContributor : EasingChannel.Contributors)
-			{
-				EasingChannel.FinalEasingResult *= EasingChannelContributor.EasingResult;
-			}
-		}
+		const FMovieSceneSequenceHierarchyNode* Node = Hierarchy->FindNode(SequenceID);
+		check(Node);
+		ParentInstanceData = AllocateEasingChannelImpl(RootInstanceHandle, Node->ParentID, Hierarchy);
 	}
-};
 
-struct FPropagateHierarchicalEasings
+	FHierarchicalInstanceData NewInstanceData;
+	NewInstanceData.RefCount = 1;
+	NewInstanceData.HierarchicalDepth = ParentInstanceData.HierarchicalDepth + 1;
+	NewInstanceData.EasingChannelID = EvaluatorSystem->AllocateEasingChannel(ParentInstanceData.EasingChannelID, NewInstanceData.HierarchicalDepth);
+
+	PersistentHandleToEasingChannel.Add(Key, NewInstanceData);
+	NewEasingChannelKeys.Add(Key);
+
+	return NewInstanceData;
+}
+
+uint16 UMovieSceneHierarchicalEasingInstantiatorSystem::LocateEasingChannel(UE::MovieScene::FRootInstanceHandle RootInstanceHandle, FMovieSceneSequenceID SequenceID) const
 {
-	TSparseArray<FHierarchicalEasingChannelData>* EasingChannels;
+	if (const FHierarchicalInstanceData* ExistingInstanceData = PersistentHandleToEasingChannel.Find(FHierarchicalKey{ RootInstanceHandle, SequenceID }))
+	{
+		return ExistingInstanceData->EasingChannelID;
+	}
+	return uint16(-1);
+}
 
-	FPropagateHierarchicalEasings(TSparseArray<FHierarchicalEasingChannelData>* InEasingChannels)
-		: EasingChannels(InEasingChannels)
-	{}
+void UMovieSceneHierarchicalEasingInstantiatorSystem::ReleaseEasingChannel(UE::MovieScene::FRootInstanceHandle RootInstanceHandle, FMovieSceneSequenceID SequenceID)
+{
+	FHierarchicalKey Key{ RootInstanceHandle, SequenceID };
 
-	void ForEachAllocation(const FEntityAllocation* Allocation, TRead<uint16> HierarchicalEasingChannels, TWrite<double> WeightAndEasingResults)
+	const FHierarchicalInstanceData* ExistingInstanceData = PersistentHandleToEasingChannel.Find(Key);
+	if (ExistingInstanceData != nullptr)
 	{
-		const int32 Num = Allocation->Num();
-		for (int32 Index = 0; Index < Num; ++Index)
-		{
-			const uint16 HierarchicalEasingChannel = HierarchicalEasingChannels[Index];
+		EvaluatorSystem->ReleaseEasingChannel(ExistingInstanceData->EasingChannelID);
 
-			if (ensure(EasingChannels->IsValidIndex(HierarchicalEasingChannel)))
-			{
-				const FHierarchicalEasingChannelData& EasingChannel = (*EasingChannels)[HierarchicalEasingChannel];
-				WeightAndEasingResults[Index] *= EasingChannel.FinalEasingResult;
-			}
-		}
+		PersistentHandleToEasingChannel.Remove(Key);
 	}
-};
+}
 
-} // namespace MovieScene
-} // namespace UE
 
 UWeightAndEasingEvaluatorSystem::UWeightAndEasingEvaluatorSystem(const FObjectInitializer& ObjInit)
 	: Super(ObjInit)
 {
 	SystemCategories = UE::MovieScene::EEntitySystemCategory::ChannelEvaluators;
 
+	bResultsNeedResort = false;
+
 	if (HasAnyFlags(RF_ClassDefaultObject))
 	{
 		DefineImplicitPrerequisite(UMovieSceneEvalTimeSystem::StaticClass(), GetClass());
 		DefineImplicitPrerequisite(UFloatChannelEvaluatorSystem::StaticClass(), GetClass());
+		DefineImplicitPrerequisite(UDoubleChannelEvaluatorSystem::StaticClass(), GetClass());
 	}
 }
 
@@ -288,63 +515,62 @@ bool UWeightAndEasingEvaluatorSystem::IsRelevantImpl(UMovieSceneEntitySystemLink
 	using namespace UE::MovieScene;
 
 	FBuiltInComponentTypes* Components = FBuiltInComponentTypes::Get();
-	return InLinker->EntityManager.ContainsAnyComponent({ Components->Easing, Components->WeightResult });
+	return InLinker->EntityManager.ContainsAnyComponent({ Components->WeightAndEasingResult });
+}
+
+void UWeightAndEasingEvaluatorSystem::OnLink()
+{
+	// Sometimes there can be easing channels left open
+	// For instance if a sub sequence caused its parent to have a
+	// hierarchical channel allocated even though it wasn't explicitly weighted.
+	// Since there are no entities left that need the channels, we can just remove them now
+	EasingChannelToIndex.Empty();
+	PreAllocatedComputationData.Empty();
 }
 
 void UWeightAndEasingEvaluatorSystem::OnUnlink()
 {
-	if (!ensure(EasingChannels.Num() == 0))
-	{
-		EasingChannels.Reset();
-	}
+	// Sometimes there can be easing channels left open
+	// For instance if a sub sequence caused its parent to have a
+	// hierarchical channel allocated even though it wasn't explicitly weighted.
+	// Since there are no entities left that need the channels, we can just remove them now
+	EasingChannelToIndex.Empty();
+	PreAllocatedComputationData.Empty();
 }
 
-uint16 UWeightAndEasingEvaluatorSystem::AllocateEasingChannel(UE::MovieScene::FInstanceHandle SubSequenceHandle)
+uint16 UWeightAndEasingEvaluatorSystem::AllocateEasingChannel(const uint16 ParentEasingChannel, const uint16 HierarchicaDepth)
 {
 	using namespace UE::MovieScene;
 
-	// Create a new data for the new channel.
-	FHierarchicalEasingChannelData NewEasingChannelData;
-	NewEasingChannelData.Contributors.Add({ SubSequenceHandle, 1.f });
+	bResultsNeedResort = true;
+
+	// Allocate the result for this channel
+	const int32 ResultIndex = PreAllocatedComputationData.Num();
+	const int32 EasingChannelID = EasingChannelToIndex.Add(ResultIndex);
 
-	const FSequenceInstance& SubSequenceInstance = Linker->GetInstanceRegistry()->GetInstance(SubSequenceHandle);
-	const FMovieSceneRootEvaluationTemplateInstance& RootEvalTemplate = SubSequenceInstance.GetPlayer()->GetEvaluationTemplate();
-	TArray<FInstanceHandle> SubSequenceParentage;
-	RootEvalTemplate.GetSequenceParentage(SubSequenceHandle, SubSequenceParentage);
-	for (const FInstanceHandle& ParentHandle : SubSequenceParentage)
+	check(EasingChannelID >= 0 && EasingChannelID < uint16(-1));
+
+	FHierarchicalEasingChannelData Channel;
+	Channel.HierarchicalDepth = HierarchicaDepth;
+	Channel.ChannelID = EasingChannelID;
+	if (ParentEasingChannel != INVALID_EASING_CHANNEL)
 	{
-		NewEasingChannelData.Contributors.Add({ ParentHandle, 1.f });
+		Channel.ParentEasingIndex = EasingChannelToIndex[ParentEasingChannel];
 	}
 
-	return EasingChannels.Add(NewEasingChannelData);
+	PreAllocatedComputationData.Add(Channel);
+	return static_cast<uint16>(EasingChannelID);
 }
 
 void UWeightAndEasingEvaluatorSystem::ReleaseEasingChannel(uint16 EasingChannelID)
 {
-	if (ensure(EasingChannels.IsValidIndex(EasingChannelID)))
-	{
-		EasingChannels.RemoveAt(EasingChannelID);
-	}
-}
+	const int32 ComputationDataIndex = EasingChannelToIndex[EasingChannelID];
 
-void UWeightAndEasingEvaluatorSystem::SetSubSequenceEasing(UE::MovieScene::FInstanceHandle SubSequenceHandle, double EasingResult)
-{
-	using namespace UE::MovieScene;
+	// Mark the channel ID as invalid, and give it the largest hierarchical depth so it gets sorted to the end
+	PreAllocatedComputationData[ComputationDataIndex].HierarchicalDepth = uint16(-1);
+	EasingChannelToIndex.RemoveAt(EasingChannelID);
 
-	// The given sub-sequence has been assigned the given easing value. We can copy that value everywhere this
-	// sub-sequence is used in a channel, i.e. for the channel of the sub-sequence itself, but also for the channels
-	// of any children sub-sequences under that sub-sequence.
-	for (FHierarchicalEasingChannelData& EasingChannel : EasingChannels)
-	{
-		for (FHierarchicalEasingChannelContributorData& EasingChannelContributor : EasingChannel.Contributors)
-		{
-			if (EasingChannelContributor.SubSequenceHandle == SubSequenceHandle)
-			{
-				EasingChannelContributor.EasingResult = EasingResult;
-				break;
-			}
-		}
-	}
+	bResultsNeedResort = true;
 }
 
 void UWeightAndEasingEvaluatorSystem::OnRun(FSystemTaskPrerequisites& InPrerequisites, FSystemSubsequentTasks& Subsequents)
@@ -354,56 +580,100 @@ void UWeightAndEasingEvaluatorSystem::OnRun(FSystemTaskPrerequisites& InPrerequi
 	FBuiltInComponentTypes* Components = FBuiltInComponentTypes::Get();
 	FInstanceRegistry* InstanceRegistry = Linker->GetInstanceRegistry();
 
-	// Step 1: Compute all the easings and weights of all entities that have any.
+	// No hierarchical weighting, just reset everything to 1.0
+	FGraphEventRef ResetWeights = FEntityTaskBuilder()
+	.Write(Components->WeightAndEasingResult)
+	.SetStat(GET_STATID(MovieSceneEval_EvaluateEasingTask))
+	.Dispatch_PerEntity<FResetFinalWeightResults>(&Linker->EntityManager, InPrerequisites, &Subsequents);
+
+	FSystemTaskPrerequisites ResetWeightsDependencies {InPrerequisites};
+	ResetWeightsDependencies.AddComponentTask(Components->WeightAndEasingResult, ResetWeights);
+
+	// Step 1: Evaluate section easing and manual weights in parallel
 	//
-	FGraphEventRef EvalTask = FEntityTaskBuilder()
-		// We need the eval time to evaluate easing curves.
-		.Read(Components->EvalTime)
-		.ReadOptional(Components->Easing)
-		// We may need to multiply easing and manual weight together.
-		.ReadOptional(Components->WeightResult)
-		// For hierarchical easing we need the following 2 components... InstanceHandle is optional
-		// because in interrogation evaluations, there are not instance handles.
-		.ReadOptional(Components->InstanceHandle)
-		.ReadOptional(Components->HierarchicalEasingProvider)
-		// We will write the result to a separate component.
+	FGraphEventRef EvaluateEasing = FEntityTaskBuilder()
+	.Read(Components->EvalTime)
+	.Read(Components->Easing)
+	.Write(Components->EasingResult)
+	.SetStat(GET_STATID(MovieSceneEval_EvaluateEasingTask))
+	.Dispatch_PerEntity<FEvaluateEasings>(&Linker->EntityManager, ResetWeightsDependencies, &Subsequents);
+
+	FGraphEventRef AccumulateManualWeights = FEntityTaskBuilder()
+	.ReadOneOrMoreOf(Components->WeightResult, Components->EasingResult)
+	.Write(Components->WeightAndEasingResult)
+	.SetStat(GET_STATID(MovieSceneEval_EvaluateEasingTask))
+	.Dispatch_PerAllocation<FAccumulateManualWeights>(&Linker->EntityManager, ResetWeightsDependencies, &Subsequents);
+
+	// If we have hierarchical easing, we initialize all the weights to their hierarchical defaults
+	if (PreAllocatedComputationData.Num() > 0)
+	{
+		ResortComputationBuffer();
+
+		FSystemTaskPrerequisites HarvestPrereqs {InPrerequisites};
+		HarvestPrereqs.AddComponentTask(Components->WeightAndEasingResult, EvaluateEasing);
+		HarvestPrereqs.AddComponentTask(Components->WeightAndEasingResult, AccumulateManualWeights);
+
+		// Step 2: Harvest any hierarchical results from providers
+		//
+		FGraphEventRef HarvestTask = FEntityTaskBuilder()
+		.Read(Components->WeightAndEasingResult)
+		.Read(Components->HierarchicalEasingChannel)
+		.FilterAll({ Components->HierarchicalEasingProvider })  // Only harvest results from entities that are providing results
+		.SetStat(GET_STATID(MovieSceneEval_HarvestEasingTask))
+		.Dispatch_PerEntity<FHarvestHierarchicalEasings>(&Linker->EntityManager, HarvestPrereqs, nullptr, &EasingChannelToIndex, &PreAllocatedComputationData);
+
+		FSystemTaskPrerequisites PropagatePrereqs {InPrerequisites};
+		PropagatePrereqs.AddRootTask(HarvestTask);
+
+		// Step 3: Apply hierarchical easing results to all entities inside affected sub-sequences.
+		//
+		FEntityTaskBuilder()
+		.Read(Components->HierarchicalEasingChannel)
 		.Write(Components->WeightAndEasingResult)
+		.FilterNone({ Components->HierarchicalEasingProvider }) // Do not propagate hierarchical weights onto providers!
 		.SetStat(GET_STATID(MovieSceneEval_EvaluateEasingTask))
-		.Dispatch_PerAllocation<FEvaluateEasings>(&Linker->EntityManager, InPrerequisites, &Subsequents, this);
+		.Dispatch_PerAllocation<FPropagateHierarchicalEasings>(&Linker->EntityManager, PropagatePrereqs, &Subsequents,
+			PreAllocatedComputationData, EasingChannelToIndex.GetMaxIndex());
+	}
+}
+
+void UWeightAndEasingEvaluatorSystem::ResortComputationBuffer()
+{
+	using namespace UE::MovieScene;
 
-	// If we have no hierarchical easing, there's only one step... otherwise, we have more work to do.
-	
-	if (EasingChannels.Num() > 0)
+	if (!bResultsNeedResort)
 	{
-		// Step 2: Gather and compute sub-sequences' hierarchical easing results.
-		//
-		// Now, some of the entities we processed above happen to be representing sub-sections which contain entire sub-sequences.
-		// We need to take their weight/easing result and propagate it to all the entities in these sub-sequences, and keep
-		// propagating that down the hierarchy.
-		//
-		FSystemTaskPrerequisites PropagatePrereqs;
+		return;
+	}
+
+	bResultsNeedResort = false;
+
+	TSparseArray<int32> OldToNewIndex;
+
+	// Resort the results array by depth if it has been modified
+	Algo::SortBy(PreAllocatedComputationData, &FHierarchicalEasingChannelData::HierarchicalDepth);
+
+	for (int32 Index = 0; Index < PreAllocatedComputationData.Num(); ++Index)
+	{
+		FHierarchicalEasingChannelData& ComputationData = PreAllocatedComputationData[Index];
 
-		if (Linker->EntityManager.GetThreadingModel() == EEntityThreadingModel::NoThreading)
+		// As soon as we find an invalid hierarchical depth, everything proceeding this index is garbage
+		if (ComputationData.HierarchicalDepth == uint16(-1))
 		{
-			FAccumulateHierarchicalEasings(&EasingChannels).Run();
+			PreAllocatedComputationData.RemoveAt(Index, PreAllocatedComputationData.Num() - Index);
+			break;
 		}
-		else
-		{
-			FGraphEventArray AccumulatePrereqs { EvalTask };
-			FGraphEventRef AccumulateTask = TGraphTask<FAccumulateHierarchicalEasings>::CreateTask(&AccumulatePrereqs, Linker->EntityManager.GetDispatchThread())
-				.ConstructAndDispatchWhenReady(&EasingChannels);
 
-			PropagatePrereqs.AddRootTask(AccumulateTask);
+		// Reassign the channel ID to index mapping
+		int32& ChannelIndex = EasingChannelToIndex[ComputationData.ChannelID];
+		OldToNewIndex.Insert(ChannelIndex, Index);
+		ChannelIndex = Index;
+		
+		if (ComputationData.ParentEasingIndex != uint16(-1))
+		{
+			// Parent Index must have been added to the OldToNewIndex map by now because the results are always sorted
+			ComputationData.ParentEasingIndex = OldToNewIndex[ComputationData.ParentEasingIndex];
 		}
-
-		// Step 3: Apply hierarchical easing results to all entities inside affected sub-sequences.
-		//
-		FEntityTaskBuilder()
-			.Read(Components->HierarchicalEasingChannel)
-			.Write(Components->WeightAndEasingResult)
-			.SetStat(GET_STATID(MovieSceneEval_EvaluateEasingTask))
-			.Dispatch_PerAllocation<FPropagateHierarchicalEasings>(&Linker->EntityManager, PropagatePrereqs, &Subsequents, &EasingChannels);
 	}
 }
 
-
diff --git a/Engine/Source/Runtime/MovieSceneTracks/Private/Tracks/MovieScenePrimitiveMaterialTrack.cpp b/Engine/Source/Runtime/MovieSceneTracks/Private/Tracks/MovieScenePrimitiveMaterialTrack.cpp
--- a/Engine/Source/Runtime/MovieSceneTracks/Private/Tracks/MovieScenePrimitiveMaterialTrack.cpp
+++ b/Engine/Source/Runtime/MovieSceneTracks/Private/Tracks/MovieScenePrimitiveMaterialTrack.cpp
@@ -12,6 +12,9 @@ UMovieScenePrimitiveMaterialTrack::UMovieScenePrimitiveMaterialTrack(const FObje
 	: Super(ObjInit)
 {
 	MaterialIndex = 0;
+	SupportedBlendTypes.Add(EMovieSceneBlendType::Absolute);
+	SupportedBlendTypes.Add(EMovieSceneBlendType::Additive);
+
 #if WITH_EDITORONLY_DATA
 	TrackTint = FColor(64,192,64,75);
 #endif
diff --git a/Engine/Source/Runtime/MovieSceneTracks/Public/MovieSceneTracksComponentTypes.h b/Engine/Source/Runtime/MovieSceneTracks/Public/MovieSceneTracksComponentTypes.h
--- a/Engine/Source/Runtime/MovieSceneTracks/Public/MovieSceneTracksComponentTypes.h
+++ b/Engine/Source/Runtime/MovieSceneTracks/Public/MovieSceneTracksComponentTypes.h
@@ -397,6 +397,8 @@ using FTransformPropertyTraits          = TIndirectPropertyTraits<FTransform, FI
 using FEulerTransformPropertyTraits     = TIndirectPropertyTraits<FEulerTransform, FIntermediate3DTransform>;
 using FComponentTransformPropertyTraits = TDirectPropertyTraits<FIntermediate3DTransform>;
 
+using FFloatParameterTraits             = TIndirectPropertyTraits<float, double>;
+using FColorParameterTraits             = TIndirectPropertyTraits<FLinearColor, FIntermediateColor>;
 
 struct MOVIESCENETRACKS_API FMovieSceneTracksComponentTypes
 {
@@ -414,6 +416,10 @@ struct MOVIESCENETRACKS_API FMovieSceneTracksComponentTypes
 	TPropertyComponents<FTransformPropertyTraits> Transform;
 	TPropertyComponents<FEulerTransformPropertyTraits> EulerTransform;
 	TPropertyComponents<FComponentTransformPropertyTraits> ComponentTransform;
+
+	TPropertyComponents<FFloatParameterTraits> FloatParameter;
+	TPropertyComponents<FColorParameterTraits> ColorParameter;
+
 	TComponentTypeID<FSourceDoubleChannel> QuaternionRotationChannel[3];
 
 	TComponentTypeID<FConstraintComponentData> ConstraintChannel;
diff --git a/Engine/Source/Runtime/MovieSceneTracks/Public/MovieSceneTracksPropertyTypes.h b/Engine/Source/Runtime/MovieSceneTracks/Public/MovieSceneTracksPropertyTypes.h
--- a/Engine/Source/Runtime/MovieSceneTracks/Public/MovieSceneTracksPropertyTypes.h
+++ b/Engine/Source/Runtime/MovieSceneTracks/Public/MovieSceneTracksPropertyTypes.h
@@ -190,6 +190,9 @@ MOVIESCENETRACKS_API FIntermediate3DTransform GetComponentTransform(const UObjec
 MOVIESCENETRACKS_API void SetComponentTransform(USceneComponent* SceneComponent, const FIntermediate3DTransform& InTransform);
 MOVIESCENETRACKS_API void SetComponentTransformAndVelocity(UObject* Object, const FIntermediate3DTransform& InTransform);
 
+MOVIESCENETRACKS_API void ConvertOperationalProperty(float In, double& Out);
+MOVIESCENETRACKS_API void ConvertOperationalProperty(double In, float& Out);
+
 MOVIESCENETRACKS_API void ConvertOperationalProperty(const FIntermediate3DTransform& In, FEulerTransform& Out);
 MOVIESCENETRACKS_API void ConvertOperationalProperty(const FEulerTransform& In, FIntermediate3DTransform& Out);
 MOVIESCENETRACKS_API void ConvertOperationalProperty(const FIntermediate3DTransform& In, FTransform& Out);
diff --git a/Engine/Source/Runtime/MovieSceneTracks/Public/Systems/WeightAndEasingEvaluatorSystem.h b/Engine/Source/Runtime/MovieSceneTracks/Public/Systems/WeightAndEasingEvaluatorSystem.h
--- a/Engine/Source/Runtime/MovieSceneTracks/Public/Systems/WeightAndEasingEvaluatorSystem.h
+++ b/Engine/Source/Runtime/MovieSceneTracks/Public/Systems/WeightAndEasingEvaluatorSystem.h
@@ -22,53 +22,134 @@
 
 class UMovieSceneEntitySystemLinker;
 class UMovieSceneEvalTimeSystem;
+class UWeightAndEasingEvaluatorSystem;
 class UObject;
 struct FMovieSceneSubSequenceData;
 
-namespace UE
-{
-namespace MovieScene
+namespace UE::MovieScene
 {
 
-	struct FHierarchicalEasingChannelContributorData
-	{
-		UE::MovieScene::FInstanceHandle SubSequenceHandle;
-		double EasingResult;
-	};
+struct FAddEasingChannelToProviderMutation;
 
-	struct FHierarchicalEasingChannelData
-	{
-		TArray<FHierarchicalEasingChannelContributorData, TInlineAllocator<8>> Contributors;
-		double FinalEasingResult;
-	};
+/** Computation data used for accumulating hierarchical weights for sub sequences */
+struct FHierarchicalEasingChannelData
+{
+	/** Our parent's computation data within PreAllocatedComputationData. Must only access if bResultsNeedResort is false. */
+	uint16 ParentEasingIndex = uint16(-1);
+	/** The accumulated hierarchical depth of this sequence within its root. */
+	uint16 HierarchicalDepth = 0;
+	/**
+	 * The easing channel ID for this data. This reprensents the index within EasingChannelToIndex that uniquely identifies our channel.
+	 * Any entity within this sequence that is subject to easing will contain a HierarchicalEasingChannel component with this ID
+	 */
+	uint16 ChannelID = uint16(-1);
+	/** The final result of this easing channel, accumulated with all parents */
+	double FinalResult = 1.0;
+};
 
-}  // namespace MovieScene
-}  // namespace UE
+}  // namespace UE::MovieScene
 
 /**
- * System that creates hierarchical easing channels
+ * System that creates hierarchical easing channels for any newly introduced HierarchicalEasingProvider components
  */
 UCLASS()
 class MOVIESCENETRACKS_API UMovieSceneHierarchicalEasingInstantiatorSystem : public UMovieSceneEntityInstantiatorSystem
 {
 public:
 
+	struct FHierarchicalInstanceData
+	{
+		int32 HierarchicalDepth = -1;
+		uint16 RefCount = 0;
+		uint16 EasingChannelID = uint16(-1);
+	};
+
 	GENERATED_BODY()
 
 	UMovieSceneHierarchicalEasingInstantiatorSystem(const FObjectInitializer& ObjInit);
 
+	/** 
+	 * Locate an already allocated easing channel for the specified sequence instance and sub sequence ID.
+	 *
+	 * @param RootInstanceHandle The instance handle of the root sequence
+	 * @param SequenceID         The Sequence ID of the sub sequence
+	 * @return The sequence's easing channel or 0xFFFF if one does not exist
+	 */
+	uint16 LocateEasingChannel(UE::MovieScene::FRootInstanceHandle RootInstanceHandle, FMovieSceneSequenceID SequenceID) const;
+
+	/**
+	 * Remove any hierarchical easing channels for channels that are no longer needed because of unlinked easing providers
+	 */
+	void RemoveUnlinkedHierarchicalEasingChannels(FSystemTaskPrerequisites& InPrerequisites, FSystemSubsequentTasks& Subsequents);
+
 private:
 
+	virtual bool IsRelevantImpl(UMovieSceneEntitySystemLinker* InLinker) const;
+	virtual void OnLink() override;
+	virtual void OnUnlink() override;
 	virtual void OnRun(FSystemTaskPrerequisites& InPrerequisites, FSystemSubsequentTasks& Subsequents) override;
 
 private:
 
+	/** 
+	 * Allocate a new easing channel for the specified sequence instance and sub sequence ID. Will return an existing channel if it has already been allocated.
+	 *
+	 * @param InstanceRegistry   Pointer to the instance registry for retrievable of sequence instances (see UMovieSceneEntitySystemLinbker::GetInstanceRegistry)
+	 * @param RootInstanceHandle The instance handle of the root sequence
+	 * @param SequenceID         The Sequence ID of the sub sequence to make a channel for, or MovieSceneSequenceID::Root for the root sequence
+	 * @return Structure containing both the easing channel ID and this sequence's hierarchical depth
+	 */
+	FHierarchicalInstanceData AllocateEasingChannel(const UE::MovieScene::FInstanceRegistry* InstanceRegistry, UE::MovieScene::FRootInstanceHandle RootInstanceHandle, FMovieSceneSequenceID SequenceID);
+
+	/** 
+	 * Release an easing channel for the specified sequence instance and sub sequence ID if one exists
+	 *
+	 * @param RootInstanceHandle The instance handle of the root sequence
+	 * @param SequenceID         The Sequence ID of the sub sequence
+	 */
+	void ReleaseEasingChannel(UE::MovieScene::FRootInstanceHandle RootInstanceHandle, FMovieSceneSequenceID SequenceID);
+
+	/** 
+	 * Allocate a new easing channel for the specified sequence instance and sub sequence ID, and all its parents. Will return an existing channel if it has already been allocated.
+	 *
+	 * @param RootInstanceHandle The instance handle of the root sequence
+	 * @param SequenceID         The Sequence ID of the sub sequence to make a channel for, or MovieSceneSequenceID::Root for the root sequence
+	 * @param Hierarchy          The sequence hierarchy (or nullptr if this is a root sequence and one is not available)
+	 * @return Structure containing both the easing channel ID and this sequence's hierarchical depth
+	 */
+	FHierarchicalInstanceData AllocateEasingChannelImpl(UE::MovieScene::FRootInstanceHandle RootInstanceHandle, FMovieSceneSequenceID SequenceID, const FMovieSceneSequenceHierarchy* Hierarchy);
+
+private:
+
+	friend UE::MovieScene::FAddEasingChannelToProviderMutation;
+
+	/** Hierarchy key for associating channel IDs to pairs of instance handles and sequence IDs */
+	struct FHierarchicalKey
+	{
+		UE::MovieScene::FRootInstanceHandle RootInstanceHandle;
+		FMovieSceneSequenceID SequenceID;
+
+		friend uint32 GetTypeHash(const FHierarchicalKey& In)
+		{
+			return HashCombine(GetTypeHash(In.RootInstanceHandle), GetTypeHash(In.SequenceID));
+		}
+		friend bool operator==(const FHierarchicalKey& A, const FHierarchicalKey& B)
+		{
+			return A.RootInstanceHandle == B.RootInstanceHandle && A.SequenceID == B.SequenceID;
+		}
+	};
+
 	/** Map between a sub-sequence handle and the easing channel affecting it */
-	TMap<UE::MovieScene::FInstanceHandle, uint16> InstanceHandleToEasingChannel;
+	TMap<FHierarchicalKey, FHierarchicalInstanceData> PersistentHandleToEasingChannel;
+	/** Set of newly created easing channels this frame. This is used to add easing data to existing entities. */
+	TSet<FHierarchicalKey> NewEasingChannelKeys;
+
+	UPROPERTY()
+	TObjectPtr<UWeightAndEasingEvaluatorSystem> EvaluatorSystem;
 };
 
 /**
- * System that is responsible for evaluating ease in/out factors.
+ * System that combines manual weights and easings and propagates them to entities with matching EasingChannelID components
  */
 UCLASS()
 class MOVIESCENETRACKS_API UWeightAndEasingEvaluatorSystem : public UMovieSceneEntitySystem
@@ -79,19 +160,27 @@ public:
 
 	UWeightAndEasingEvaluatorSystem(const FObjectInitializer& ObjInit);
 
-	uint16 AllocateEasingChannel(UE::MovieScene::FInstanceHandle SubSequenceHandle);
+	uint16 AllocateEasingChannel(const uint16 ParentEasingChannel, const uint16 HierarchicaDepth);
 	void ReleaseEasingChannel(uint16 EasingChannelID);
 
-	void SetSubSequenceEasing(UE::MovieScene::FInstanceHandle SubSequenceHandle, double EasingResult);
-
 private:
 
 	virtual void OnRun(FSystemTaskPrerequisites& InPrerequisites, FSystemSubsequentTasks& Subsequents) override;
 	virtual bool IsRelevantImpl(UMovieSceneEntitySystemLinker* InLinker) const override;
+	virtual void OnLink() override final;
 	virtual void OnUnlink() override final;
 
+	void ResortComputationBuffer();
+
 private:
 
-	TSparseArray<UE::MovieScene::FHierarchicalEasingChannelData> EasingChannels;
+	/** Stable sparse array of indices into PreAllocatedEasingResults for each easing channel */
+	TSparseArray<int32> EasingChannelToIndex;
+
+	/** Unstable array of preallocated storage for computing easing results sorted by hierarchical depth. */
+	TArray<UE::MovieScene::FHierarchicalEasingChannelData> PreAllocatedComputationData;
+
+	/** True if the preallocated easing results need resorting. */
+	bool bResultsNeedResort;
 };
 
