Sequencer: re-add evaluation of weight channels

diff --git a/Engine/Source/Runtime/MovieSceneTracks/Private/Systems/FloatChannelEvaluatorSystem.cpp b/Engine/Source/Runtime/MovieSceneTracks/Private/Systems/FloatChannelEvaluatorSystem.cpp
--- a/Engine/Source/Runtime/MovieSceneTracks/Private/Systems/FloatChannelEvaluatorSystem.cpp
+++ b/Engine/Source/Runtime/MovieSceneTracks/Private/Systems/FloatChannelEvaluatorSystem.cpp
@@ -161,6 +161,7 @@ UFloatChannelEvaluatorSystem::UFloatChannelEvaluatorSystem(const FObjectInitiali
 		{
 			RegisterChannelType(Components->FloatChannel[Index], Components->CachedInterpolation[Index], Components->DoubleResult[Index]);
 		}
+		RegisterChannelType(Components->WeightChannel, Components->CachedWeightChannelInterpolation, Components->WeightResult);
 	}
 }
 
