Sequencer: Fix unitialized bits

#rnx

diff --git a/Engine/Source/Runtime/MovieSceneTracks/Private/Systems/WeightAndEasingEvaluatorSystem.cpp b/Engine/Source/Runtime/MovieSceneTracks/Private/Systems/WeightAndEasingEvaluatorSystem.cpp
--- a/Engine/Source/Runtime/MovieSceneTracks/Private/Systems/WeightAndEasingEvaluatorSystem.cpp
+++ b/Engine/Source/Runtime/MovieSceneTracks/Private/Systems/WeightAndEasingEvaluatorSystem.cpp
@@ -105,7 +105,8 @@ private:
 				const uint16 EasingChannel = System->LocateEasingChannel(RootInstanceHandles[Index], SubSequenceIDs[Index]);
 				if (EasingChannel != uint16(-1))
 				{
-					OutEntitiesToMutate.Insert(true, Index);
+					OutEntitiesToMutate.PadToNum(Index + 1, false);
+					OutEntitiesToMutate[Index] = true;
 				}
 			}
 		}
@@ -116,7 +117,8 @@ private:
 				const uint16 EasingChannel = System->LocateEasingChannel(RootInstanceHandles[Index], MovieSceneSequenceID::Root);
 				if (EasingChannel != uint16(-1))
 				{
-					OutEntitiesToMutate.Insert(true, Index);
+					OutEntitiesToMutate.PadToNum(Index + 1, false);
+					OutEntitiesToMutate[Index] = true;
 				}
 			}
 		}
