Sequencer: fix transform and weight system dependencies

- The component transform system needs to depend on anything that produces all 9 double results. This includes the double channel evaluators but also the quaternion evaluation system. Instead of adding it as an explicit dependency, we just change the code to make the system into a consumer of all 9 double results. Since both the double channel evaluator and the quaternion evaluator are specified as producers of those, it works out.

- The weight and easing system was depending on the float and double evaluators for no reason. It just needs to depend on whoever is creating the weight result (generally a float channel, but possibly an overriden channel too in some near future).

diff --git a/Engine/Source/Runtime/MovieSceneTracks/Private/Systems/MovieSceneComponentTransformSystem.cpp b/Engine/Source/Runtime/MovieSceneTracks/Private/Systems/MovieSceneComponentTransformSystem.cpp
--- a/Engine/Source/Runtime/MovieSceneTracks/Private/Systems/MovieSceneComponentTransformSystem.cpp
+++ b/Engine/Source/Runtime/MovieSceneTracks/Private/Systems/MovieSceneComponentTransformSystem.cpp
@@ -4,6 +4,7 @@
 #include "Systems/DoubleChannelEvaluatorSystem.h"
 #include "Systems/MovieScenePiecewiseDoubleBlenderSystem.h"
 #include "Systems/MovieScenePropertyInstantiator.h"
+#include "Systems/MovieSceneQuaternionInterpolationRotationSystem.h"
 
 #include "EntitySystem/BuiltInComponentTypes.h"
 #include "EntitySystem/Interrogation/MovieSceneInterrogationLinker.h"
@@ -20,14 +21,21 @@ UMovieSceneComponentTransformSystem::UMovieSceneComponentTransformSystem(const F
 	// This system can be used for interrogation
 	SystemCategories &= ~FSystemInterrogator::GetExcludedFromInterrogationCategory();
 
-	BindToProperty(FMovieSceneTracksComponentTypes::Get()->ComponentTransform);
+	const FBuiltInComponentTypes* BuiltInComponents = FBuiltInComponentTypes::Get();
+	const FMovieSceneTracksComponentTypes* TrackComponents = FMovieSceneTracksComponentTypes::Get();
+
+	BindToProperty(TrackComponents->ComponentTransform);
 
 	if (HasAnyFlags(RF_ClassDefaultObject))
 	{
 		DefineImplicitPrerequisite(UMovieScenePiecewiseDoubleBlenderSystem::StaticClass(), GetClass());
-		DefineImplicitPrerequisite(UDoubleChannelEvaluatorSystem::StaticClass(), GetClass());
 
-		DefineComponentConsumer(GetClass(), UE::MovieScene::FMovieSceneTracksComponentTypes::Get()->ComponentTransform.PropertyTag);
+		DefineComponentConsumer(GetClass(), TrackComponents->ComponentTransform.PropertyTag);
+
+		for (int32 Index = 0; Index < UE_ARRAY_COUNT(BuiltInComponents->DoubleResult); ++Index)
+		{
+			DefineComponentConsumer(GetClass(), BuiltInComponents->DoubleResult[Index]);
+		}
 	}
 }
 
diff --git a/Engine/Source/Runtime/MovieSceneTracks/Private/Systems/WeightAndEasingEvaluatorSystem.cpp b/Engine/Source/Runtime/MovieSceneTracks/Private/Systems/WeightAndEasingEvaluatorSystem.cpp
--- a/Engine/Source/Runtime/MovieSceneTracks/Private/Systems/WeightAndEasingEvaluatorSystem.cpp
+++ b/Engine/Source/Runtime/MovieSceneTracks/Private/Systems/WeightAndEasingEvaluatorSystem.cpp
@@ -500,15 +500,21 @@ void UMovieSceneHierarchicalEasingInstantiatorSystem::ReleaseEasingChannel(UE::M
 UWeightAndEasingEvaluatorSystem::UWeightAndEasingEvaluatorSystem(const FObjectInitializer& ObjInit)
 	: Super(ObjInit)
 {
+	using namespace UE::MovieScene;
+
 	SystemCategories = UE::MovieScene::EEntitySystemCategory::ChannelEvaluators;
 
 	bResultsNeedResort = false;
 
 	if (HasAnyFlags(RF_ClassDefaultObject))
 	{
+		const FBuiltInComponentTypes* BuiltInComponents = FBuiltInComponentTypes::Get();
+
+		DefineComponentConsumer(GetClass(), BuiltInComponents->WeightResult);
+
 		DefineImplicitPrerequisite(UMovieSceneEvalTimeSystem::StaticClass(), GetClass());
-		DefineImplicitPrerequisite(UFloatChannelEvaluatorSystem::StaticClass(), GetClass());
-		DefineImplicitPrerequisite(UDoubleChannelEvaluatorSystem::StaticClass(), GetClass());
+		DefineComponentProducer(GetClass(), BuiltInComponents->EasingResult);
+		DefineComponentProducer(GetClass(), BuiltInComponents->WeightAndEasingResult);
 	}
 }
 
