Sequencer: skeletal mesh animation in ECS

This change moves all the skeletal mesh animation code from the legacy sequencer evaluation API (evaluation templates) to the new evaluation API (ECS system). The upgrade is done in as simple a way as possible, in order not to mix extra changes with all the moved code... although there is *one* fix done on MapTimeToAnimation, in which bLooping was incorrectly computed when it came to start/end offsets into the animation clip. Everything else is pretty much the same:

- Blend accumulation is now done "by hand", especially since we were actually not blending anything: we were just accumulating active animations into an array, and passing that onto the animation system.
- Pre-animated state is now saved using the new storage API, but the logic of restoring that state hasn't changed.

The differences are as follows:

- The new ECS system has a global view on all animations running across all playing sequences, not just inside on sequence. As such, because some of the previous code was referencing the current sequence player, we have an extra loop that groups animations not just by bound skeleton mesh, but also by player (such as when two sequences are bound to the same character). This saves us a little bit of logic compared to the old evalation template, although this will almost never happen since this is an extreme edge case.
- We split the execution of the ECS system in two phases: one which gathers up all the active animations can be dispatched on any thread. The second step does the actual evalation and calls into animation system code, so it is specifically dispatched on the game thread.

diff --git a/Engine/Source/Runtime/MovieSceneTracks/Private/Evaluation/MovieSceneSkeletalAnimationTemplate.cpp b/Engine/Source/Runtime/MovieSceneTracks/Private/Evaluation/MovieSceneSkeletalAnimationTemplate.cpp
deleted file mode 100644
diff --git a/Engine/Source/Runtime/MovieSceneTracks/Private/Evaluation/MovieSceneSkeletalAnimationTemplate.h b/Engine/Source/Runtime/MovieSceneTracks/Private/Evaluation/MovieSceneSkeletalAnimationTemplate.h
deleted file mode 100644
diff --git a/Engine/Source/Runtime/MovieSceneTracks/Private/MovieSceneTracksComponentTypes.cpp b/Engine/Source/Runtime/MovieSceneTracks/Private/MovieSceneTracksComponentTypes.cpp
--- a/Engine/Source/Runtime/MovieSceneTracks/Private/MovieSceneTracksComponentTypes.cpp
+++ b/Engine/Source/Runtime/MovieSceneTracks/Private/MovieSceneTracksComponentTypes.cpp
@@ -582,6 +582,8 @@ FMovieSceneTracksComponentTypes::FMovieSceneTracksComponentTypes()
 	ComponentRegistry->NewComponentType(&FloatPerlinNoiseChannel, TEXT("Float Perlin Noise Channel"));
 	ComponentRegistry->NewComponentType(&DoublePerlinNoiseChannel, TEXT("Double Perlin Noise Channel"));
 
+	ComponentRegistry->NewComponentType(&SkeletalAnimation, TEXT("Skeletal Animation"));
+
 	ComponentRegistry->NewComponentType(&LevelVisibility, TEXT("Level Visibility"));
 	ComponentRegistry->NewComponentType(&DataLayer, TEXT("Data Layer"));
 
@@ -797,6 +799,10 @@ FMovieSceneTracksComponentTypes::FMovieSceneTracksComponentTypes()
 	// Set up PerlinNoise components
 	ComponentRegistry->Factories.DuplicateChildComponent(FloatPerlinNoiseChannel);
 	ComponentRegistry->Factories.DuplicateChildComponent(DoublePerlinNoiseChannel);
+
+	// --------------------------------------------------------------------------------------------
+	// Set up SkeletalAnimation components
+	ComponentRegistry->Factories.DuplicateChildComponent(SkeletalAnimation);
 }
 
 FMovieSceneTracksComponentTypes::~FMovieSceneTracksComponentTypes()
diff --git a/Engine/Source/Runtime/MovieSceneTracks/Private/Sections/MovieSceneSkeletalAnimationSection.cpp b/Engine/Source/Runtime/MovieSceneTracks/Private/Sections/MovieSceneSkeletalAnimationSection.cpp
--- a/Engine/Source/Runtime/MovieSceneTracks/Private/Sections/MovieSceneSkeletalAnimationSection.cpp
+++ b/Engine/Source/Runtime/MovieSceneTracks/Private/Sections/MovieSceneSkeletalAnimationSection.cpp
@@ -4,7 +4,6 @@
 #include "Channels/MovieSceneChannelProxy.h"
 #include "Animation/AnimSequence.h"
 #include "AnimSequencerInstanceProxy.h"
-#include "Evaluation/MovieSceneSkeletalAnimationTemplate.h"
 #include "Logging/MessageLog.h"
 #include "MovieScene.h"
 #include "UObject/SequencerObjectVersion.h"
@@ -14,6 +13,8 @@
 #include "Animation/AnimationPoseData.h"
 #include "Animation/AttributesRuntime.h"
 #include "Misc/FrameRate.h"
+#include "EntitySystem/BuiltInComponentTypes.h"
+#include "MovieSceneTracksComponentTypes.h"
 
 #include UE_INLINE_GENERATED_CPP_BY_NAME(MovieSceneSkeletalAnimationSection)
 
@@ -44,6 +45,51 @@ FMovieSceneSkeletalAnimationParams::FMovieSceneSkeletalAnimationParams()
 	SwapRootBone = ESwapRootBone::SwapRootBone_None;
 }
 
+double FMovieSceneSkeletalAnimationParams::MapTimeToAnimation(const UMovieSceneSection* InSection, FFrameTime InPosition, FFrameRate InFrameRate) const
+{
+	const FFrameNumber SectionStartTime = InSection->GetInclusiveStartFrame();
+	const FFrameNumber SectionEndTime = InSection->GetExclusiveEndFrame();
+	return MapTimeToAnimation(SectionStartTime, SectionEndTime, InPosition, InFrameRate);
+}
+
+double FMovieSceneSkeletalAnimationParams::MapTimeToAnimation(FFrameNumber InSectionStartTime, FFrameNumber InSectionEndTime, FFrameTime InPosition, FFrameRate InFrameRate) const
+{
+	// Get Animation Length and frame time
+	const FFrameTime AnimationLength = GetSequenceLength() * InFrameRate;
+	const int32 LengthInFrames = AnimationLength.FrameNumber.Value + (int)(AnimationLength.GetSubFrame() + 0.5f) + 1;
+
+	// we only play end if we are not looping, and assuming we are looping if Length is greater than default length;
+	const bool bLooping = (InSectionEndTime.Value - InSectionStartTime.Value - StartFrameOffset - EndFrameOffset) > LengthInFrames;	
+
+	// Make sure InPosition FrameTime doesn't underflow InSectionStartTime or overflow InSectionEndTime
+	InPosition = FMath::Clamp(InPosition, FFrameTime(InSectionStartTime), FFrameTime(InSectionEndTime - 1));
+	
+	// Gather helper values
+	const float SectionPlayRate = PlayRate * Animation->RateScale;
+	const float AnimPlayRate = FMath::IsNearlyZero(SectionPlayRate) ? 1.0f : SectionPlayRate;
+	const double SeqLength = GetSequenceLength() - InFrameRate.AsSeconds(StartFrameOffset + EndFrameOffset);
+
+	// The Time from the beginning of InSectionStartTime to InPosition in seconds
+	double SecondsFromSectionStart = FFrameTime::FromDecimal((InPosition - InSectionStartTime).AsDecimal() * AnimPlayRate) / InFrameRate;
+
+	// Logic for reversed animation
+	if (bReverse)
+	{
+		// Duration of this section 
+		double SectionDuration = (((InSectionEndTime - InSectionStartTime) * AnimPlayRate) / InFrameRate);
+		SecondsFromSectionStart = SectionDuration - SecondsFromSectionStart;
+	}
+	// Make sure Seconds is in range
+	if (SeqLength > 0.0 && (bLooping || !FMath::IsNearlyEqual(SecondsFromSectionStart, SeqLength, 1e-4)))
+	{
+		SecondsFromSectionStart = FMath::Fmod(SecondsFromSectionStart, SeqLength);
+	}
+
+	// Add the StartFrameOffset and FirstLoopStartFrameOffset to the current seconds in the section to get the right animation frame
+	SecondsFromSectionStart += InFrameRate.AsSeconds(StartFrameOffset) + InFrameRate.AsSeconds(FirstLoopStartFrameOffset);
+	return SecondsFromSectionStart;
+}
+
 UMovieSceneSkeletalAnimationSection::UMovieSceneSkeletalAnimationSection( const FObjectInitializer& ObjectInitializer )
 	: Super( ObjectInitializer )
 {
@@ -305,8 +351,7 @@ void UMovieSceneSkeletalAnimationSection::GetSnapTimes(TArray<FFrameNumber>& Out
 
 double UMovieSceneSkeletalAnimationSection::MapTimeToAnimation(FFrameTime InPosition, FFrameRate InFrameRate) const
 {
-	FMovieSceneSkeletalAnimationSectionTemplateParameters TemplateParams(Params, GetInclusiveStartFrame(), GetExclusiveEndFrame());
-	return TemplateParams.MapTimeToAnimation(InPosition, InFrameRate);
+	return Params.MapTimeToAnimation(this, InPosition, InFrameRate);
 }
 
 float UMovieSceneSkeletalAnimationSection::GetTotalWeightValue(FFrameTime InTime) const
@@ -342,6 +387,23 @@ void UMovieSceneSkeletalAnimationSection::SetEndFrame(TRangeBound<FFrameNumber>
 	}
 }
 
+void UMovieSceneSkeletalAnimationSection::ImportEntityImpl(UMovieSceneEntitySystemLinker* EntityLinker, const FEntityImportParams& InParams, FImportedEntity* OutImportedEntity)
+{
+	using namespace UE::MovieScene;
+
+	const FBuiltInComponentTypes* BuiltInComponents = FBuiltInComponentTypes::Get();
+	const FMovieSceneTracksComponentTypes* TrackComponents = FMovieSceneTracksComponentTypes::Get();
+
+	const FGuid ObjectBindingID = InParams.GetObjectBindingID();
+	FMovieSceneSkeletalAnimationComponentData ComponentData { this };
+
+	OutImportedEntity->AddBuilder(
+		FEntityBuilder()
+		.Add(TrackComponents->SkeletalAnimation, ComponentData)
+		.AddConditional(BuiltInComponents->GenericObjectBinding, ObjectBindingID, ObjectBindingID.IsValid())
+		.AddConditional(BuiltInComponents->WeightChannel, &Params.Weight, Params.Weight.HasAnyData())
+	);
+}
 
 #if WITH_EDITOR
 void UMovieSceneSkeletalAnimationSection::PreEditChange(FProperty* PropertyAboutToChange)
diff --git a/Engine/Source/Runtime/MovieSceneTracks/Private/Systems/MovieSceneSkeletalAnimationSystem.cpp b/Engine/Source/Runtime/MovieSceneTracks/Private/Systems/MovieSceneSkeletalAnimationSystem.cpp
new file mode 100644
--- /dev/null
+++ b/Engine/Source/Runtime/MovieSceneTracks/Private/Systems/MovieSceneSkeletalAnimationSystem.cpp
@@ -0,0 +1,943 @@
+// Copyright Epic Games, Inc. All Rights Reserved.
+
+#include "Systems/MovieSceneSkeletalAnimationSystem.h"
+
+#include "Async/TaskGraphInterfaces.h"
+#include "EntitySystem/BuiltInComponentTypes.h"
+#include "EntitySystem/Interrogation/MovieSceneInterrogationLinker.h"
+#include "EntitySystem/MovieSceneEntityIDs.h"
+#include "EntitySystem/MovieSceneEntitySystemLinker.h"
+#include "EntitySystem/MovieSceneEntitySystemRunner.h"
+#include "EntitySystem/MovieSceneEntitySystemTask.h"
+#include "IMovieScenePlayer.h"
+#include "MovieSceneExecutionToken.h"
+#include "MovieSceneTracksComponentTypes.h"
+#include "Sections/MovieSceneSkeletalAnimationSection.h"
+
+#include "Evaluation/PreAnimatedState/IMovieScenePreAnimatedStorage.h"
+#include "Evaluation/PreAnimatedState/MovieScenePreAnimatedObjectStorage.h"
+#include "Evaluation/PreAnimatedState/MovieScenePreAnimatedStorageID.h"
+#include "Evaluation/PreAnimatedState/MovieScenePreAnimatedStorageID.inl"
+
+#include "AnimCustomInstanceHelper.h"
+#include "AnimSequencerInstance.h"
+#include "AnimSequencerInstanceProxy.h"
+#include "Animation/AnimMontage.h"
+#include "Animation/AnimSingleNodeInstance.h"
+#include "Animation/AnimSingleNodeInstance.h"
+#include "Components/SkeletalMeshComponent.h"
+#include "Engine/SkeletalMesh.h"
+#include "SkeletalMeshRestoreState.h"
+
+#include "Rendering/MotionVectorSimulation.h"
+#include "Systems/MovieSceneMotionVectorSimulationSystem.h"
+#include "Systems/MovieSceneComponentTransformSystem.h"
+#include "Systems/MovieSceneQuaternionInterpolationRotationSystem.h"
+#include "Systems/WeightAndEasingEvaluatorSystem.h"
+
+#include "Components/SkeletalMeshComponent.h"
+
+#include UE_INLINE_GENERATED_CPP_BY_NAME(MovieSceneSkeletalAnimationSystem)
+
+DECLARE_CYCLE_STAT(TEXT("Gather skeletal animations"), MovieSceneEval_GatherSkeletalAnimations, STATGROUP_MovieSceneECS);
+DECLARE_CYCLE_STAT(TEXT("Evaluate skeletal animations"), MovieSceneEval_EvaluateSkeletalAnimations, STATGROUP_MovieSceneECS);
+
+namespace UE::MovieScene
+{
+
+/** Helper function to get our sequencer animation node from a skeletal mesh component */
+UAnimSequencerInstance* GetAnimSequencerInstance(USkeletalMeshComponent* SkeletalMeshComponent)
+{
+	ISequencerAnimationSupport* SeqInterface = Cast<ISequencerAnimationSupport>(SkeletalMeshComponent->GetAnimInstance());
+	if (SeqInterface)
+	{
+		return Cast<UAnimSequencerInstance>(SeqInterface->GetSourceAnimInstance());
+	}
+
+	return nullptr;
+}
+
+/** ------------------------------------------------------------------------- */
+
+/** Pre-animated state for skeletal animations */
+struct FPreAnimatedSkeletalAnimationState
+{
+	EAnimationMode::Type AnimationMode;
+	TStrongObjectPtr<UAnimInstance> CachedAnimInstance;
+	FSkeletalMeshRestoreState SkeletalMeshRestoreState;
+};
+
+/** Pre-animation traits for skeletal animations */
+struct FPreAnimatedSkeletalAnimationTraits : FBoundObjectPreAnimatedStateTraits
+{
+	using KeyType     = FObjectKey;
+	using StorageType = FPreAnimatedSkeletalAnimationState;
+
+	static void CachePreAnimatedValue(const KeyType& Object, StorageType& OutCachedValue)
+	{
+		USkeletalMeshComponent* Component = Cast<USkeletalMeshComponent>(Object.ResolveObjectPtr());
+		if (ensure(Component))
+		{
+			OutCachedValue.AnimationMode = Component->GetAnimationMode();
+			OutCachedValue.CachedAnimInstance.Reset(Component->AnimScriptInstance);
+			OutCachedValue.SkeletalMeshRestoreState.SaveState(Component);
+		}
+	}
+
+	static void RestorePreAnimatedValue(const KeyType& Object, StorageType& InOutCachedValue, const FRestoreStateParams& Params)
+	{
+		USkeletalMeshComponent* Component = Cast<USkeletalMeshComponent>(Object.ResolveObjectPtr());
+		if (!Component)
+		{
+			return;
+		}
+
+		ISequencerAnimationSupport* SequencerInst = Cast<ISequencerAnimationSupport>(GetAnimSequencerInstance(Component));
+		if (SequencerInst)
+		{
+			SequencerInst->ResetPose();
+			SequencerInst->ResetNodes();
+		}
+
+		FAnimCustomInstanceHelper::UnbindFromSkeletalMeshComponent<UAnimSequencerInstance>(Component);
+
+		// Restore LOD before reinitializing anim instance
+		InOutCachedValue.SkeletalMeshRestoreState.RestoreLOD(Component);
+
+		if (Component->GetAnimationMode() != InOutCachedValue.AnimationMode)
+		{
+			// this SetAnimationMode reinitializes even if the mode is same
+			// if we're using same anim blueprint, we don't want to keep reinitializing it. 
+			Component->SetAnimationMode(InOutCachedValue.AnimationMode);
+		}
+		if (InOutCachedValue.CachedAnimInstance.Get())
+		{
+			Component->AnimScriptInstance = InOutCachedValue.CachedAnimInstance.Get();
+			InOutCachedValue.CachedAnimInstance.Reset();
+			if (Component->AnimScriptInstance && Component->GetSkeletalMeshAsset() && Component->AnimScriptInstance->CurrentSkeleton != Component->GetSkeletalMeshAsset()->GetSkeleton())
+			{
+				//the skeleton may have changed so need to recalc required bones as needed.
+				Component->AnimScriptInstance->CurrentSkeleton = Component->GetSkeletalMeshAsset()->GetSkeleton();
+				//Need at least RecalcRequiredbones and UpdateMorphTargetrs
+				Component->InitializeAnimScriptInstance(true);
+			}
+		}
+
+		// Restore pose after unbinding to force the restored pose
+		Component->SetUpdateAnimationInEditor(true);
+		Component->SetUpdateClothInEditor(true);
+		Component->TickAnimation(0.f, false);
+
+		Component->RefreshBoneTransforms();
+		Component->RefreshFollowerComponents();
+		Component->UpdateComponentToWorld();
+		Component->FinalizeBoneTransform();
+		Component->MarkRenderTransformDirty();
+		Component->MarkRenderDynamicDataDirty();
+
+		// Reset the mesh component update flag and animation mode to what they were before we animated the object
+		InOutCachedValue.SkeletalMeshRestoreState.RestoreState(Component);
+
+		// if not game world, don't clean this up
+		if (Component->GetWorld() != nullptr && Component->GetWorld()->IsGameWorld() == false)
+		{
+			Component->ClearMotionVector();
+		}
+	}
+};
+
+/** Pre-animation storage for skeletal animations */
+struct FPreAnimatedSkeletalAnimationStorage : TPreAnimatedStateStorage_ObjectTraits<FPreAnimatedSkeletalAnimationTraits>
+{
+	static TAutoRegisterPreAnimatedStorageID<FPreAnimatedSkeletalAnimationStorage> StorageID;
+};
+
+TAutoRegisterPreAnimatedStorageID<FPreAnimatedSkeletalAnimationStorage> FPreAnimatedSkeletalAnimationStorage::StorageID;
+
+/** ------------------------------------------------------------------------- */
+
+/** Pre-animated state for a sequencer montage node */
+struct FPreAnimatedSkeletalAnimationMontageState
+{
+	TWeakObjectPtr<UAnimInstance> WeakInstance;
+	int32 MontageInstanceId;
+};
+
+/** Pre-animated traits for a sequencer montage node */
+struct FPreAnimatedSkeletalAnimationMontageTraits : FBoundObjectPreAnimatedStateTraits
+{
+	using KeyType     = FObjectKey;
+	using StorageType = FPreAnimatedSkeletalAnimationMontageState;
+
+	static void CachePreAnimatedValue(const KeyType& Object, StorageType& OutCachedValue)
+	{
+		// Should already be initialized.
+	}
+
+	static void RestorePreAnimatedValue(const FObjectKey& Object, FPreAnimatedSkeletalAnimationMontageState& InOutCachedValue, const FRestoreStateParams& Params)
+	{
+		UAnimInstance* AnimInstance = InOutCachedValue.WeakInstance.Get();
+		if (AnimInstance)
+		{
+			FAnimMontageInstance* MontageInstance = AnimInstance->GetMontageInstanceForID(InOutCachedValue.MontageInstanceId);
+			if (MontageInstance)
+			{
+				MontageInstance->Stop(FAlphaBlend(0.f), false);
+			}
+		}
+	}
+};
+
+/** Pre-animated storage for a sequencer montage node */
+struct FPreAnimatedSkeletalAnimationMontageStorage : TPreAnimatedStateStorage_ObjectTraits<FPreAnimatedSkeletalAnimationMontageTraits>
+{
+	static TAutoRegisterPreAnimatedStorageID<FPreAnimatedSkeletalAnimationMontageStorage> StorageID;
+};
+
+TAutoRegisterPreAnimatedStorageID<FPreAnimatedSkeletalAnimationMontageStorage> FPreAnimatedSkeletalAnimationMontageStorage::StorageID;
+
+/** ------------------------------------------------------------------------- */
+
+/** Pre-animated traits for a sequencer animation node */
+struct FPreAnimatedSkeletalAnimationAnimInstanceTraits : FBoundObjectPreAnimatedStateTraits
+{
+	using KeyType     = FObjectKey;
+	using StorageType = bool; // We actually don't need any state, so this will be a dummy value
+
+	static void CachePreAnimatedValue(const KeyType& Object, StorageType& Unused)
+	{
+		// Nothing to do, we just need the object pointer to restore state.
+	}
+
+	static void RestorePreAnimatedValue(const FObjectKey& Object, bool& Unused, const FRestoreStateParams& Params)
+	{
+		ISequencerAnimationSupport* SequencerAnimationSupport = Cast<ISequencerAnimationSupport>(Object.ResolveObjectPtr());
+		if (ensure(SequencerAnimationSupport))
+		{
+			SequencerAnimationSupport->ResetNodes();
+		}
+	}
+};
+
+/** Pre-animated storage for a sequencer animation node */
+struct FPreAnimatedSkeletalAnimationAnimInstanceStorage : TPreAnimatedStateStorage_ObjectTraits<FPreAnimatedSkeletalAnimationAnimInstanceTraits>
+{
+	static TAutoRegisterPreAnimatedStorageID<FPreAnimatedSkeletalAnimationAnimInstanceStorage> StorageID;
+};
+
+TAutoRegisterPreAnimatedStorageID<FPreAnimatedSkeletalAnimationAnimInstanceStorage> FPreAnimatedSkeletalAnimationAnimInstanceStorage::StorageID;
+
+/** ------------------------------------------------------------------------- */
+
+/** Task for gathering active skeletal animations */
+struct FGatherSkeletalAnimations
+{
+	const FInstanceRegistry* InstanceRegistry;
+	FSkeletalAnimationSystemData& SystemData;
+
+	FGatherSkeletalAnimations(const FInstanceRegistry* InInstanceRegistry, FSkeletalAnimationSystemData& InSystemData)
+		: InstanceRegistry(InInstanceRegistry)
+		, SystemData(InSystemData)
+	{
+	}
+
+	void ForEachAllocation(
+			const FEntityAllocationProxy AllocationProxy, 
+			TRead<FMovieSceneEntityID> EntityIDs,
+			TRead<FRootInstanceHandle> RootInstanceHandles,
+			TRead<FInstanceHandle> InstanceHandles,
+			TRead<UObject*> BoundObjects,
+			TRead<FMovieSceneSkeletalAnimationComponentData> SkeletalAnimations,
+			TReadOptional<double> WeightAndEasings)
+	{
+		// Gather all the skeletal animations currently active in all sequences.
+		// We map these animations to their bound object, which means we might blend animations from different sequences
+		// that have bound to the same object.
+		const FBuiltInComponentTypes* BuiltInComponents = FBuiltInComponentTypes::Get();
+		const FEntityAllocation* Allocation = AllocationProxy.GetAllocation();
+		const int32 Num = Allocation->Num();
+		for (int32 Index = 0; Index < Num; ++Index)
+		{
+			FMovieSceneEntityID EntityID(EntityIDs[Index]);
+			const FRootInstanceHandle& RootInstanceHandle(RootInstanceHandles[Index]);
+			const FInstanceHandle& InstanceHandle(InstanceHandles[Index]);
+			UObject* BoundObject(BoundObjects[Index]);
+			const FMovieSceneSkeletalAnimationComponentData& SkeletalAnimation(SkeletalAnimations[Index]);
+			const double Weight = (WeightAndEasings ? WeightAndEasings[Index] : 1.f);
+
+			const bool bWantsRestoreState = AllocationProxy.GetAllocationType().Contains(BuiltInComponents->Tags.RestoreState);
+
+			// Get the full context, so we can get both the current and previous evaluation times.
+			const FSequenceInstance& SequenceInstance = InstanceRegistry->GetInstance(InstanceHandle);
+			const FMovieSceneContext& Context = SequenceInstance.GetContext();
+			IMovieScenePlayer* Player = SequenceInstance.GetPlayer();
+
+			// Calculate the time at which to evaluate the animation
+			const UMovieSceneSkeletalAnimationSection* AnimSection = SkeletalAnimation.Section;
+			const FMovieSceneSkeletalAnimationParams& AnimParams = AnimSection->Params;
+			const float EvalTime = AnimParams.MapTimeToAnimation(AnimSection, Context.GetTime(), Context.GetFrameRate());
+			const float PreviousEvalTime = AnimParams.MapTimeToAnimation(AnimSection, Context.GetPreviousTime(), Context.GetFrameRate());
+
+			FBoundObjectActiveSkeletalAnimations& BoundObjectAnimations = SystemData.SkeletalAnimations.FindOrAdd(BoundObject);
+			BoundObjectAnimations.Animations.Add(FActiveSkeletalAnimation{ Player, AnimSection, Context, EntityID, RootInstanceHandle, PreviousEvalTime, EvalTime, Weight, bWantsRestoreState });
+
+			if (FMotionVectorSimulation::IsEnabled())
+			{
+				const FFrameTime SimulatedTime = UE::MovieScene::GetSimulatedMotionVectorTime(Context);
+
+				// Calculate the time at which to evaluate the animation
+				const float SimulatedEvalTime = AnimParams.MapTimeToAnimation(AnimSection, SimulatedTime, Context.GetFrameRate());
+
+				// Evaluate the weight channel and section easing at the simulation time... right now we don't benefit
+				// from that being evaluated by the channel evaluators.
+				float SimulatedManualWeight = 1.f;
+				AnimParams.Weight.Evaluate(SimulatedTime, SimulatedManualWeight);
+
+				const float SimulatedWeight = SimulatedManualWeight * AnimSection->EvaluateEasing(SimulatedTime);
+
+				BoundObjectAnimations.SimulatedAnimations.Add(FActiveSkeletalAnimation{ Player, AnimSection, Context, EntityID, RootInstanceHandle, EvalTime, SimulatedEvalTime, SimulatedWeight, bWantsRestoreState });
+			}
+		}
+	}
+};
+
+/** Task for evaluating skeletal animations */
+struct FEvaluateSkeletalAnimations
+{
+private:
+
+	UMovieSceneEntitySystemLinker* Linker;
+	FSkeletalAnimationSystemData& SystemData;
+
+	TSharedPtr<FPreAnimatedSkeletalAnimationStorage> PreAnimatedStorage;
+	TSharedPtr<FPreAnimatedSkeletalAnimationMontageStorage> PreAnimatedMontageStorage;
+	TSharedPtr<FPreAnimatedSkeletalAnimationAnimInstanceStorage> PreAnimatedAnimInstanceStorage;
+
+public:
+
+	FEvaluateSkeletalAnimations(UMovieSceneEntitySystemLinker* InLinker, FSkeletalAnimationSystemData& InSystemData)
+		: Linker(InLinker)
+		, SystemData(InSystemData)
+	{
+		PreAnimatedStorage = InLinker->PreAnimatedState.GetOrCreateStorage<FPreAnimatedSkeletalAnimationStorage>();
+		PreAnimatedMontageStorage = InLinker->PreAnimatedState.GetOrCreateStorage<FPreAnimatedSkeletalAnimationMontageStorage>();
+		PreAnimatedAnimInstanceStorage = InLinker->PreAnimatedState.GetOrCreateStorage<FPreAnimatedSkeletalAnimationAnimInstanceStorage>();
+	}
+
+	void Run()
+	{
+		for (const TTuple<UObject*, FBoundObjectActiveSkeletalAnimations>& Pair : SystemData.SkeletalAnimations)
+		{
+			EvaluateSkeletalAnimations(Pair.Key, Pair.Value);
+		}
+	}
+
+private:
+
+	void EvaluateSkeletalAnimations(UObject* InObject, const FBoundObjectActiveSkeletalAnimations& InSkeletalAnimations)
+	{
+		ensureMsgf(InObject, TEXT("Attempting to evaluate an Animation track with a null object."));
+
+		// Get the bound skeletal mesh component.
+		USkeletalMeshComponent* SkeletalMeshComponent = SkeletalMeshComponentFromObject(InObject);
+		if (!SkeletalMeshComponent || !SkeletalMeshComponent->GetSkeletalMeshAsset())
+		{
+			return;
+		}
+
+		// Cache pre-animated state for this bound object before doing anything.
+		// We don't yet track what entities have already started animated vs. entities that just started this frame,
+		// so we just process all the currently active ones. If they are already tracked and have already had their
+		// pre-animated state saved, it these calls will just early return.
+		for (const FActiveSkeletalAnimation& SkeletalAnimation : InSkeletalAnimations.Animations)
+		{
+			PreAnimatedStorage->BeginTrackingEntity(SkeletalAnimation.EntityID, SkeletalAnimation.bWantsRestoreState, SkeletalAnimation.RootInstanceHandle, SkeletalMeshComponent);
+		}
+		FCachePreAnimatedValueParams CacheParams;
+		PreAnimatedStorage->CachePreAnimatedValue(CacheParams, SkeletalMeshComponent);
+
+		// Setup any needed animation nodes for sequencer playback.
+		UAnimInstance* ExistingAnimInstance = GetSourceAnimInstance(SkeletalMeshComponent);
+		bool bWasCreated = false;
+		ISequencerAnimationSupport* SequencerInstance = FAnimCustomInstanceHelper::BindToSkeletalMeshComponent<UAnimSequencerInstance>(SkeletalMeshComponent, bWasCreated);
+		if (SequencerInstance)
+		{
+			if (bWasCreated)
+			{
+				SequencerInstance->SavePose();
+			}
+			else
+			{
+				SequencerInstance->ConstructNodes();
+			}
+		}
+
+		// Group skeletal animations by player, in case we have multiple ongoing sequences binding to the same object.
+		TMap<IMovieScenePlayer*, FBoundObjectActiveSkeletalAnimations> SkeletalAnimationsPerPlayer;
+		for (const FActiveSkeletalAnimation& Animation : InSkeletalAnimations.Animations)
+		{
+			FBoundObjectActiveSkeletalAnimations& SkeletalAnimationsForPlayer = SkeletalAnimationsPerPlayer.FindOrAdd(Animation.Player);
+			SkeletalAnimationsForPlayer.Animations.Add(Animation);
+		}
+		for (const FActiveSkeletalAnimation& SimulatedAnimation : InSkeletalAnimations.SimulatedAnimations)
+		{
+			FBoundObjectActiveSkeletalAnimations& SkeletalAnimationsForPlayer = SkeletalAnimationsPerPlayer.FindOrAdd(SimulatedAnimation.Player);
+			SkeletalAnimationsForPlayer.SimulatedAnimations.Add(SimulatedAnimation);
+		}
+
+		// Evaluate each group for the given bound object.
+		for (const TTuple<IMovieScenePlayer*, FBoundObjectActiveSkeletalAnimations>& Pair : SkeletalAnimationsPerPlayer)
+		{
+			EvaluateSkeletalAnimationsForPlayer(*Pair.Key, SkeletalMeshComponent, ExistingAnimInstance, SequencerInstance, Pair.Value);
+		}
+	}
+
+	void EvaluateSkeletalAnimationsForPlayer(
+			IMovieScenePlayer& InPlayer, 
+			USkeletalMeshComponent* InSkeletalMeshComponent, 
+			UAnimInstance* InExistingAnimInstance, 
+			ISequencerAnimationSupport* InSequencerInstance, 
+			const FBoundObjectActiveSkeletalAnimations& InSkeletalAnimations)
+	{
+		const bool bPreviewPlayback = ShouldUsePreviewPlayback(InPlayer, *InSkeletalMeshComponent);
+
+		const EMovieScenePlayerStatus::Type PlayerStatus = InPlayer.GetPlaybackStatus();
+
+		// If the playback status is jumping, ie. one such occurrence is setting the time for thumbnail generation, disable anim notifies updates because it could fire audio.
+		// If the playback status is scrubbing, we disable notifies for now because we can't properly fire them in all cases until we get evaluation range info.
+		// We now layer this with the passed in notify toggle to force a disable in this case.
+		const bool bFireNotifies = !bPreviewPlayback || (PlayerStatus != EMovieScenePlayerStatus::Jumping && PlayerStatus != EMovieScenePlayerStatus::Stopped && PlayerStatus != EMovieScenePlayerStatus::Scrubbing);
+
+		// When jumping from one cut to another cut, the delta time should be 0 so that anim notifies before the current position are not evaluated. Note, anim notifies at the current time should still be evaluated.
+		const FInstanceHandle RootInstanceHandle = InPlayer.GetEvaluationTemplate().GetRootInstanceHandle();
+		const FSequenceInstance& RootInstance = Linker->GetInstanceRegistry()->GetInstance(RootInstanceHandle);
+		const FMovieSceneContext& RootContext = RootInstance.GetContext();
+		const double RootDeltaTime = (RootContext.HasJumped() ? FFrameTime(0) : RootContext.GetRange().Size<FFrameTime>() ) / RootContext.GetFrameRate();
+
+		const bool bResetDynamics = PlayerStatus == EMovieScenePlayerStatus::Stepping || 
+			PlayerStatus == EMovieScenePlayerStatus::Jumping || 
+			PlayerStatus == EMovieScenePlayerStatus::Scrubbing || 
+			(RootDeltaTime == 0.0f && PlayerStatus != EMovieScenePlayerStatus::Stopped);
+
+		// Need to zero all weights first since we may be blending animation that are keeping state but are no longer active.
+
+		if (InSequencerInstance)
+		{
+			InSequencerInstance->ResetNodes();
+		}
+		else if (InExistingAnimInstance)
+		{
+			for (const TPair<FObjectKey, FMontagePlayerPerSectionData >& Pair : SystemData.MontageData)
+			{
+				int32 InstanceId = Pair.Value.MontageInstanceId;
+				FAnimMontageInstance* MontageInstanceToUpdate = InExistingAnimInstance->GetMontageInstanceForID(InstanceId);
+				if (MontageInstanceToUpdate)
+				{
+					MontageInstanceToUpdate->SetDesiredWeight(0.0f);
+					MontageInstanceToUpdate->SetWeight(0.0f);
+				}
+			}
+		}
+
+		if (InSkeletalAnimations.SimulatedAnimations.Num() != 0)
+		{
+			UMovieSceneMotionVectorSimulationSystem* MotionVectorSim = InPlayer.GetEvaluationTemplate().GetEntitySystemLinker()->FindSystem<UMovieSceneMotionVectorSimulationSystem>();
+			if (MotionVectorSim && MotionVectorSim->IsSimulationEnabled())
+			{
+				ApplyAnimations(InPlayer, InSkeletalMeshComponent, InSkeletalAnimations.SimulatedAnimations, bPreviewPlayback, bFireNotifies, bResetDynamics);
+				InSkeletalMeshComponent->TickAnimation(0.f, false);
+				InSkeletalMeshComponent->ForceMotionVector();
+		
+				SimulateMotionVectors(InSkeletalMeshComponent, MotionVectorSim);
+			}
+		}
+
+		ApplyAnimations(InPlayer, InSkeletalMeshComponent, InSkeletalAnimations.Animations, bPreviewPlayback, bFireNotifies, bResetDynamics);
+
+		// If the skeletal component has already ticked this frame because tick prerequisites weren't set up yet or a new binding was created, forcibly tick this component to update.
+		// This resolves first frame issues where the skeletal component ticks first, then the sequencer binding is resolved which sets up tick prerequisites
+		// for the next frame.
+		if (InSkeletalMeshComponent->PoseTickedThisFrame() || (InSequencerInstance && InSequencerInstance->GetSourceAnimInstance() != InExistingAnimInstance))
+		{
+			InSkeletalMeshComponent->TickAnimation(0.f, false);
+
+			InSkeletalMeshComponent->RefreshBoneTransforms();
+			InSkeletalMeshComponent->RefreshFollowerComponents();
+			InSkeletalMeshComponent->UpdateComponentToWorld();
+			InSkeletalMeshComponent->FinalizeBoneTransform();
+			InSkeletalMeshComponent->MarkRenderTransformDirty();
+			InSkeletalMeshComponent->MarkRenderDynamicDataDirty();
+		}
+	}
+
+private:
+
+	static bool ShouldUsePreviewPlayback(IMovieScenePlayer& Player, UObject& RuntimeObject)
+	{
+		// We also use PreviewSetAnimPosition in PIE when not playing, as we can preview in PIE.
+		bool bIsNotInPIEOrNotPlaying = (RuntimeObject.GetWorld() && !RuntimeObject.GetWorld()->HasBegunPlay()) || Player.GetPlaybackStatus() != EMovieScenePlayerStatus::Playing;
+		return GIsEditor && bIsNotInPIEOrNotPlaying;
+	}
+
+	static bool CanPlayAnimation(USkeletalMeshComponent* SkeletalMeshComponent, UAnimSequenceBase* AnimAssetBase)
+	{
+		return (SkeletalMeshComponent->GetSkeletalMeshAsset() &&
+			SkeletalMeshComponent->GetSkeletalMeshAsset()->GetSkeleton() &&
+			AnimAssetBase != nullptr &&
+			AnimAssetBase->GetSkeleton() != nullptr);
+	}
+
+	static UAnimInstance* GetSourceAnimInstance(USkeletalMeshComponent* SkeletalMeshComponent)
+	{
+		UAnimInstance* SkelAnimInstance = SkeletalMeshComponent->GetAnimInstance();
+		ISequencerAnimationSupport* SeqInterface = Cast<ISequencerAnimationSupport>(SkelAnimInstance);
+		if (SeqInterface)
+		{
+			return SeqInterface->GetSourceAnimInstance();
+		}
+
+		return SkelAnimInstance;
+	}
+
+	static USkeletalMeshComponent* SkeletalMeshComponentFromObject(UObject* InObject)
+	{
+		// Check if we are bound directly to a skeletal mesh component.
+		USkeletalMeshComponent* SkeletalMeshComponent = Cast<USkeletalMeshComponent>(InObject);
+		if (SkeletalMeshComponent)
+		{
+			return SkeletalMeshComponent;
+		}
+
+		// Then check to see if we are controlling an actor. If so use its first skeletal mesh component.
+		AActor* Actor = Cast<AActor>(InObject);
+		if (!Actor)
+		{
+			if (UChildActorComponent* ChildActorComponent = Cast<UChildActorComponent>(InObject))
+			{
+				Actor = ChildActorComponent->GetChildActor();
+			}
+		}
+		if (Actor)
+		{
+			return Actor->FindComponentByClass<USkeletalMeshComponent>();
+		}
+		return nullptr;
+	}
+
+private:
+
+	/** Parameter structure for setting the skeletal animation position */
+	struct FSetAnimPositionParams
+	{
+		FMovieSceneEntityID EntityID;
+		FRootInstanceHandle RootInstanceHandle;
+
+		const UMovieSceneSkeletalAnimationSection* Section = nullptr;
+		USkeletalMeshComponent* SkeletalMeshComponent = nullptr;
+
+		FFrameTime CurrentTime;
+		float FromPosition;
+		float ToPosition;
+		float Weight;
+
+		bool bWantsRestoreState;
+		bool bPlaying;
+		bool bFireNotifies;
+		bool bResetDynamics;
+	};
+
+	void SimulateMotionVectors(USkeletalMeshComponent* SkeletalMeshComponent, UMovieSceneMotionVectorSimulationSystem* MotionVectorSim)
+	{
+		for (USceneComponent* Child : SkeletalMeshComponent->GetAttachChildren())
+		{
+			if (!Child)
+			{
+				continue;
+			}
+
+			FName SocketName = Child->GetAttachSocketName();
+			if (SocketName != NAME_None)
+			{
+				FTransform SocketTransform = SkeletalMeshComponent->GetSocketTransform(SocketName, RTS_Component);
+				MotionVectorSim->AddSimulatedTransform(SkeletalMeshComponent, SocketTransform, SocketName);
+			}
+		}
+	}
+
+	void ApplyAnimations(
+		IMovieScenePlayer& Player,
+		USkeletalMeshComponent* SkeletalMeshComponent,
+		TArrayView<const FActiveSkeletalAnimation> SkeletalAnimations,
+		bool bPreviewPlayback,
+		bool bFireNotifies,
+		bool bResetDynamics)
+	{
+		const EMovieScenePlayerStatus::Type PlayerStatus = Player.GetPlaybackStatus();
+
+		for (const FActiveSkeletalAnimation& SkeletalAnimation : SkeletalAnimations)
+		{
+			const UMovieSceneSkeletalAnimationSection* AnimSection = SkeletalAnimation.AnimSection;
+			const FMovieSceneSkeletalAnimationParams& AnimParams = AnimSection->Params;
+
+			// Don't fire notifies if looping around.
+			bool bLooped = false;
+			if (AnimParams.bReverse)
+			{
+				if (SkeletalAnimation.FromEvalTime <= SkeletalAnimation.ToEvalTime)
+				{
+					bLooped = true;
+				}
+			}
+			else if (SkeletalAnimation.FromEvalTime >= SkeletalAnimation.ToEvalTime)
+			{
+				bLooped = true;
+			}
+
+			FSetAnimPositionParams SetAnimPositionParams;
+			SetAnimPositionParams.EntityID = SkeletalAnimation.EntityID;
+			SetAnimPositionParams.RootInstanceHandle = SkeletalAnimation.RootInstanceHandle;
+			SetAnimPositionParams.Section = AnimSection;
+			SetAnimPositionParams.SkeletalMeshComponent = SkeletalMeshComponent;
+			SetAnimPositionParams.CurrentTime = SkeletalAnimation.Context.GetTime();
+			SetAnimPositionParams.FromPosition = SkeletalAnimation.FromEvalTime;
+			SetAnimPositionParams.ToPosition = SkeletalAnimation.ToEvalTime;
+			SetAnimPositionParams.Weight = SkeletalAnimation.BlendWeight;
+			SetAnimPositionParams.bWantsRestoreState = SkeletalAnimation.bWantsRestoreState;
+			SetAnimPositionParams.bPlaying = (PlayerStatus == EMovieScenePlayerStatus::Playing);
+			SetAnimPositionParams.bFireNotifies = (bFireNotifies && !AnimParams.bSkipAnimNotifiers && !bLooped);
+			SetAnimPositionParams.bResetDynamics = bResetDynamics;
+
+			if (bPreviewPlayback)
+			{
+				PreviewSetAnimPosition(SetAnimPositionParams);
+			}
+			else
+			{
+				SetAnimPosition(SetAnimPositionParams);
+			}
+		}
+	}
+	
+	// Determines whether the bound object has a component transform property tag
+	bool ContainsTransform(UObject* InBoundObject) const
+	{
+		using namespace UE::MovieScene;
+
+		bool bContainsTransform = false;
+
+		auto HarvestTransforms = [InBoundObject, &bContainsTransform](UObject* BoundObject)
+		{
+			if (BoundObject == InBoundObject)
+			{
+				bContainsTransform = true;
+			}
+		};
+				
+		FBuiltInComponentTypes* BuiltInComponents = FBuiltInComponentTypes::Get();
+
+		FMovieSceneTracksComponentTypes* Components = FMovieSceneTracksComponentTypes::Get();
+		FEntityTaskBuilder()
+			.Read(BuiltInComponents->BoundObject)
+			// Only include component transforms
+			.FilterAll({ Components->ComponentTransform.PropertyTag })
+			// Only read things with the resolved properties on - this ensures we do not read any intermediate component transforms for blended properties
+			.FilterAny({ BuiltInComponents->CustomPropertyIndex, BuiltInComponents->FastPropertyOffset, BuiltInComponents->SlowProperty })
+			.Iterate_PerEntity(&Linker->EntityManager, HarvestTransforms);
+
+		return bContainsTransform;
+	}
+
+	// Get the current transform for the component that the root bone will be swaped to
+	TOptional<FTransform> GetCurrentTransform(ESwapRootBone SwapRootBone, USkeletalMeshComponent* SkeletalMeshComponent) const
+	{
+		TOptional<FTransform> CurrentTransform;
+		if (SwapRootBone == ESwapRootBone::SwapRootBone_Component)
+		{
+			if (ContainsTransform(SkeletalMeshComponent))
+			{
+				CurrentTransform = SkeletalMeshComponent->GetRelativeTransform();
+			}
+		}
+		else if (SwapRootBone == ESwapRootBone::SwapRootBone_Actor)
+		{
+			if (AActor* Actor = SkeletalMeshComponent->GetOwner())
+			{
+				if (USceneComponent* RootComponent = Actor->GetRootComponent())
+				{
+					if (ContainsTransform(RootComponent))
+					{
+						CurrentTransform = RootComponent->GetRelativeTransform();
+					}
+				}
+			}
+		}
+
+		return CurrentTransform;
+	}
+
+	void SetAnimPosition(const FSetAnimPositionParams& Params)
+	{
+		static const bool bLooping = false;
+
+		const FMovieSceneSkeletalAnimationParams& AnimParams = Params.Section->Params;
+		if (!CanPlayAnimation(Params.SkeletalMeshComponent, AnimParams.Animation))
+		{
+			return;
+		}
+		if (AnimParams.bForceCustomMode)
+		{
+			Params.SkeletalMeshComponent->SetAnimationMode(EAnimationMode::AnimationCustomMode);
+		}
+
+		UAnimSequencerInstance* SequencerInst = GetAnimSequencerInstance(Params.SkeletalMeshComponent);
+		if (SequencerInst)
+		{
+			PreAnimatedAnimInstanceStorage->BeginTrackingEntity(Params.EntityID, Params.bWantsRestoreState, Params.RootInstanceHandle, SequencerInst);
+			PreAnimatedAnimInstanceStorage->CachePreAnimatedValue(FCachePreAnimatedValueParams(), SequencerInst);
+			
+			TOptional <FRootMotionOverride> RootMotion;
+			UMovieSceneSkeletalAnimationSection::FRootMotionParams RootMotionParams;
+
+			Params.Section->GetRootMotion(Params.CurrentTime, RootMotionParams);
+			if (RootMotionParams.Transform.IsSet())
+			{
+				RootMotion = FRootMotionOverride();
+				RootMotion.GetValue().RootMotion = RootMotionParams.Transform.GetValue();
+				RootMotion.GetValue().bBlendFirstChildOfRoot = RootMotionParams.bBlendFirstChildOfRoot;
+				RootMotion.GetValue().ChildBoneIndex = RootMotionParams.ChildBoneIndex;
+				RootMotion.GetValue().PreviousTransform = RootMotionParams.PreviousTransform.GetValue();
+			}
+
+			// Use the section's address as the ID for the anim sequence.
+			const uint32 AnimSequenceID = GetTypeHash(Params.Section);
+
+			// If Sequencer has a transform track, we want to set the initial transform so that root motion (if it exists) can be applied relative to that.
+			TOptional<FTransform> CurrentTransform = GetCurrentTransform(Params.Section->Params.SwapRootBone, Params.SkeletalMeshComponent);
+
+			FAnimSequencerData AnimSequencerData(
+					AnimParams.Animation,
+					AnimSequenceID,
+					RootMotion, 
+					Params.FromPosition,
+					Params.ToPosition,
+					Params.Weight, 
+					Params.bFireNotifies, 
+					Params.Section->Params.SwapRootBone, 
+					CurrentTransform,
+					Params.Section->Params.MirrorDataTable.Get());
+			SequencerInst->UpdateAnimTrackWithRootMotion(AnimSequencerData);
+		}
+		else if (UAnimInstance* AnimInst = GetSourceAnimInstance(Params.SkeletalMeshComponent))
+		{
+			FMontagePlayerPerSectionData* SectionData = SystemData.MontageData.Find(Params.Section);
+
+			int32 InstanceId = (SectionData) ? SectionData->MontageInstanceId : INDEX_NONE;
+
+			const float AssetPlayRate = FMath::IsNearlyZero(AnimParams.Animation->RateScale) ? 1.0f : AnimParams.Animation->RateScale;
+			TWeakObjectPtr<UAnimMontage> WeakMontage = FAnimMontageInstance::SetSequencerMontagePosition(
+					AnimParams.SlotName, 
+					AnimInst, 
+					InstanceId, 
+					AnimParams.Animation, 
+					Params.FromPosition / AssetPlayRate, 
+					Params.ToPosition / AssetPlayRate, 
+					Params.Weight, 
+					bLooping, 
+					Params.bPlaying);
+
+			UAnimMontage* Montage = WeakMontage.Get();
+			if (Montage)
+			{
+				FMontagePlayerPerSectionData& DataContainer = SystemData.MontageData.FindOrAdd(Params.Section);
+				DataContainer.Montage = WeakMontage;
+				DataContainer.MontageInstanceId = InstanceId;
+
+				PreAnimatedMontageStorage->BeginTrackingEntity(Params.EntityID, Params.bWantsRestoreState, Params.RootInstanceHandle, Montage);
+				PreAnimatedMontageStorage->CachePreAnimatedValue(FCachePreAnimatedValueParams(), Montage, [=](UObject* Unused, FPreAnimatedSkeletalAnimationMontageState& OutState) {
+						OutState.WeakInstance = AnimInst;
+						OutState.MontageInstanceId = InstanceId;
+					});
+
+				// Make sure it's playing if the sequence is.
+				FAnimMontageInstance* Instance = AnimInst->GetMontageInstanceForID(InstanceId);
+				Instance->bPlaying = Params.bPlaying;
+			}
+		}
+	}
+
+	void PreviewSetAnimPosition(const FSetAnimPositionParams& Params)
+	{
+		using namespace UE::MovieScene;
+
+		static const bool bLooping = false;
+
+		const FMovieSceneSkeletalAnimationParams& AnimParams = Params.Section->Params;
+		if (!CanPlayAnimation(Params.SkeletalMeshComponent, AnimParams.Animation))
+		{
+			return;
+		}
+		if (AnimParams.bForceCustomMode)
+		{
+			Params.SkeletalMeshComponent->SetAnimationMode(EAnimationMode::AnimationCustomMode);
+		}
+		UAnimSequencerInstance* SequencerInst = GetAnimSequencerInstance(Params.SkeletalMeshComponent);
+		if (SequencerInst)
+		{
+			PreAnimatedAnimInstanceStorage->BeginTrackingEntity(Params.EntityID, Params.bWantsRestoreState, Params.RootInstanceHandle, SequencerInst);
+			PreAnimatedAnimInstanceStorage->CachePreAnimatedValue(FCachePreAnimatedValueParams(), SequencerInst);
+
+			TOptional <FRootMotionOverride> RootMotion;
+			UMovieSceneSkeletalAnimationSection::FRootMotionParams RootMotionParams;
+			Params.Section->GetRootMotion(Params.CurrentTime, RootMotionParams);
+			if (RootMotionParams.Transform.IsSet())
+			{
+				RootMotion = FRootMotionOverride();
+				RootMotion.GetValue().RootMotion = RootMotionParams.Transform.GetValue();
+				RootMotion.GetValue().ChildBoneIndex = RootMotionParams.ChildBoneIndex;
+				RootMotion.GetValue().bBlendFirstChildOfRoot = RootMotionParams.bBlendFirstChildOfRoot;
+				RootMotion.GetValue().PreviousTransform = RootMotionParams.PreviousTransform.GetValue();
+			}
+
+			// Use the section's address as the ID for the anim sequence.
+			const uint32 AnimSequenceID = GetTypeHash(Params.Section);
+
+			// If Sequencer has a transform track, we want to set the initial transform so that root motion (if it exists) can be applied relative to that.
+			TOptional<FTransform> CurrentTransform = GetCurrentTransform(Params.Section->Params.SwapRootBone, Params.SkeletalMeshComponent);
+
+			FAnimSequencerData AnimSequencerData(
+					AnimParams.Animation,
+					AnimSequenceID,
+					RootMotion,
+					Params.FromPosition,
+					Params.ToPosition,
+					Params.Weight,
+					Params.bFireNotifies,
+					Params.Section->Params.SwapRootBone,
+					CurrentTransform,
+					Params.Section->Params.MirrorDataTable.Get());
+			SequencerInst->UpdateAnimTrackWithRootMotion(AnimSequencerData);
+		}
+		else if (UAnimInstance* AnimInst = GetSourceAnimInstance(Params.SkeletalMeshComponent))
+		{
+			FMontagePlayerPerSectionData* SectionData = SystemData.MontageData.Find(Params.Section);
+
+			int32 InstanceId = SectionData ? SectionData->MontageInstanceId : INDEX_NONE;
+
+			const float AssetPlayRate = FMath::IsNearlyZero(AnimParams.Animation->RateScale) ? 1.0f : AnimParams.Animation->RateScale;
+			TWeakObjectPtr<UAnimMontage> WeakMontage = FAnimMontageInstance::PreviewSequencerMontagePosition(
+					AnimParams.SlotName,
+					Params.SkeletalMeshComponent,
+					AnimInst,
+					InstanceId,
+					AnimParams.Animation,
+					Params.FromPosition / AssetPlayRate,
+					Params.ToPosition / AssetPlayRate,
+					Params.Weight,
+					bLooping,
+					Params.bFireNotifies,
+					Params.bPlaying);
+
+			UAnimMontage* Montage = WeakMontage.Get();
+			if (Montage)
+			{
+				FMontagePlayerPerSectionData& DataContainer = SystemData.MontageData.FindOrAdd(Params.Section);
+				DataContainer.Montage = WeakMontage;
+				DataContainer.MontageInstanceId = InstanceId;
+
+				PreAnimatedMontageStorage->BeginTrackingEntity(Params.EntityID, Params.bWantsRestoreState, Params.RootInstanceHandle, Montage);
+				PreAnimatedMontageStorage->CachePreAnimatedValue(FCachePreAnimatedValueParams(), Montage, [=](UObject* Unused, FPreAnimatedSkeletalAnimationMontageState& OutState) {
+						OutState.WeakInstance = AnimInst;
+						OutState.MontageInstanceId = InstanceId;
+					});
+
+				FAnimMontageInstance* Instance = AnimInst->GetMontageInstanceForID(InstanceId);
+				Instance->bPlaying = Params.bPlaying;
+			}
+
+			if (Params.bResetDynamics)
+			{
+				// Make sure we reset any simulations.
+				AnimInst->ResetDynamics(ETeleportType::ResetPhysics);
+			}
+		}
+	}
+};
+
+} // namespace UE::MovieScene
+
+UMovieSceneSkeletalAnimationSystem::UMovieSceneSkeletalAnimationSystem(const FObjectInitializer& ObjInit)
+	: Super(ObjInit)
+{
+	using namespace UE::MovieScene;
+
+	FMovieSceneTracksComponentTypes* TrackComponents = FMovieSceneTracksComponentTypes::Get();
+	RelevantComponent = TrackComponents->SkeletalAnimation;
+	Phase = ESystemPhase::Instantiation | ESystemPhase::Evaluation;
+
+	SystemCategories |= FSystemInterrogator::GetExcludedFromInterrogationCategory();
+
+	if (HasAnyFlags(RF_ClassDefaultObject))
+	{
+		DefineImplicitPrerequisite(UWeightAndEasingEvaluatorSystem::StaticClass(), GetClass());
+		DefineImplicitPrerequisite(UMovieSceneComponentTransformSystem::StaticClass(), GetClass());
+		DefineImplicitPrerequisite(UMovieSceneQuaternionInterpolationRotationSystem::StaticClass(), GetClass());
+	}
+}
+
+void UMovieSceneSkeletalAnimationSystem::OnRun(FSystemTaskPrerequisites& InPrerequisites, FSystemSubsequentTasks& Subsequents)
+{
+	using namespace UE::MovieScene;
+
+	const TStatId GatherStatId = GET_STATID(MovieSceneEval_GatherSkeletalAnimations);
+
+	const FMovieSceneEntitySystemRunner* Runner = Linker->GetActiveRunner();
+	if (Runner->GetCurrentPhase() == ESystemPhase::Instantiation)
+	{
+		CleanSystemData();
+		return;
+	}
+
+	FBuiltInComponentTypes* BuiltInComponents = FBuiltInComponentTypes::Get();
+	FMovieSceneTracksComponentTypes* TrackComponents = FMovieSceneTracksComponentTypes::Get();
+
+	// Start fresh every frame, gathering all active skeletal animations.
+	SystemData.SkeletalAnimations.Reset();
+
+	FGraphEventRef GatherTask = FEntityTaskBuilder()
+	.ReadEntityIDs()
+	.Read(BuiltInComponents->RootInstanceHandle)
+	.Read(BuiltInComponents->InstanceHandle)
+	.Read(BuiltInComponents->BoundObject)
+	.Read(TrackComponents->SkeletalAnimation)
+	.ReadOptional(BuiltInComponents->WeightAndEasingResult)
+	.FilterNone({ BuiltInComponents->Tags.Ignored })
+	.SetStat(GatherStatId)
+	.Dispatch_PerAllocation<FGatherSkeletalAnimations>(&Linker->EntityManager, InPrerequisites, nullptr, 
+			Linker->GetInstanceRegistry(), SystemData);
+
+	FSystemTaskPrerequisites EvalPrereqs;
+	if (GatherTask)
+	{
+		EvalPrereqs.AddRootTask(GatherTask);
+	}
+
+	// Now evaluate gathered animations. We need to do this on the game thread (when in multi-threaded mode)
+	// because this task will call into a lot of animation system code that needs to be called there.
+	FEntityTaskBuilder()
+	.SetStat(GET_STATID(MovieSceneEval_EvaluateSkeletalAnimations))
+	.SetDesiredThread(Linker->EntityManager.GetGatherThread())
+	.Dispatch<FEvaluateSkeletalAnimations>(&Linker->EntityManager, EvalPrereqs, &Subsequents, Linker, SystemData);
+}
+
+void UMovieSceneSkeletalAnimationSystem::CleanSystemData()
+{
+	// Clean-up old montage data.
+	for (auto It = SystemData.MontageData.CreateIterator(); It; ++It)
+	{
+		if (It.Key().ResolveObjectPtr() == nullptr)
+		{
+			It.RemoveCurrent();
+		}
+	}
+}
+
diff --git a/Engine/Source/Runtime/MovieSceneTracks/Private/Tracks/MovieSceneSkeletalAnimationTrack.cpp b/Engine/Source/Runtime/MovieSceneTracks/Private/Tracks/MovieSceneSkeletalAnimationTrack.cpp
--- a/Engine/Source/Runtime/MovieSceneTracks/Private/Tracks/MovieSceneSkeletalAnimationTrack.cpp
+++ b/Engine/Source/Runtime/MovieSceneTracks/Private/Tracks/MovieSceneSkeletalAnimationTrack.cpp
@@ -6,8 +6,6 @@
 #include "Compilation/MovieSceneCompilerRules.h"
 #include "Components/SkeletalMeshComponent.h"
 #include "Evaluation/MovieSceneEvaluationTrack.h"
-#include "Evaluation/MovieSceneSkeletalAnimationTemplate.h"
-#include "Compilation/IMovieSceneTemplateGenerator.h"
 #include "Compilation/MovieSceneEvaluationTreePopulationRules.h"
 #include "MovieScene.h"
 #include "AnimationRuntime.h"
@@ -53,11 +51,6 @@ UMovieSceneSkeletalAnimationTrack::UMovieSceneSkeletalAnimationTrack(const FObje
 /* UMovieSceneSkeletalAnimationTrack interface
  *****************************************************************************/
 
-FMovieSceneEvalTemplatePtr UMovieSceneSkeletalAnimationTrack::CreateTemplateForSection(const UMovieSceneSection& InSection) const
-{
-	return FMovieSceneSkeletalAnimationSectionTemplate(*CastChecked<const UMovieSceneSkeletalAnimationSection>(&InSection));
-}
-
 UMovieSceneSection* UMovieSceneSkeletalAnimationTrack::AddNewAnimationOnRow(FFrameNumber KeyTime, UAnimSequenceBase* AnimSequence, int32 RowIndex)
 {
 	UMovieSceneSkeletalAnimationSection* NewSection = Cast<UMovieSceneSkeletalAnimationSection>(CreateNewSection());
diff --git a/Engine/Source/Runtime/MovieSceneTracks/Public/MovieSceneTracksComponentTypes.h b/Engine/Source/Runtime/MovieSceneTracks/Public/MovieSceneTracksComponentTypes.h
--- a/Engine/Source/Runtime/MovieSceneTracks/Public/MovieSceneTracksComponentTypes.h
+++ b/Engine/Source/Runtime/MovieSceneTracks/Public/MovieSceneTracksComponentTypes.h
@@ -17,9 +17,10 @@
 #include "MovieSceneTracksComponentTypes.generated.h"
 
 class UMaterialParameterCollection;
+class UMovieScene3DTransformSection;
 class UMovieSceneDataLayerSection;
 class UMovieSceneLevelVisibilitySection;
-class UMovieScene3DTransformSection;
+class UMovieSceneSkeletalAnimationSection;
 struct FMovieSceneObjectBindingID;
 
 
@@ -80,6 +81,16 @@ struct FConstraintComponentData
 	UMovieScene3DTransformSection* Section;
 };
 
+/** Component data for a skeletal mesh animation */
+USTRUCT()
+struct FMovieSceneSkeletalAnimationComponentData
+{
+	GENERATED_BODY()
+
+	UPROPERTY()
+	TObjectPtr<UMovieSceneSkeletalAnimationSection> Section;
+};
+
 namespace UE
 {
 namespace MovieScene
@@ -430,6 +441,8 @@ struct MOVIESCENETRACKS_API FMovieSceneTracksComponentTypes
 	TComponentTypeID<FPerlinNoiseParams> FloatPerlinNoiseChannel;
 	TComponentTypeID<FPerlinNoiseParams> DoublePerlinNoiseChannel;
 
+	TComponentTypeID<FMovieSceneSkeletalAnimationComponentData> SkeletalAnimation;
+
 	TComponentTypeID<int32> ComponentMaterialIndex;
 
 	TComponentTypeID<FName> BoolParameterName;
@@ -472,6 +485,5 @@ private:
 	FMovieSceneTracksComponentTypes();
 };
 
-
 } // namespace MovieScene
 } // namespace UE
diff --git a/Engine/Source/Runtime/MovieSceneTracks/Public/Sections/MovieSceneSkeletalAnimationSection.h b/Engine/Source/Runtime/MovieSceneTracks/Public/Sections/MovieSceneSkeletalAnimationSection.h
--- a/Engine/Source/Runtime/MovieSceneTracks/Public/Sections/MovieSceneSkeletalAnimationSection.h
+++ b/Engine/Source/Runtime/MovieSceneTracks/Public/Sections/MovieSceneSkeletalAnimationSection.h
@@ -7,6 +7,7 @@
 #include "MovieSceneSection.h"
 #include "Animation/AnimSequenceBase.h"
 #include "Channels/MovieSceneFloatChannel.h"
+#include "EntitySystem/IMovieSceneEntityProvider.h"
 #include "MovieSceneSkeletalAnimationSection.generated.h"
 
 struct FMovieSceneSkeletalAnimRootMotionTrackParams;
@@ -27,6 +28,17 @@ struct FMovieSceneSkeletalAnimationParams
 	/** Gets the animation sequence length, not modified by play rate */
 	float GetSequenceLength() const { return Animation != nullptr ? Animation->GetPlayLength() : 0.f; }
 
+	/**
+	 * Convert a sequence frame to a time in seconds inside the animation clip, taking into account start/end offsets,
+	 * looping, etc.
+	 */
+	double MapTimeToAnimation(const UMovieSceneSection* InSection, FFrameTime InPosition, FFrameRate InFrameRate) const;
+
+	/**
+	 * As above, but with already computed section bounds.
+	 */
+	double MapTimeToAnimation(FFrameNumber InSectionStartTime, FFrameNumber InSectionEndTime, FFrameTime InPosition, FFrameRate InFrameRate) const;
+
 	/** The animation this section plays */
 	UPROPERTY(BlueprintReadWrite, EditAnywhere, Category="Animation", meta=(AllowedClasses = "/Script/Engine.AnimSequence,/Script/Engine.AnimComposite,/Script/Engine.AnimStreamable"))
 	TObjectPtr<UAnimSequenceBase> Animation;
@@ -88,6 +100,7 @@ struct FMovieSceneSkeletalAnimationParams
 UCLASS( MinimalAPI )
 class UMovieSceneSkeletalAnimationSection
 	: public UMovieSceneSection
+	, public IMovieSceneEntityProvider
 {
 	GENERATED_UCLASS_BODY()
 
@@ -125,6 +138,8 @@ protected:
 	virtual void PostLoad() override;
 	virtual void Serialize(FArchive& Ar) override;
 
+	/** ~IMovieSceneEntityProvider interface */
+	virtual void ImportEntityImpl(UMovieSceneEntitySystemLinker* EntityLinker, const FEntityImportParams& InParams, FImportedEntity* OutImportedEntity) override;
 
 private:
 
diff --git a/Engine/Source/Runtime/MovieSceneTracks/Public/Systems/MovieSceneSkeletalAnimationSystem.h b/Engine/Source/Runtime/MovieSceneTracks/Public/Systems/MovieSceneSkeletalAnimationSystem.h
new file mode 100644
--- /dev/null
+++ b/Engine/Source/Runtime/MovieSceneTracks/Public/Systems/MovieSceneSkeletalAnimationSystem.h
@@ -0,0 +1,81 @@
+// Copyright Epic Games, Inc. All Rights Reserved.
+
+#pragma once
+
+#include "EntitySystem/MovieSceneEntitySystem.h"
+#include "EntitySystem/MovieSceneSequenceInstanceHandle.h"
+#include "Evaluation/MovieScenePlayback.h"
+#include "UObject/ObjectKey.h"
+#include "MovieSceneSkeletalAnimationSystem.generated.h"
+
+class IMovieScenePlayer;
+class UAnimMontage;
+class UMovieSceneSkeletalAnimationSection;
+
+namespace UE::MovieScene
+{
+
+/** Information for a single skeletal animation playing on a bound object */
+struct FActiveSkeletalAnimation
+{
+	IMovieScenePlayer* Player;
+	const UMovieSceneSkeletalAnimationSection* AnimSection;
+	FMovieSceneContext Context;
+	FMovieSceneEntityID EntityID;
+	FRootInstanceHandle RootInstanceHandle;
+	float FromEvalTime;
+	float ToEvalTime;
+	double BlendWeight;
+	bool bWantsRestoreState;
+};
+
+/** Information for all skeletal animations playing on a bound object */
+struct FBoundObjectActiveSkeletalAnimations
+{
+	using FAnimationArray = TArray<FActiveSkeletalAnimation, TInlineAllocator<2>>;
+
+	/** All active animations on the corresponding bound object */
+	FAnimationArray Animations;
+	/** Motion vector simulation animations on the corresponding bound object */
+	FAnimationArray SimulatedAnimations;
+};
+
+/** Temporary information about montage setups. */
+struct FMontagePlayerPerSectionData 
+{
+	TWeakObjectPtr<UAnimMontage> Montage;
+	int32 MontageInstanceId;
+};
+
+struct FSkeletalAnimationSystemData
+{
+	/** Map of active skeletal animations for each bound object */
+	TMap<UObject*, FBoundObjectActiveSkeletalAnimations> SkeletalAnimations;
+
+	/** Map of persistent montage data */
+	TMap<FObjectKey, FMontagePlayerPerSectionData> MontageData;
+};
+
+} // namespace UE::MovieScene
+
+UCLASS(MinimalAPI)
+class UMovieSceneSkeletalAnimationSystem
+	: public UMovieSceneEntitySystem
+{
+public:
+
+	GENERATED_BODY()
+
+	UMovieSceneSkeletalAnimationSystem(const FObjectInitializer& ObjInit);
+
+private:
+
+	virtual void OnRun(FSystemTaskPrerequisites& InPrerequisites, FSystemSubsequentTasks& Subsequents) override final;
+
+	void CleanSystemData();
+
+private:
+
+	UE::MovieScene::FSkeletalAnimationSystemData SystemData;
+};
+
diff --git a/Engine/Source/Runtime/MovieSceneTracks/Public/Tracks/MovieSceneSkeletalAnimationTrack.h b/Engine/Source/Runtime/MovieSceneTracks/Public/Tracks/MovieSceneSkeletalAnimationTrack.h
--- a/Engine/Source/Runtime/MovieSceneTracks/Public/Tracks/MovieSceneSkeletalAnimationTrack.h
+++ b/Engine/Source/Runtime/MovieSceneTracks/Public/Tracks/MovieSceneSkeletalAnimationTrack.h
@@ -5,7 +5,6 @@
 #include "CoreMinimal.h"
 #include "UObject/ObjectMacros.h"
 #include "MovieSceneNameableTrack.h"
-#include "Compilation/IMovieSceneTrackTemplateProducer.h"
 #include "Sections/MovieSceneSkeletalAnimationSection.h"
 #include "MovieSceneSkeletalAnimationTrack.generated.h"
 
@@ -37,7 +36,6 @@ struct FMovieSceneSkeletalAnimRootMotionTrackParams
 UCLASS(MinimalAPI)
 class UMovieSceneSkeletalAnimationTrack
 	: public UMovieSceneNameableTrack
-	, public IMovieSceneTrackTemplateProducer
 {
 	GENERATED_UCLASS_BODY()
 
@@ -77,9 +75,6 @@ public:
 	virtual bool SupportsMultipleRows() const override;
 	virtual void UpdateEasing() override;
 
-	// ~IMovieSceneTrackTemplateProducer interface
-	virtual FMovieSceneEvalTemplatePtr CreateTemplateForSection(const UMovieSceneSection& InSection) const override;
-
 #if WITH_EDITORONLY_DATA
 	virtual FText GetDefaultDisplayName() const override;
 #endif
