Sequencer: reverting piece of looping logic code

I misunderstood what it was doing before

diff --git a/Engine/Source/Runtime/MovieSceneTracks/Private/Sections/MovieSceneSkeletalAnimationSection.cpp b/Engine/Source/Runtime/MovieSceneTracks/Private/Sections/MovieSceneSkeletalAnimationSection.cpp
--- a/Engine/Source/Runtime/MovieSceneTracks/Private/Sections/MovieSceneSkeletalAnimationSection.cpp
+++ b/Engine/Source/Runtime/MovieSceneTracks/Private/Sections/MovieSceneSkeletalAnimationSection.cpp
@@ -59,7 +59,7 @@ double FMovieSceneSkeletalAnimationParams::MapTimeToAnimation(FFrameNumber InSec
 	const int32 LengthInFrames = AnimationLength.FrameNumber.Value + (int)(AnimationLength.GetSubFrame() + 0.5f) + 1;
 
 	// we only play end if we are not looping, and assuming we are looping if Length is greater than default length;
-	const bool bLooping = (InSectionEndTime.Value - InSectionStartTime.Value - StartFrameOffset - EndFrameOffset) > LengthInFrames;	
+	const bool bLooping = (InSectionEndTime.Value - InSectionStartTime.Value + StartFrameOffset + EndFrameOffset) > LengthInFrames;	
 
 	// Make sure InPosition FrameTime doesn't underflow InSectionStartTime or overflow InSectionEndTime
 	InPosition = FMath::Clamp(InPosition, FFrameTime(InSectionStartTime), FFrameTime(InSectionEndTime - 1));
