Sequencer: Skip past empty ranges when populating the material evaluation field

diff --git a/Engine/Source/Runtime/MovieSceneTracks/Private/Sections/MovieScenePrimitiveMaterialSection.cpp b/Engine/Source/Runtime/MovieSceneTracks/Private/Sections/MovieScenePrimitiveMaterialSection.cpp
--- a/Engine/Source/Runtime/MovieSceneTracks/Private/Sections/MovieScenePrimitiveMaterialSection.cpp
+++ b/Engine/Source/Runtime/MovieSceneTracks/Private/Sections/MovieScenePrimitiveMaterialSection.cpp
@@ -70,10 +70,13 @@ bool UMovieScenePrimitiveMaterialSection::PopulateEvaluationFieldImpl(const TRan
 
 		// Add the last range to the tree
 		TRange<FFrameNumber> Range(StartBound, TRangeBound<FFrameNumber>::Exclusive(Times[ValueIndex]));
-		OutFieldBuilder->AddPersistentEntity(Range, LastEntityIndex, MetaDataIndex);
+		if (!Range.IsEmpty())
+		{
+			OutFieldBuilder->AddPersistentEntity(Range, LastEntityIndex, MetaDataIndex);
 
-		StartBound = TRangeBound<FFrameNumber>::Inclusive(Times[ValueIndex]);
-		LastEntityIndex = OutFieldBuilder->FindOrAddEntity(this, ValueIndex);
+			StartBound = TRangeBound<FFrameNumber>::Inclusive(Times[ValueIndex]);
+			LastEntityIndex = OutFieldBuilder->FindOrAddEntity(this, ValueIndex);
+		}
 	}
 
 	TRange<FFrameNumber> TailRange(StartBound, EffectiveRange.GetUpperBound());
