Sequencer: NLA:  Trim and crash fix
#trivial

diff --git a/Engine/Source/Editor/MovieSceneTools/Private/EditModes/SkeletalAnimationTrackEditMode.cpp b/Engine/Source/Editor/MovieSceneTools/Private/EditModes/SkeletalAnimationTrackEditMode.cpp
--- a/Engine/Source/Editor/MovieSceneTools/Private/EditModes/SkeletalAnimationTrackEditMode.cpp
+++ b/Engine/Source/Editor/MovieSceneTools/Private/EditModes/SkeletalAnimationTrackEditMode.cpp
@@ -467,7 +467,7 @@ bool FSkeletalAnimationTrackEditMode::GetTransformAtFirstSectionStart(FTransform
 		int32 NumSelected = 0;
 		for (FSelectedRootData Data : SelectedRootData)
 		{
-			if (Data.SelectedMeshComp.IsValid() && Data.SelectedSection.IsValid() && Data.SelectedSection->GetRange().Contains(CurrentTime.Time.GetFrame()))
+			if (Data.SelectedMeshComp.IsValid() && Data.SelectedMeshComp->GetAnimInstance() && Data.SelectedSection.IsValid() && Data.SelectedSection->GetRange().Contains(CurrentTime.Time.GetFrame()))
 			{
 				FTransform ComponentTransform = Data.SelectedMeshComp->GetComponentTransform();
 				FTransform PivotTransform, PivotParentTransform;
@@ -536,7 +536,7 @@ void FSelectedRootData::CalcTransform(const FFrameTime& FrameTime, FTransform& O
 	USkeletalMeshComponent* Comp = SelectedMeshComp.Get();
 	OutTransform =  FTransform(FTransform::Identity);
 	OutParentTransform = FTransform(FTransform::Identity);
-	if (Section && Comp)
+	if (Section && Comp && Comp->GetAnimInstance())
 	{
 		if (UMovieScene* MovieScene = Section->GetTypedOuter<UMovieScene>())
 		{
diff --git a/Engine/Source/Runtime/MovieSceneTracks/Private/Sections/MovieSceneSkeletalAnimationSection.cpp b/Engine/Source/Runtime/MovieSceneTracks/Private/Sections/MovieSceneSkeletalAnimationSection.cpp
--- a/Engine/Source/Runtime/MovieSceneTracks/Private/Sections/MovieSceneSkeletalAnimationSection.cpp
+++ b/Engine/Source/Runtime/MovieSceneTracks/Private/Sections/MovieSceneSkeletalAnimationSection.cpp
@@ -652,7 +652,7 @@ bool UMovieSceneSkeletalAnimationSection::GetRootMotionTransform(FAnimationPoseD
 		if (!InOutParams.bOutIsAdditive)
 		{
 			const bool bStartTransformOffsetInBoneSpace = CVarStartTransformOffsetInBoneSpace.GetValueOnGameThread();
-			if (TempRootBoneIndex.IsSet() && TempRootBoneIndex.GetValue() != 0 && bStartTransformOffsetInBoneSpace)
+			if (bStartTransformOffsetInBoneSpace)
 			{
 				FTransform StartMatchedInRoot = StartBoneTransform * MatchedTransform;
 				FTransform LocalToRoot = (InOutParams.OutPoseTransform * StartBoneTransform.Inverse());
diff --git a/Engine/Source/Runtime/MovieSceneTracks/Private/Tracks/MovieSceneSkeletalAnimationTrack.cpp b/Engine/Source/Runtime/MovieSceneTracks/Private/Tracks/MovieSceneSkeletalAnimationTrack.cpp
--- a/Engine/Source/Runtime/MovieSceneTracks/Private/Tracks/MovieSceneSkeletalAnimationTrack.cpp
+++ b/Engine/Source/Runtime/MovieSceneTracks/Private/Tracks/MovieSceneSkeletalAnimationTrack.cpp
@@ -472,6 +472,11 @@ void UMovieSceneSkeletalAnimationTrack::SetUpRootMotions(bool bForce)
 			UMovieSceneSkeletalAnimationSection* AnimSection = Cast<UMovieSceneSkeletalAnimationSection>(Section);
 			if (AnimSection)
 			{
+				if (AnimSection->StartLocationOffset.IsNearlyZero() == false || AnimSection->StartRotationOffset.IsNearlyZero() == false ||
+					AnimSection->MatchedLocationOffset.IsNearlyZero() == false || AnimSection->MatchedRotationOffset.IsNearlyZero() == false)
+				{
+					bAnySectionsHaveOffset = true;
+				}
 				if (ValidAnimSequence == nullptr)
 				{
 					ValidAnimSequence = AnimSection->Params.Animation;
@@ -484,7 +489,6 @@ void UMovieSceneSkeletalAnimationTrack::SetUpRootMotions(bool bForce)
 							AnimSection->MatchedLocationOffset.IsNearlyZero() == false || AnimSection->MatchedRotationOffset.IsNearlyZero() == false)
 						{ 
 							RootMotionParams.RootMotionStartOffset = AnimSection->GetRootMotionStartOffset();
-							bAnySectionsHaveOffset = true;
 							FMemMark Mark(FMemStack::Get());
 							FCompactPose OutPose;
 							TArray<FBoneIndexType> RequiredBoneIndexArray;
