Sequencer: NLA If section is modified make sure we dirty root motion params

diff --git a/Engine/Source/Runtime/MovieSceneTracks/Private/Sections/MovieSceneSkeletalAnimationSection.cpp b/Engine/Source/Runtime/MovieSceneTracks/Private/Sections/MovieSceneSkeletalAnimationSection.cpp
--- a/Engine/Source/Runtime/MovieSceneTracks/Private/Sections/MovieSceneSkeletalAnimationSection.cpp
+++ b/Engine/Source/Runtime/MovieSceneTracks/Private/Sections/MovieSceneSkeletalAnimationSection.cpp
@@ -406,6 +406,15 @@ void UMovieSceneSkeletalAnimationSection::ImportEntityImpl(UMovieSceneEntitySyst
 }
 
 #if WITH_EDITOR
+bool UMovieSceneSkeletalAnimationSection::Modify(bool bAlwaysMarkDirty)
+{
+	if (GetRootMotionParams())
+	{
+		GetRootMotionParams()->bRootMotionsDirty = true;
+	}
+	return Super::Modify(bAlwaysMarkDirty);
+}
+
 void UMovieSceneSkeletalAnimationSection::PreEditChange(FProperty* PropertyAboutToChange)
 {
 	// Store the current play rate so that we can compute the amount to compensate the section end time when the play rate changes
diff --git a/Engine/Source/Runtime/MovieSceneTracks/Public/Sections/MovieSceneSkeletalAnimationSection.h b/Engine/Source/Runtime/MovieSceneTracks/Public/Sections/MovieSceneSkeletalAnimationSection.h
--- a/Engine/Source/Runtime/MovieSceneTracks/Public/Sections/MovieSceneSkeletalAnimationSection.h
+++ b/Engine/Source/Runtime/MovieSceneTracks/Public/Sections/MovieSceneSkeletalAnimationSection.h
@@ -146,10 +146,11 @@ private:
 	//~ UObject interface
 
 #if WITH_EDITOR
-
 	virtual void PreEditChange(FProperty* PropertyAboutToChange) override;
 	virtual void PostEditChangeProperty(FPropertyChangedEvent& PropertyChangedEvent) override;
+
 public:
+	virtual bool Modify(bool bAlwaysMarkDirty = true) override;
 	float PreviousPlayRate;
 private:
 	virtual void PostEditImport() override;
