Sequencer: fix dependencies for the material parameter system

The material parameter system wasn't depending on the double channel evaluator (or any system that produces double results). So the double channel evaluator could be originally sorted after the material parameter system in the system graph... until the material parameter system explictily links the blender system, at which point the double channel evaluator gets shuffled in front and trigger the assert about systems being inserted too early in the graph.

This fix adds the needed dependency, and adjusts the wording of the assert.

diff --git a/Engine/Source/Runtime/MovieScene/Private/EntitySystem/MovieSceneEntitySystemGraphs.cpp b/Engine/Source/Runtime/MovieScene/Private/EntitySystem/MovieSceneEntitySystemGraphs.cpp
--- a/Engine/Source/Runtime/MovieScene/Private/EntitySystem/MovieSceneEntitySystemGraphs.cpp
+++ b/Engine/Source/Runtime/MovieScene/Private/EntitySystem/MovieSceneEntitySystemGraphs.cpp
@@ -789,7 +789,9 @@ void FMovieSceneEntitySystemGraph::ExecutePhase(const ArrayType& SortedEntries,
 				{
 					const uint16 NewNodeIndex     = SortedEntries[NewIndex];
 					const uint16 CurrentNodeIndex = SortedEntries[CurrentIndex];
-					ensureAlwaysMsgf(false, TEXT("New system %s has been added upstream of %s in the same execution phase that is currently in-flight - this will not be run this frame"),
+					ensureAlwaysMsgf(false, 
+						TEXT("System %s has been inserted upstream of %s in the same execution phase that is currently in-flight, and will not be run this frame. "
+							 "This can be either because this system has been newly linked, or because it has been re-ordered due to other newly linked systems."),
 						*this->Nodes.Array[NewNodeIndex].System->GetName(),
 						*this->Nodes.Array[CurrentNodeIndex].System->GetName()
 					);
diff --git a/Engine/Source/Runtime/MovieSceneTracks/Private/Systems/MovieSceneMaterialParameterSystem.cpp b/Engine/Source/Runtime/MovieSceneTracks/Private/Systems/MovieSceneMaterialParameterSystem.cpp
--- a/Engine/Source/Runtime/MovieSceneTracks/Private/Systems/MovieSceneMaterialParameterSystem.cpp
+++ b/Engine/Source/Runtime/MovieSceneTracks/Private/Systems/MovieSceneMaterialParameterSystem.cpp
@@ -2,8 +2,8 @@
 
 #include "Systems/MovieSceneMaterialParameterSystem.h"
 
-#include "Systems/BoolChannelEvaluatorSystem.h"
 #include "Systems/FloatChannelEvaluatorSystem.h"
+#include "Systems/DoubleChannelEvaluatorSystem.h"
 #include "Systems/MovieScenePiecewiseDoubleBlenderSystem.h"
 #include "Systems/MovieSceneHierarchicalBiasSystem.h"
 #include "Systems/MovieSceneMaterialSystem.h"
@@ -381,6 +381,12 @@ UMovieSceneMaterialParameterSystem::UMovieSceneMaterialParameterSystem(const FOb
 		DefineComponentConsumer(GetClass(), TracksComponents->BoundMaterial);
 
 		DefineImplicitPrerequisite(UFloatChannelEvaluatorSystem::StaticClass(), GetClass());
+		DefineImplicitPrerequisite(UDoubleChannelEvaluatorSystem::StaticClass(), GetClass());
+
+		for (int32 Index = 0; Index < UE_ARRAY_COUNT(BuiltInComponents->DoubleResult); ++Index)
+		{
+			DefineComponentConsumer(GetClass(), BuiltInComponents->DoubleResult[Index]);
+		}
 
 		DefineImplicitPrerequisite(UMovieScenePiecewiseDoubleBlenderSystem::StaticClass(), GetClass());
 		DefineImplicitPrerequisite(GetClass(), UMovieSceneHierarchicalBiasSystem::StaticClass());
