Sequencer: reset extra metadata when the interrogator gets reset

This was an oversight when ExtraMetaData was added to the class.

diff --git a/Engine/Source/Runtime/MovieSceneTracks/Private/EntitySystem/Interrogation/MovieSceneInterrogationLinker.cpp b/Engine/Source/Runtime/MovieSceneTracks/Private/EntitySystem/Interrogation/MovieSceneInterrogationLinker.cpp
--- a/Engine/Source/Runtime/MovieSceneTracks/Private/EntitySystem/Interrogation/MovieSceneInterrogationLinker.cpp
+++ b/Engine/Source/Runtime/MovieSceneTracks/Private/EntitySystem/Interrogation/MovieSceneInterrogationLinker.cpp
@@ -667,6 +667,8 @@ void FSystemInterrogator::Reset()
 	EntityComponentField = FMovieSceneEntityComponentField();
 
 	Channels.Reset();
+	ExtraMetaData.Reset();
+
 	Linker->Reset();
 }
 
