Sequencer: ignore trying to restore pre-animated state on objects that have gone away

diff --git a/Engine/Source/Runtime/MovieSceneTracks/Private/Systems/MovieSceneSkeletalAnimationSystem.cpp b/Engine/Source/Runtime/MovieSceneTracks/Private/Systems/MovieSceneSkeletalAnimationSystem.cpp
--- a/Engine/Source/Runtime/MovieSceneTracks/Private/Systems/MovieSceneSkeletalAnimationSystem.cpp
+++ b/Engine/Source/Runtime/MovieSceneTracks/Private/Systems/MovieSceneSkeletalAnimationSystem.cpp
@@ -211,10 +211,13 @@ struct FPreAnimatedSkeletalAnimationAnimInstanceTraits : FBoundObjectPreAnimated
 
 	static void RestorePreAnimatedValue(const FObjectKey& Object, bool& Unused, const FRestoreStateParams& Params)
 	{
-		ISequencerAnimationSupport* SequencerAnimationSupport = Cast<ISequencerAnimationSupport>(Object.ResolveObjectPtr());
-		if (ensure(SequencerAnimationSupport))
+		if (UObject* ObjectPtr = Object.ResolveObjectPtr())
 		{
-			SequencerAnimationSupport->ResetNodes();
+			ISequencerAnimationSupport* SequencerAnimationSupport = Cast<ISequencerAnimationSupport>(ObjectPtr);
+			if (ensure(SequencerAnimationSupport))
+			{
+				SequencerAnimationSupport->ResetNodes();
+			}
 		}
 	}
 };
