Sequencer: Fixed rare edge case that meant material parameters could retain blend channels

diff --git a/Engine/Source/Runtime/MovieSceneTracks/Private/Systems/MovieSceneMaterialParameterSystem.cpp b/Engine/Source/Runtime/MovieSceneTracks/Private/Systems/MovieSceneMaterialParameterSystem.cpp
--- a/Engine/Source/Runtime/MovieSceneTracks/Private/Systems/MovieSceneMaterialParameterSystem.cpp
+++ b/Engine/Source/Runtime/MovieSceneTracks/Private/Systems/MovieSceneMaterialParameterSystem.cpp
@@ -338,6 +338,8 @@ struct TOverlappingMaterialParameterHandler : Mixin
 		}
 		else if (!Output->OutputEntityID && Inputs.Num() == 1)
 		{
+			Linker->EntityManager.RemoveComponent(Inputs[0], BuiltInComponents->BlendChannelInput);
+
 			Mixin::InitializeSoleInput(Linker, BoundMaterial, ParameterName, Inputs[0], Output);
 		}
 
