Constraints: Make handles transactional and add check to system eval to re-add a constraint if not present.

diff --git a/Engine/Plugins/Animation/ControlRig/Source/ControlRig/Private/ControlRig.cpp b/Engine/Plugins/Animation/ControlRig/Source/ControlRig/Private/ControlRig.cpp
--- a/Engine/Plugins/Animation/ControlRig/Source/ControlRig/Private/ControlRig.cpp
+++ b/Engine/Plugins/Animation/ControlRig/Source/ControlRig/Private/ControlRig.cpp
@@ -3378,6 +3378,7 @@ UTransformableControlHandle* UControlRig::CreateTransformableControlHandle(
 	
 	UTransformableControlHandle* CtrlHandle = NewObject<UTransformableControlHandle>(InOuter);
 	ensure(CtrlHandle);
+	CtrlHandle->SetFlags(RF_Transactional);
 	CtrlHandle->ControlRig = this;
 	CtrlHandle->ControlName = InControlName;
 	CtrlHandle->RegisterDelegates();
diff --git a/Engine/Source/Runtime/Experimental/Animation/Constraints/Private/TransformConstraint.cpp b/Engine/Source/Runtime/Experimental/Animation/Constraints/Private/TransformConstraint.cpp
--- a/Engine/Source/Runtime/Experimental/Animation/Constraints/Private/TransformConstraint.cpp
+++ b/Engine/Source/Runtime/Experimental/Animation/Constraints/Private/TransformConstraint.cpp
@@ -1318,6 +1318,7 @@ UTransformableComponentHandle* FTransformConstraintUtils::CreateHandleForSceneCo
 	if (InSceneComponent)
 	{
 		ComponentHandle = NewObject<UTransformableComponentHandle>(Outer);
+		ComponentHandle->SetFlags(RF_Transactional);
 		ComponentHandle->Component = InSceneComponent;
 		ComponentHandle->SocketName = InSocketName;
 		InSceneComponent->SetMobility(EComponentMobility::Movable);
diff --git a/Engine/Source/Runtime/MovieSceneTracks/Private/Systems/MovieSceneConstraintSystem.cpp b/Engine/Source/Runtime/MovieSceneTracks/Private/Systems/MovieSceneConstraintSystem.cpp
--- a/Engine/Source/Runtime/MovieSceneTracks/Private/Systems/MovieSceneConstraintSystem.cpp
+++ b/Engine/Source/Runtime/MovieSceneTracks/Private/Systems/MovieSceneConstraintSystem.cpp
@@ -121,12 +121,19 @@ void UMovieSceneConstraintSystem::OnRun(FSystemTaskPrerequisites& InPrerequisite
 					ConstraintChannel.ConstraintAndActiveChannel->Constraint.LoadSynchronous();
 				}
 				UTickableConstraint* Constraint = ConstraintChannel.ConstraintAndActiveChannel->Constraint.Get();
-				if (!Constraint)
+				if (!Constraint) //if constraint doesn't exist it probably got unspawned so recreate it and add it
 				{
 					UTickableConstraint* NewOne = Controller->AddConstraintFromCopy(ConstraintChannel.ConstraintAndActiveChannel->ConstraintCopyToSpawn);
 					Constraint = Controller->GetConstraint(ConstraintChannel.ConstraintName);
 					ConstraintChannel.Section->ReplaceConstraint(ConstraintChannel.ConstraintName, Constraint);
 				}
+				else // it's possible that we have it but it's not in the manager, due to manager not being saved with it (due to spawning or undo/redo).
+				{
+					if(Controller->GetConstraint(ConstraintChannel.ConstraintName) != Constraint)
+					{
+						Controller->AddConstraint(Constraint);
+					}
+				}
 				if (Constraint)
 				{
 					bool Result = false;
