Sequencer: Improved weight and easing evaluator for various hierarchical situations

  - Keeping track of allocated channels for each easing provider was problematic for many reasons, as it often required re-parenting channels and/or having to explicitly go through and allocate/free channels for child instances if a parent had weighting. Ultimately the previous approach had to consider too many situations to be effective, and the logic was becoming more and more convoluted.

  - The new code doesn't try to allocate and free persistent hierarchical weight channels at all - it simply regenerates the channels any time it needs to. This vastly simplifies the logic required for updating the channels because we don't need to care about what channels already exist - channels simply need adding, re-assigning or removing on instantiation. The measured cost of this is equivalent to the previous code.

  - Sub sections will now always create an entity that is tagged with SubInstance - this allows us to quickly enumerate active sub-sequences in the ECS without having to constantly ask the InstanceRegistry

  - When building the hierarchical channels we use the following process:
    1. Gather all active easing providers - these probably need a channel (they won't if there is no active instance for the Sequence ID its providing for)
    2. Gather all active sub instances _if they reside as a child of a sequence we gathered in (1)_
    3. Walk over everything we gathered in (1) and (2) and generate parent-first channels. Anything added in (2) and not (1) will receieve the same channel ID as its parent.
    4. Mutate all easing providers and sub instances to add/remove or update the HierarchicalEasingChannel component.
    5. Mutate anything with an instance handle (excluding (4)) and add/remove or update HierarchicalEasingChannel components.

diff --git a/Engine/Source/Runtime/MovieScene/Private/EntitySystem/BuiltInComponentTypes.cpp b/Engine/Source/Runtime/MovieScene/Private/EntitySystem/BuiltInComponentTypes.cpp
--- a/Engine/Source/Runtime/MovieScene/Private/EntitySystem/BuiltInComponentTypes.cpp
+++ b/Engine/Source/Runtime/MovieScene/Private/EntitySystem/BuiltInComponentTypes.cpp
@@ -141,6 +141,7 @@ FBuiltInComponentTypes::FBuiltInComponentTypes()
 	Tags.HasUnresolvedBinding    = ComponentRegistry->NewTag(TEXT("Has Unresolved Binding"));
 	Tags.MigratedFromFastPath    = ComponentRegistry->NewTag(TEXT("Migrated From Fast Path"));
 	Tags.Root                    = ComponentRegistry->NewTag(TEXT("Root"));
+	Tags.SubInstance             = ComponentRegistry->NewTag(TEXT("Sub Instance"));
 	Tags.ImportedEntity          = ComponentRegistry->NewTag(TEXT("Imported Entity"));
 	Tags.Finished                = ComponentRegistry->NewTag(TEXT("Finished Evaluating"));
 	Tags.Ignored                 = ComponentRegistry->NewTag(TEXT("Ignored"));
@@ -280,7 +281,6 @@ FBuiltInComponentTypes::FBuiltInComponentTypes()
 		// Easing components should be duplicated to children
 		ComponentRegistry->Factories.DuplicateChildComponent(Easing);
 		ComponentRegistry->Factories.DuplicateChildComponent(HierarchicalEasingChannel);
-		ComponentRegistry->Factories.DuplicateChildComponent(HierarchicalEasingProvider);
 
 		// Easing needs a time to evaluate
 		ComponentRegistry->Factories.DefineMutuallyInclusiveComponent(Easing, EvalTime);
diff --git a/Engine/Source/Runtime/MovieScene/Private/Evaluation/MovieSceneSequenceWeights.cpp b/Engine/Source/Runtime/MovieScene/Private/Evaluation/MovieSceneSequenceWeights.cpp
--- a/Engine/Source/Runtime/MovieScene/Private/Evaluation/MovieSceneSequenceWeights.cpp
+++ b/Engine/Source/Runtime/MovieScene/Private/Evaluation/MovieSceneSequenceWeights.cpp
@@ -102,6 +102,7 @@ void FSequenceWeights::RemoveWeight(FMovieSceneSequenceID InSequenceID)
 	if (EntityID && ensureAlways(Linker))
 	{
 		Linker->EntityManager.AddComponent(EntityID, FBuiltInComponentTypes::Get()->Tags.NeedsUnlink);
+		WeightEntitiesBySequenceID.Remove(InSequenceID);
 	}
 }
 
diff --git a/Engine/Source/Runtime/MovieScene/Private/MovieSceneSequencePlayer.cpp b/Engine/Source/Runtime/MovieScene/Private/MovieSceneSequencePlayer.cpp
--- a/Engine/Source/Runtime/MovieScene/Private/MovieSceneSequencePlayer.cpp
+++ b/Engine/Source/Runtime/MovieScene/Private/MovieSceneSequencePlayer.cpp
@@ -1841,3 +1841,17 @@ void UMovieSceneSequencePlayer::SetWeight(double InWeight, FMovieSceneSequenceID
 		SequenceWeights->SetWeight(SequenceID, InWeight);
 	}
 }
+
+void UMovieSceneSequencePlayer::RemoveWeight()
+{
+	RemoveWeight(MovieSceneSequenceID::Root);
+}
+
+void UMovieSceneSequencePlayer::RemoveWeight(FMovieSceneSequenceID SequenceID)
+{
+	UMovieSceneEntitySystemLinker* Linker = RootTemplateInstance.GetEntitySystemLinker();
+	if (Linker && SequenceWeights)
+	{
+		SequenceWeights->RemoveWeight(SequenceID);
+	}
+}
diff --git a/Engine/Source/Runtime/MovieScene/Private/Sections/MovieSceneSubSection.cpp b/Engine/Source/Runtime/MovieScene/Private/Sections/MovieSceneSubSection.cpp
--- a/Engine/Source/Runtime/MovieScene/Private/Sections/MovieSceneSubSection.cpp
+++ b/Engine/Source/Runtime/MovieScene/Private/Sections/MovieSceneSubSection.cpp
@@ -202,6 +202,7 @@ void UMovieSceneSubSection::ImportEntityImpl(UMovieSceneEntitySystemLinker* Enti
 	OutImportedEntity->AddBuilder(
 		FEntityBuilder().AddTag(FBuiltInComponentTypes::Get()->Tags.Root)
 	);
+
 	BuildDefaultSubSectionComponents(EntityLinker, Params, OutImportedEntity);
 }
 
@@ -484,17 +485,19 @@ void UMovieSceneSubSection::BuildDefaultSubSectionComponents(UMovieSceneEntitySy
 {
 	using namespace UE::MovieScene;
 
-	if (Easing.GetEaseInDuration() > 0 || Easing.GetEaseOutDuration() > 0)
-	{
-		FBuiltInComponentTypes* Components = FBuiltInComponentTypes::Get();
+	FBuiltInComponentTypes* Components = FBuiltInComponentTypes::Get();
 
-		const FSubSequencePath      PathToRoot         = EntityLinker->GetInstanceRegistry()->GetInstance(Params.Sequence.InstanceHandle).GetSubSequencePath();
-		const FMovieSceneSequenceID ResolvedSequenceID = PathToRoot.ResolveChildSequenceID(this->GetSequenceID());
+	const bool bHasEasing = Easing.GetEaseInDuration() > 0 || Easing.GetEaseOutDuration() > 0;
 
-		OutImportedEntity->AddBuilder(
-			FEntityBuilder().Add(Components->HierarchicalEasingProvider, ResolvedSequenceID)
-		);
-	}
+	const FSubSequencePath PathToRoot = EntityLinker->GetInstanceRegistry()->GetInstance(Params.Sequence.InstanceHandle).GetSubSequencePath();
+	FMovieSceneSequenceID ResolvedSequenceID = PathToRoot.ResolveChildSequenceID(this->GetSequenceID());
+
+	OutImportedEntity->AddBuilder(
+		FEntityBuilder()
+		.Add(Components->SequenceID, ResolvedSequenceID)
+		.AddTag(Components->Tags.SubInstance)
+		.AddConditional(Components->HierarchicalEasingProvider, ResolvedSequenceID, bHasEasing)
+	);
 }
 
 
diff --git a/Engine/Source/Runtime/MovieScene/Public/EntitySystem/BuiltInComponentTypes.h b/Engine/Source/Runtime/MovieScene/Public/EntitySystem/BuiltInComponentTypes.h
--- a/Engine/Source/Runtime/MovieScene/Public/EntitySystem/BuiltInComponentTypes.h
+++ b/Engine/Source/Runtime/MovieScene/Public/EntitySystem/BuiltInComponentTypes.h
@@ -386,6 +386,9 @@ public:
 		FComponentTypeID MigratedFromFastPath;
 
 		FComponentTypeID ImportedEntity;
+
+		FComponentTypeID SubInstance;
+
 		FComponentTypeID Root;
 
 		FComponentTypeID FixedTime;
diff --git a/Engine/Source/Runtime/MovieScene/Public/MovieSceneSequencePlayer.h b/Engine/Source/Runtime/MovieScene/Public/MovieSceneSequencePlayer.h
--- a/Engine/Source/Runtime/MovieScene/Public/MovieSceneSequencePlayer.h
+++ b/Engine/Source/Runtime/MovieScene/Public/MovieSceneSequencePlayer.h
@@ -263,6 +263,12 @@ public:
 	UFUNCTION(BlueprintCallable, Category = "Sequencer|Player")
 	void SetWeight(double InWeight);
 
+	/**
+	 * Removes a previously assigned weight
+	 */
+	UFUNCTION(BlueprintCallable, Category = "Sequencer|Player")
+	void RemoveWeight();
+
 	/**
 	 * Set a manual weight to be multiplied with all blendable elements within the specified sequence
 	 * @note: It is recommended that a weight between 0 and 1 is supplied, though this is not enforced
@@ -272,6 +278,11 @@ public:
 	 */
 	void SetWeight(double InWeight, FMovieSceneSequenceID SequenceID);
 
+	/**
+	 * Removes a previously assigned weight
+	 */
+	void RemoveWeight(FMovieSceneSequenceID SequenceID);
+
 public:
 
 	/**
diff --git a/Engine/Source/Runtime/MovieSceneTracks/Private/Systems/WeightAndEasingEvaluatorSystem.cpp b/Engine/Source/Runtime/MovieSceneTracks/Private/Systems/WeightAndEasingEvaluatorSystem.cpp
--- a/Engine/Source/Runtime/MovieSceneTracks/Private/Systems/WeightAndEasingEvaluatorSystem.cpp
+++ b/Engine/Source/Runtime/MovieSceneTracks/Private/Systems/WeightAndEasingEvaluatorSystem.cpp
@@ -9,6 +9,7 @@
 #include "EntitySystem/MovieSceneEntitySystemLinker.h"
 #include "EntitySystem/MovieSceneEntityMutations.h"
 #include "EntitySystem/MovieSceneRootInstantiatorSystem.h"
+#include "EntitySystem/MovieSceneComponentRegistry.h"
 #include "Evaluation/MovieSceneEvaluationTemplateInstance.h"
 #include "Systems/MovieScenePropertyInstantiator.h"
 #include "EntitySystem/Interrogation/MovieSceneInterrogatedPropertyInstantiator.h"
@@ -33,143 +34,400 @@ namespace UE::MovieScene
 
 static constexpr uint16 INVALID_EASING_CHANNEL = uint16(-1);
 
-struct FAddEasingChannelToProviderMutation : IMovieSceneEntityMutation
+enum class EHierarchicalSequenceChannelFlags
 {
-	FAddEasingChannelToProviderMutation(UMovieSceneHierarchicalEasingInstantiatorSystem* InSystem)
-		: System(InSystem)
-		, BuiltInComponents(FBuiltInComponentTypes::Get())
-		, InstanceRegistry(InSystem->GetLinker()->GetInstanceRegistry())
+	None = 0,
+	SharedWithParent = 1,
+};
+ENUM_CLASS_FLAGS(EHierarchicalSequenceChannelFlags)
+
+/**
+ * Class that accumulates active easing data for a particular root sequence hierarchy
+ */
+struct FRootSequenceData
+{
+	explicit FRootSequenceData(FRootInstanceHandle InRootInstanceHandle, const FInstanceRegistry* InInstanceRegistry)
+		: InstanceRegistry(InInstanceRegistry)
+		, RootHierarchy(nullptr) // defer initializing the hierarchy until we need it
+		, RootInstanceHandle(InRootInstanceHandle)
 	{}
 
-private:
+	/** Retrieve the index of the specified sequence ID if it has already been encountered */
+	int32 IndexOf(FMovieSceneSequenceID SequenceID) const
+	{
+		const int32 Index = Algo::LowerBoundBy(SequenceIDs, SequenceID, &FSequenceIDAndChannel::SequenceID);
+		if (SequenceIDs.IsValidIndex(Index) && SequenceIDs[Index].SequenceID == SequenceID)
+		{
+			return Index;
+		}
+		return INDEX_NONE;
+	}
 
-	void CreateMutation(FEntityManager* EntityManager, FComponentMask* InOutEntityComponentTypes) const override
+	/* Retrieve the assigned channel ID for the specified sequence ID. Only valid once ProduceChannels has been called. */
+	uint16 GetChannelID(FMovieSceneSequenceID SequenceID) const
 	{
-		InOutEntityComponentTypes->Set(BuiltInComponents->HierarchicalEasingChannel);
+		const int32 ChannelIndex = IndexOf(SequenceID);
+		if (ChannelIndex != INDEX_NONE)
+		{
+			return SequenceIDs[ChannelIndex].ChannelID;
+		}
+		return INVALID_EASING_CHANNEL;
 	}
 
-	void InitializeAllocation(FEntityAllocation* Allocation, const FComponentMask& AllocationType) const override
+	/** Add the specified sequence ID to our map */
+	int32 Add(FMovieSceneSequenceID SequenceID, EHierarchicalSequenceChannelFlags InFlags)
 	{
-		AllocateEasingChannelsForAllocation(Allocation);
+		const int32 Index = Algo::LowerBoundBy(SequenceIDs, SequenceID, &FSequenceIDAndChannel::SequenceID);
+		check(!SequenceIDs.IsValidIndex(Index) ||SequenceIDs[Index].SequenceID != SequenceID);
+
+		SequenceIDs.Insert(FSequenceIDAndChannel{ SequenceID, INVALID_EASING_CHANNEL, InFlags }, Index);
+
+		return Index;
 	}
 
-	void InitializeUnmodifiedAllocation(FEntityAllocation* Allocation, const FComponentMask& AllocationType) const override
+	/** Add the specified sequence ID to our map if it resides as a child underneath one that has an active weight */
+	bool AddImplicitChild(FMovieSceneSequenceID SequenceID)
 	{
-		AllocateEasingChannelsForAllocation(Allocation);
+		const bool bHasChannel = IndexOf(SequenceID) != INDEX_NONE;
+		if (bHasChannel)
+		{
+			return true;
+		}
+
+		if (!RootHierarchy)
+		{
+			RootHierarchy = InstanceRegistry->GetInstance(RootInstanceHandle).GetPlayer()->GetEvaluationTemplate().GetHierarchy();
+		}
+
+		if (ensure(RootHierarchy))
+		{
+			const FMovieSceneSequenceHierarchyNode* Node = RootHierarchy->FindNode(SequenceID);
+			check(Node);
+
+			if (Node->ParentID != MovieSceneSequenceID::Invalid && AddImplicitChild(Node->ParentID))
+			{
+				Add(SequenceID, EHierarchicalSequenceChannelFlags::SharedWithParent);
+				return true;
+			}
+		}
+		return false;
 	}
 
-	void AllocateEasingChannelsForAllocation(FEntityAllocation* Allocation) const
+	/**
+	 * Produce channel IDs in the specified output data, ordered parent-first
+	 */
+	void ProduceChannels(TSparseArray<uint16>& OutInstanceIDToChannel, TArray<FHierarchicalEasingChannelData>& OutData)
 	{
-		FEntityAllocationWriteContext NewAllocation = FEntityAllocationWriteContext::NewAllocation();
+		const FSequenceInstance& RootInstance = InstanceRegistry->GetInstance(RootInstanceHandle);
+		for (int32 Index = 0; Index < SequenceIDs.Num(); ++Index)
+		{
+			ProduceChannelForIndex(Index, RootInstance, OutInstanceIDToChannel, OutData);
+		}
+	}
 
-		TComponentLock<TWrite<uint16>>               EasingChannels      = Allocation->WriteComponents(BuiltInComponents->HierarchicalEasingChannel, NewAllocation);
-		TComponentLock<TRead<FRootInstanceHandle>>   RootInstanceHandles = Allocation->ReadComponents(BuiltInComponents->RootInstanceHandle);
-		TComponentLock<TRead<FMovieSceneSequenceID>> SubSequenceIDs      = Allocation->ReadComponents(BuiltInComponents->HierarchicalEasingProvider);
-		
-		for (int32 Index = 0; Index < Allocation->Num(); ++Index)
+	/**
+	 * Produce channels for the specified index within our map, including any of its parents
+	 */
+	uint16 ProduceChannelForIndex(int32 Index, const FSequenceInstance& RootInstance, TSparseArray<uint16>& OutInstanceIDToChannel, TArray<FHierarchicalEasingChannelData>& OutData)
+	{
+		const uint16 ExistingChannel = SequenceIDs[Index].ChannelID;
+		if (ExistingChannel != INVALID_EASING_CHANNEL)
 		{
-			EasingChannels[Index] = System->AllocateEasingChannel(InstanceRegistry, RootInstanceHandles[Index], SubSequenceIDs[Index]).EasingChannelID;
+			return ExistingChannel;
+		}
+
+		FSequenceIDAndChannel ChannelData = SequenceIDs[Index];
+
+		if (ChannelData.SequenceID == MovieSceneSequenceID::Root)
+		{
+			// Add a new channel with no parent
+			const uint16 ChannelID = static_cast<uint16>(OutData.Num());
+			OutData.Add(FHierarchicalEasingChannelData{ INVALID_EASING_CHANNEL });
+			OutInstanceIDToChannel.Insert(RootInstanceHandle.InstanceID, ChannelID);
+
+			SequenceIDs[Index].ChannelID = ChannelID;
+			return ChannelID;
+		}
+
+		if (!RootHierarchy)
+		{
+			RootHierarchy = InstanceRegistry->GetInstance(RootInstanceHandle).GetPlayer()->GetEvaluationTemplate().GetHierarchy();
+		}
+
+		if (!RootHierarchy)
+		{
+			return INVALID_EASING_CHANNEL;
+		}
+
+		// Set up hierarchy info for non-root sequences
+		const FMovieSceneSequenceHierarchyNode* Node    = RootHierarchy->FindNode(ChannelData.SequenceID);
+		const FMovieSceneSubSequenceData*       SubData = RootHierarchy->FindSubData(ChannelData.SequenceID);
+		if (!Node || !SubData)
+		{
+			return INVALID_EASING_CHANNEL;
+		}
+
+		uint16 ParentChannel = INVALID_EASING_CHANNEL;
+		FMovieSceneSequenceID Parent = MovieSceneSequenceID::Invalid;
+		while (Node)
+		{
+			const int32 ParentIndex = IndexOf(Node->ParentID);
+			if (ParentIndex != INDEX_NONE)
+			{
+				ParentChannel = ProduceChannelForIndex(ParentIndex, RootInstance, OutInstanceIDToChannel, OutData);
+				break;
+			}
+
+			Node = RootHierarchy->FindNode(Node->ParentID);
+		}
+
+		// Don't bother creating channels for sequences that don't have an instance
+		FInstanceHandle SubInstanceHandle = RootInstance.FindSubInstance(ChannelData.SequenceID);
+		if (!SubInstanceHandle.IsValid())
+		{
+			return ParentChannel;
+		}
+
+		if (EnumHasAnyFlags(ChannelData.Flags, EHierarchicalSequenceChannelFlags::SharedWithParent))
+		{
+			OutInstanceIDToChannel.Insert(SubInstanceHandle.InstanceID, ParentChannel);
+			SequenceIDs[Index].ChannelID = ParentChannel;
+			return ParentChannel;
+		}
+		else
+		{
+			const uint16 ChannelID = static_cast<uint16>(OutData.Num());
+			OutData.Add(FHierarchicalEasingChannelData{ ParentChannel });
+			OutInstanceIDToChannel.Insert(SubInstanceHandle.InstanceID, ChannelID);
+			SequenceIDs[Index].ChannelID = ChannelID;
+			return ChannelID;
 		}
 	}
+
 private:
 
-	UMovieSceneHierarchicalEasingInstantiatorSystem* System;
-	FBuiltInComponentTypes* BuiltInComponents;
-	FInstanceRegistry* InstanceRegistry;
-};
+	struct FSequenceIDAndChannel
+	{
+		FMovieSceneSequenceID SequenceID;
+		uint16 ChannelID;
+		EHierarchicalSequenceChannelFlags Flags;
+	};
 
+	const FInstanceRegistry* InstanceRegistry;
+	const FMovieSceneSequenceHierarchy* RootHierarchy;
+	TArray<FSequenceIDAndChannel, TInlineAllocator<32>> SequenceIDs;
+	FRootInstanceHandle RootInstanceHandle;
+};
 
+using FRootSequenceDataMap = TSortedMap<FRootInstanceHandle, FRootSequenceData, TInlineAllocator<8>>;
 
-struct FAddEasingChannelToConsumerMutation : IMovieSceneConditionalEntityMutation
+/**
+ * Mutation that adds easing channel components to consumers that have active weights
+ */
+struct FEasingChannelMutationBase : IMovieSceneConditionalEntityMutation
 {
-	FEntityManager* EntityManager;
-	FAddEasingChannelToConsumerMutation(UMovieSceneHierarchicalEasingInstantiatorSystem* InSystem)
-		: System(InSystem)
-		, BuiltInComponents(FBuiltInComponentTypes::Get())
-		, InstanceRegistry(InSystem->GetLinker()->GetInstanceRegistry())
+	explicit FEasingChannelMutationBase(UMovieSceneEntitySystemLinker* Linker)
 	{
-		EntityManager = &InSystem->GetLinker()->EntityManager;
+		BuiltInComponents = FBuiltInComponentTypes::Get();
+		Factories = &Linker->EntityManager.GetComponents()->Factories;
 	}
 
-private:
+	/**
+	 * Create the new entity type for a group of entities
+	 */
+	void CreateMutation(FEntityManager* InEntityManager, FComponentMask* InOutEntityComponentTypes) const override
+	{
+		InOutEntityComponentTypes->Set(BuiltInComponents->HierarchicalEasingChannel);
+		Factories->ComputeMutuallyInclusiveComponents(*InOutEntityComponentTypes);
+		InOutEntityComponentTypes->Set(BuiltInComponents->Tags.NeedsLink);
+	}
 
-	void MarkAllocation(FEntityAllocation* Allocation, TBitArray<>& OutEntitiesToMutate) const override
+	bool HasStaleComponents() const
 	{
-		TComponentReader<FRootInstanceHandle> RootInstanceHandles = Allocation->ReadComponents(BuiltInComponents->RootInstanceHandle);
-		TOptionalComponentReader<FMovieSceneSequenceID> SubSequenceIDs = Allocation->TryReadComponents(BuiltInComponents->SequenceID);
+		return StaleEntities.Num() != 0;
+	}
 
-		if (SubSequenceIDs)
+	void RemoveStaleComponents(UMovieSceneEntitySystemLinker* Linker)
+	{
+		if (StaleEntities.Num())
 		{
-			for (int32 Index = 0; Index < Allocation->Num(); ++Index)
+			FComponentMask ComponentsToRemove;
+			ComponentsToRemove.Set(BuiltInComponents->HierarchicalEasingChannel);
+			ComponentsToRemove.Set(BuiltInComponents->WeightAndEasingResult);
+
+			for (FMovieSceneEntityID Entity : StaleEntities)
 			{
-				const uint16 EasingChannel = System->LocateEasingChannel(RootInstanceHandles[Index], SubSequenceIDs[Index]);
-				if (EasingChannel != uint16(-1))
-				{
-					OutEntitiesToMutate.PadToNum(Index + 1, false);
-					OutEntitiesToMutate[Index] = true;
-				}
+				Linker->EntityManager.RemoveComponents(Entity, ComponentsToRemove);
 			}
 		}
-		else
+	}
+
+protected:
+
+	FBuiltInComponentTypes* BuiltInComponents;
+	FEntityFactories* Factories;
+	mutable TArray<FMovieSceneEntityID> StaleEntities;
+};
+
+/**
+ * Mutation that adds or assigns easing channels on easing providers
+ */
+struct FProviderEasingChannelMutation : FEasingChannelMutationBase
+{
+	explicit FProviderEasingChannelMutation(UMovieSceneEntitySystemLinker* Linker, const FRootSequenceDataMap* InSequenceDataMap)
+		: FEasingChannelMutationBase(Linker)
+		, SequenceDataMap(InSequenceDataMap)
+	{}
+
+	/**
+	 * Called on matching allocations to mark specific entities that need mutating.
+	 * If the entity exists within a weighted sequence (directly or indirectly), a channel will be created (or referenced) and assigned
+	 */
+	void MarkAllocation(FEntityAllocation* Allocation, TBitArray<>& OutEntitiesToMutate) const override
+	{
+		const FMovieSceneEntityID* EntityIDs = Allocation->GetRawEntityIDs();
+
+		TComponentReader<FRootInstanceHandle>   RootInstanceHandles   = Allocation->ReadComponents(BuiltInComponents->RootInstanceHandle);
+		TComponentReader<FMovieSceneSequenceID> EasingProviders       = Allocation->ReadComponents(BuiltInComponents->HierarchicalEasingProvider);
+		TOptionalComponentWriter<uint16>        ExistingEasingHandles = Allocation->TryWriteComponents(BuiltInComponents->HierarchicalEasingChannel, FEntityAllocationWriteContext::NewAllocation());
+
+		const bool bHasExistingEasing = ExistingEasingHandles.IsValid();
+
+		// Loop through each entity and check whether it has easing.
+		// If so, either assign the new easing channel or mark it for mutation.
+		// If not, mark it for removal
+		for (int32 Index = 0; Index < Allocation->Num(); ++Index)
 		{
-			for (int32 Index = 0; Index < Allocation->Num(); ++Index)
+			const FRootSequenceData* ChannelData = SequenceDataMap->Find(RootInstanceHandles[Index]);
+			const uint16             ChannelID   = ChannelData ? ChannelData->GetChannelID(EasingProviders[Index]) : INVALID_EASING_CHANNEL;
+
+			if (ChannelID != INVALID_EASING_CHANNEL)
 			{
-				const uint16 EasingChannel = System->LocateEasingChannel(RootInstanceHandles[Index], MovieSceneSequenceID::Root);
-				if (EasingChannel != uint16(-1))
+				if (bHasExistingEasing)
+				{
+					ExistingEasingHandles[Index] = ChannelID;
+				}
+				else
 				{
+					// Mark this entity for mutation
 					OutEntitiesToMutate.PadToNum(Index + 1, false);
 					OutEntitiesToMutate[Index] = true;
 				}
 			}
+			else if (bHasExistingEasing)
+			{
+				StaleEntities.Add(EntityIDs[Index]);
+			}
 		}
 	}
 
-	void CreateMutation(FEntityManager* InEntityManager, FComponentMask* InOutEntityComponentTypes) const override
+	/**
+	 * Initialize a range of new entities with their easing channels
+	 */
+	void InitializeEntities(const FEntityRange& EntityRange, const FComponentMask& AllocationType) const override
 	{
-		InOutEntityComponentTypes->Set(BuiltInComponents->HierarchicalEasingChannel);
+		TComponentReader<FRootInstanceHandle>   RootInstanceHandles = EntityRange.Allocation->ReadComponents(BuiltInComponents->RootInstanceHandle);
+		TComponentReader<FMovieSceneSequenceID> EasingProviders     = EntityRange.Allocation->ReadComponents(BuiltInComponents->HierarchicalEasingProvider);
+		TComponentWriter<uint16>                EasingChannels      = EntityRange.Allocation->WriteComponents(BuiltInComponents->HierarchicalEasingChannel, FEntityAllocationWriteContext::NewAllocation());
 
-		InEntityManager->GetComponents()->Factories.ComputeMutuallyInclusiveComponents(*InOutEntityComponentTypes);
+		for (int32 Index = 0; Index < EntityRange.Num; ++Index)
+		{
+			const int32 Offset = EntityRange.ComponentStartOffset + Index;
 
-		InOutEntityComponentTypes->Set(BuiltInComponents->Tags.NeedsLink);
+			const FRootSequenceData& ChannelData = SequenceDataMap->FindChecked(RootInstanceHandles[Offset]);
+			EasingChannels[Offset] = ChannelData.GetChannelID(EasingProviders[Offset]);
+		}
 	}
 
+private:
+
+	const FRootSequenceDataMap* SequenceDataMap;
+};
+
+/**
+ * Mutation that adds easing channel components to consumers that have active weights
+ */
+struct FConsumerEasingChannelMutation : FEasingChannelMutationBase
+{
+	explicit FConsumerEasingChannelMutation(UMovieSceneEntitySystemLinker* Linker, TSparseArray<uint16>* InInstanceIDToChannel)
+		: FEasingChannelMutationBase(Linker)
+		, InstanceIDToChannel(InInstanceIDToChannel)
+	{}
+
+	/**
+	 * Called on matching allocations to mark specific entities that need mutating.
+	 * If the entity exists within a weighted sequence (directly or indirectly), a channel will be created (or referenced) and assigned
+	 */
+	void MarkAllocation(FEntityAllocation* Allocation, TBitArray<>& OutEntitiesToMutate) const override
+	{
+		const FMovieSceneEntityID* EntityIDs = Allocation->GetRawEntityIDs();
+
+		TComponentReader<FInstanceHandle> InstanceHandles       = Allocation->ReadComponents(BuiltInComponents->InstanceHandle);
+		TOptionalComponentWriter<uint16>  ExistingEasingHandles = Allocation->TryWriteComponents(BuiltInComponents->HierarchicalEasingChannel, FEntityAllocationWriteContext::NewAllocation());
+
+		const bool bHasExistingEasing = ExistingEasingHandles.IsValid();
+
+		// Loop through each entity and check whether it has easing.
+		// If so, either assign the new easing channel or mark it for mutation.
+		// If not, mark it for removal
+		for (int32 Index = 0; Index < Allocation->Num(); ++Index)
+		{
+			if (InstanceIDToChannel->IsValidIndex(InstanceHandles[Index].InstanceID))
+			{
+				if (bHasExistingEasing)
+				{
+					ExistingEasingHandles[Index] = (*InstanceIDToChannel)[InstanceHandles[Index].InstanceID];
+				}
+				else
+				{
+					// Mark this entity for mutation
+					OutEntitiesToMutate.PadToNum(Index + 1, false);
+					OutEntitiesToMutate[Index] = true;
+				}
+			}
+			else if (bHasExistingEasing)
+			{
+				StaleEntities.Add(EntityIDs[Index]);
+			}
+		}
+	}
+
+	/**
+	 * Initialize a range of new entities with their easing channels
+	 */
 	void InitializeEntities(const FEntityRange& EntityRange, const FComponentMask& AllocationType) const override
 	{
 		FEntityAllocationWriteContext NewAllocation = FEntityAllocationWriteContext::NewAllocation();
 
-		EntityManager->GetComponents()->Factories.RunInitializers(AllocationType, EntityRange);
-
-		TComponentWriter<uint16>                        EasingChannels      = EntityRange.Allocation->WriteComponents(BuiltInComponents->HierarchicalEasingChannel, NewAllocation);
-		TComponentReader<FRootInstanceHandle>           RootInstanceHandles = EntityRange.Allocation->ReadComponents(BuiltInComponents->RootInstanceHandle);
-		TOptionalComponentReader<FMovieSceneSequenceID> OptSequenceIDs      = EntityRange.Allocation->TryReadComponents(BuiltInComponents->SequenceID);
+		TComponentWriter<uint16>              EasingChannels      = EntityRange.Allocation->WriteComponents(BuiltInComponents->HierarchicalEasingChannel, NewAllocation);
+		TComponentReader<FInstanceHandle>     InstanceHandles     = EntityRange.Allocation->ReadComponents(BuiltInComponents->InstanceHandle);
 
 		for (int32 Index = 0; Index < EntityRange.Num; ++Index)
 		{
 			const int32 Offset = EntityRange.ComponentStartOffset + Index;
-
-			FMovieSceneSequenceID SequenceID = OptSequenceIDs ? OptSequenceIDs[Offset] : MovieSceneSequenceID::Root;
-			const uint16 EasingChannel = System->LocateEasingChannel(RootInstanceHandles[Offset], SequenceID);
-			check(EasingChannel != uint16(-1));
-			EasingChannels[Offset] = EasingChannel;
+			EasingChannels[Offset] = (*InstanceIDToChannel)[InstanceHandles[Offset].InstanceID];
 		}
 	}
+
 private:
 
-	UMovieSceneHierarchicalEasingInstantiatorSystem* System;
-	FBuiltInComponentTypes* BuiltInComponents;
-	FInstanceRegistry* InstanceRegistry;
+	TSparseArray<uint16>* InstanceIDToChannel;
 };
 
-
+/**
+ * Reset WeightAndEasingResult components before accumulation
+ */
 struct FResetFinalWeightResults
 {
-
 	static void ForEachEntity(double& Result)
 	{
 		Result = 1.f;
 	}
 };
 
+/**
+ * Evaluate section easing into EasingResult components
+ */
 struct FEvaluateEasings
 {
 	static void ForEachEntity(FFrameTime EvalTime, const FEasingComponentData& Easing, double& Result)
@@ -179,6 +437,9 @@ struct FEvaluateEasings
 	}
 };
 
+/**
+ * Accumulate EasingResult and WeightResult components into WeightAndEasingResult components
+ */
 struct FAccumulateManualWeights
 {
 	static void ForEachAllocation(const FEntityAllocation* Allocation, const TReadOneOrMoreOf<double, double>& Results, double* OutAccumulatedResults)
@@ -207,15 +468,15 @@ struct FAccumulateManualWeights
 	}
 };
 
-
+/**
+ * Harvest hierarchical WeightAndEasingResult components from HierarchicalEasingProvider components into a linear buffer for propagation to consumers
+ */
 struct FHarvestHierarchicalEasings
 {
-	TSparseArray<int32>* EasingChannelToIndex;
 	TArrayView<FHierarchicalEasingChannelData> ComputationData;
 
-	FHarvestHierarchicalEasings(TSparseArray<int32>* InEasingChannelToIndex, TArray<FHierarchicalEasingChannelData>* InComputationData)
-		: EasingChannelToIndex(InEasingChannelToIndex)
-		, ComputationData(*InComputationData)
+	FHarvestHierarchicalEasings(TArray<FHierarchicalEasingChannelData>* InComputationData)
+		: ComputationData(*InComputationData)
 	{
 		check(InComputationData);
 	}
@@ -232,8 +493,7 @@ struct FHarvestHierarchicalEasings
 	// Accumulate all entities that contribute to the channel
 	void ForEachEntity(double Result, uint16 EasingChannel)
 	{
-		const int32 ResultIndex = (*EasingChannelToIndex)[EasingChannel];
-		ComputationData[ResultIndex].FinalResult *= Result;
+		ComputationData[EasingChannel].FinalResult *= Result;
 	}
 
 	// Multiply hierarchical weights with sub sequences
@@ -244,49 +504,35 @@ struct FHarvestHierarchicalEasings
 		for (int32 Index = 0; Index < ComputationData.Num(); ++Index)
 		{
 			FHierarchicalEasingChannelData ChannelData = ComputationData[Index];
-			if (ChannelData.ParentEasingIndex != uint16(-1))
+			if (ChannelData.ParentChannel != uint16(-1))
 			{
 				// The parent result has already been multiplied by all its parent weights by this point
-				ComputationData[Index].FinalResult *= ComputationData[ChannelData.ParentEasingIndex].FinalResult;
+				ComputationData[Index].FinalResult *= ComputationData[ChannelData.ParentChannel].FinalResult;
 			}
 		}
 	}
 };
 
+/**
+ * Propagate hierarchical weight values to WeightAndEasingResult components on consumers
+ */
 struct FPropagateHierarchicalEasings
 {
 	TArrayView<const FHierarchicalEasingChannelData> ComputationData;
 
-	TArray<double> HierarchicalResultsByChannelID;
-	int32 MaxChannelNum;
-
-	FPropagateHierarchicalEasings(TArrayView<const FHierarchicalEasingChannelData> InComputationData, const int32 InMaxChannelNum)
+	FPropagateHierarchicalEasings(TArrayView<const FHierarchicalEasingChannelData> InComputationData)
 		: ComputationData(InComputationData)
-		, MaxChannelNum(InMaxChannelNum)
 	{}
 
-	// Before the task runs, initialize the results array to avoid a double indirection during the expansion in the actual task
-	void PreTask()
-	{
-		HierarchicalResultsByChannelID.SetNumZeroed(MaxChannelNum);
-
-		for (int32 Index = 0; Index < ComputationData.Num(); ++Index)
-		{
-			const int32 ChannelIdAsIndex = ComputationData[Index].ChannelID;
-			HierarchicalResultsByChannelID[ChannelIdAsIndex] = ComputationData[Index].FinalResult;
-		}
-	}
-
 	void ForEachAllocation(const FEntityAllocation* Allocation, TRead<uint16> HierarchicalEasingChannels, TWrite<double> WeightAndEasingResults)
 	{
 		const int32 Num = Allocation->Num();
 		for (int32 Index = 0; Index < Num; ++Index)
 		{
 			const uint16 HierarchicalEasingChannel = HierarchicalEasingChannels[Index];
-
-			if (HierarchicalEasingChannel != INVALID_EASING_CHANNEL && ensure(HierarchicalResultsByChannelID.IsValidIndex(HierarchicalEasingChannel)))
+			if (ensure(ComputationData.IsValidIndex(HierarchicalEasingChannel)))
 			{
-				WeightAndEasingResults[Index] *= HierarchicalResultsByChannelID[HierarchicalEasingChannel];
+				WeightAndEasingResults[Index] *= ComputationData[HierarchicalEasingChannel].FinalResult;
 			}
 		}
 	}
@@ -306,10 +552,13 @@ UMovieSceneHierarchicalEasingInstantiatorSystem::UMovieSceneHierarchicalEasingIn
 
 	if (HasAnyFlags(RF_ClassDefaultObject))
 	{
+		// Has to come after bound object and root instantiation
 		DefineImplicitPrerequisite(UMovieSceneRootInstantiatorSystem::StaticClass(), GetClass());
+		DefineComponentConsumer(GetClass(), BuiltInComponents->BoundObject);
+
+		// Has to come before property instantiation so that initial values get created correctly
 		DefineImplicitPrerequisite(GetClass(), UMovieScenePropertyInstantiatorSystem::StaticClass());
 		DefineImplicitPrerequisite(GetClass(), UMovieSceneInterrogatedPropertyInstantiatorSystem::StaticClass());
-		DefineComponentConsumer(GetClass(), BuiltInComponents->BoundObject);
 	}
 }
 
@@ -317,7 +566,7 @@ bool UMovieSceneHierarchicalEasingInstantiatorSystem::IsRelevantImpl(UMovieScene
 {
 	using namespace UE::MovieScene;
 
-	return PersistentHandleToEasingChannel.Num() || InLinker->EntityManager.ContainsComponent(FBuiltInComponentTypes::Get()->HierarchicalEasingProvider);
+	return CachedInstanceIDToChannel.Num() || InLinker->EntityManager.ContainsComponent(FBuiltInComponentTypes::Get()->HierarchicalEasingProvider);
 }
 
 void UMovieSceneHierarchicalEasingInstantiatorSystem::OnLink()
@@ -329,7 +578,7 @@ void UMovieSceneHierarchicalEasingInstantiatorSystem::OnLink()
 
 void UMovieSceneHierarchicalEasingInstantiatorSystem::OnUnlink()
 {
-	PersistentHandleToEasingChannel.Empty();
+	CachedInstanceIDToChannel.Empty();
 }
 
 void UMovieSceneHierarchicalEasingInstantiatorSystem::OnRun(FSystemTaskPrerequisites& InPrerequisites, FSystemSubsequentTasks& Subsequents)
@@ -339,164 +588,127 @@ void UMovieSceneHierarchicalEasingInstantiatorSystem::OnRun(FSystemTaskPrerequis
 	const FBuiltInComponentTypes* const BuiltInComponents = FBuiltInComponentTypes::Get();
 	const FInstanceRegistry* const InstanceRegistry = Linker->GetInstanceRegistry();
 
-	// Step 1: Create easing channels for newly created easing providers, and add a channel ID component
-	{
-		FEntityComponentFilter Filter;
-		Filter.All({ BuiltInComponents->RootInstanceHandle, BuiltInComponents->HierarchicalEasingProvider, BuiltInComponents->Tags.NeedsLink });
-		Filter.Deny({ BuiltInComponents->Tags.ImportedEntity });
-		Linker->EntityManager.MutateAll(Filter, FAddEasingChannelToProviderMutation(this));
-	}
-
-	FAddEasingChannelToConsumerMutation AddEasingChannelMutation(this);
+	// Update if we have any NeedsLink or NeedsUnlink providers (that contribute to hierachical weights)
+	FEntityComponentFilter ChangedProviderFilter;
+	ChangedProviderFilter.All({ BuiltInComponents->HierarchicalEasingProvider });
+	ChangedProviderFilter.Any({ BuiltInComponents->Tags.NeedsLink, BuiltInComponents->Tags.NeedsUnlink });
 
-	// Step 2: Add easing ID components to any new entities that exist within a blended sequence
-	{
-		FEntityComponentFilter Filter;
-		Filter.All({ BuiltInComponents->InstanceHandle, BuiltInComponents->Tags.NeedsLink });
-		Filter.None({ BuiltInComponents->HierarchicalEasingChannel, BuiltInComponents->Tags.NeedsUnlink, BuiltInComponents->Tags.ImportedEntity });
+	// Or any NeedsLink sub instances (that might need to receieve a parent weight)
+	FEntityComponentFilter NewSubInstanceFilter;
+	NewSubInstanceFilter.All({ BuiltInComponents->Tags.SubInstance, BuiltInComponents->Tags.NeedsLink });
 
-		Linker->EntityManager.MutateConditional(Filter, AddEasingChannelMutation);
-	}
+	// Or any NeedsUnlink sub instances that have an easing channel already assigned
+	FEntityComponentFilter ExpiringSubInstanceFilter;
+	ExpiringSubInstanceFilter.All({ BuiltInComponents->Tags.SubInstance, BuiltInComponents->HierarchicalEasingChannel, BuiltInComponents->Tags.NeedsUnlink });
 
-	// Step 3: Add easing ID components to any pre-existing entities that exist within a sequence that just had a channel allocated
-	for (const FHierarchicalKey& Key : NewEasingChannelKeys)
+	// Check if we have any new or expiring easing providers
+	if (Linker->EntityManager.Contains(ChangedProviderFilter) || Linker->EntityManager.Contains(NewSubInstanceFilter) || Linker->EntityManager.Contains(ExpiringSubInstanceFilter))
 	{
-		const FSequenceInstance& RootInstance = InstanceRegistry->GetInstance(Key.RootInstanceHandle);
+		FRootSequenceDataMap SequenceDataMap;
 
-		const FSequenceInstance* InstanceToCheck = nullptr;
-		if (Key.SequenceID == MovieSceneSequenceID::Root)
-		{
-			InstanceToCheck = &RootInstance;
-		}
-		else
+		// Called for all easing providers to add them to our accumulation map
+		auto VisitEasingProvider = [InstanceRegistry, &SequenceDataMap](FMovieSceneEntityID EntityID, FRootInstanceHandle RootInstance, FMovieSceneSequenceID ProviderID)
 		{
-			FInstanceHandle SubInstanceHandle = RootInstance.FindSubInstance(Key.SequenceID);
-			if (SubInstanceHandle.IsValid())
+			if (InstanceRegistry->IsHandleValid(RootInstance))
 			{
-				InstanceToCheck = &InstanceRegistry->GetInstance(SubInstanceHandle);
+				// Make the entry for our root instance
+				FRootSequenceData* RootSequenceData = SequenceDataMap.Find(RootInstance);
+				if (!RootSequenceData)
+				{
+					RootSequenceData = &SequenceDataMap.Emplace(RootInstance, FRootSequenceData(RootInstance, InstanceRegistry));
+				}
+
+				// Add the sequence ID to the channel map if it doesn't already exist
+				if (RootSequenceData->IndexOf(ProviderID) == INDEX_NONE)
+				{
+					RootSequenceData->Add(ProviderID, EHierarchicalSequenceChannelFlags::None);
+				}
 			}
-		}
+		};
 
-		if (InstanceToCheck)
+		// Called for all sub instances to add them to our accumulation map if they reside within a parent that is blended
+		auto VisitSubInstance = [InstanceRegistry, &SequenceDataMap](FMovieSceneEntityID EntityID, FRootInstanceHandle RootInstance, FMovieSceneSequenceID SequenceID)
 		{
-			// Find any entities that have already been linked and add the easing channel to them
-			FEntityComponentFilter Filter;
-			Filter.None({ BuiltInComponents->Tags.NeedsLink, BuiltInComponents->Tags.NeedsUnlink, BuiltInComponents->HierarchicalEasingChannel, BuiltInComponents->Tags.ImportedEntity });
-
-			InstanceToCheck->Ledger.MutateAll(Linker, Filter, AddEasingChannelMutation);
-		}
-	}
-
-	NewEasingChannelKeys.Empty();
-
-
-	RemoveUnlinkedHierarchicalEasingChannels(InPrerequisites, Subsequents);
-}
+			if (SequenceID != MovieSceneSequenceID::Root && InstanceRegistry->IsHandleValid(RootInstance))
+			{
+				if (FRootSequenceData* RootSequenceData = SequenceDataMap.Find(RootInstance))
+				{
+					RootSequenceData->AddImplicitChild(SequenceID);
+				}
+			}
+		};
 
-void UMovieSceneHierarchicalEasingInstantiatorSystem::RemoveUnlinkedHierarchicalEasingChannels(FSystemTaskPrerequisites& InPrerequisites, FSystemSubsequentTasks& Subsequents)
-{
-	using namespace UE::MovieScene;
+		// Step 1 : Iterate all easing providers and add their SequenceIDs to the map
+		// 
+		FEntityTaskBuilder()
+		.ReadEntityIDs()
+		.Read(BuiltInComponents->RootInstanceHandle)
+		.Read(BuiltInComponents->HierarchicalEasingProvider)
+		.Iterate_PerEntity(&Linker->EntityManager, VisitEasingProvider);
 
-	const FBuiltInComponentTypes* const BuiltInComponents = FBuiltInComponentTypes::Get();
-	const FInstanceRegistry* const InstanceRegistry = Linker->GetInstanceRegistry();
+		// Step 2 : Iterate all sub instances and see if they exist within one of the previously collected easing providers
+		// 
+		FEntityTaskBuilder()
+		.ReadEntityIDs()
+		.Read(BuiltInComponents->RootInstanceHandle)
+		.Read(BuiltInComponents->SequenceID)
+		.FilterAll({ BuiltInComponents->Tags.SubInstance })
+		.Iterate_PerEntity(&Linker->EntityManager, VisitSubInstance);
 
-	// Step 3: Visit removed hierarchical easing providers, so we can free up our channels.
-	//         @todo: There is a risk here that the easing provider is removed before the instance has been finished.
-	//                leaving some entities hanging around that still have this channel ID assigned. This would casue an ensure in the evaluator system.
-	//                This currently can't happen though because sub-sequence easing always lasts the duration of the sub-section, and dynamic weights
-	//                are never removed once they are created (until the root sequence ends)
-	auto VisitRemovedEasingProviders = [this, InstanceRegistry](const FEntityAllocation* Allocation, const FRootInstanceHandle* RootInstanceHandles, const FMovieSceneSequenceID* SubSequenceIDs)
-	{
-		for (int32 Index = 0; Index < Allocation->Num(); ++Index)
+		// Step 3: Produce parent-first channel IDs for anything we collected, then pass on the buffer to the evaluator system
+		//
 		{
-			FHierarchicalKey Key { RootInstanceHandles[Index], SubSequenceIDs[Index] };
+			TArray<FHierarchicalEasingChannelData> AllChannels;
+			CachedInstanceIDToChannel.Empty();
 
-			if (FHierarchicalInstanceData* InstanceData = this->PersistentHandleToEasingChannel.Find(Key))
+			for (TPair<FRootInstanceHandle, FRootSequenceData>& Pair : SequenceDataMap)
 			{
-				if (--InstanceData->RefCount == 0)
-				{
-					this->EvaluatorSystem->ReleaseEasingChannel(InstanceData->EasingChannelID);
-					this->PersistentHandleToEasingChannel.Remove(Key);
-				}
+				Pair.Value.ProduceChannels(CachedInstanceIDToChannel, AllChannels);
 			}
+			EvaluatorSystem->SetComputationBuffer(MoveTemp(AllChannels));
 		}
-	};
-
-	FEntityTaskBuilder()
-	.Read(BuiltInComponents->RootInstanceHandle)
-	.Read(BuiltInComponents->HierarchicalEasingProvider)
-	.FilterAll({ BuiltInComponents->Tags.NeedsUnlink })
-	.FilterNone({ BuiltInComponents->ParentEntity })
-	.Iterate_PerAllocation(&Linker->EntityManager, VisitRemovedEasingProviders);
-}
-
-UMovieSceneHierarchicalEasingInstantiatorSystem::FHierarchicalInstanceData UMovieSceneHierarchicalEasingInstantiatorSystem::AllocateEasingChannel(const UE::MovieScene::FInstanceRegistry* InstanceRegistry, UE::MovieScene::FRootInstanceHandle RootInstanceHandle, FMovieSceneSequenceID SequenceID)
-{
-	using namespace UE::MovieScene;
-
-	const FMovieSceneSequenceHierarchy* RootHierarchy = (SequenceID != MovieSceneSequenceID::Root)
-		? InstanceRegistry->GetInstance(RootInstanceHandle).GetPlayer()->GetEvaluationTemplate().GetHierarchy()
-		: nullptr;
 
-	return AllocateEasingChannelImpl(RootInstanceHandle, SequenceID, RootHierarchy);
-}
+		// Step 4: Assign channel IDs to easing providers, potentially adding or removing invalid ones
+		//
+		{
+			FEntityComponentFilter ProviderFilter;
+			ProviderFilter.All({ BuiltInComponents->RootInstanceHandle, BuiltInComponents->HierarchicalEasingProvider });
+			ProviderFilter.None({ BuiltInComponents->Tags.NeedsUnlink });
 
-UMovieSceneHierarchicalEasingInstantiatorSystem::FHierarchicalInstanceData UMovieSceneHierarchicalEasingInstantiatorSystem::AllocateEasingChannelImpl(UE::MovieScene::FRootInstanceHandle RootInstanceHandle, FMovieSceneSequenceID SequenceID, const FMovieSceneSequenceHierarchy* Hierarchy)
-{
-	using namespace UE::MovieScene;
+			FProviderEasingChannelMutation ProviderMutation(Linker, &SequenceDataMap);
+			Linker->EntityManager.MutateConditional(ProviderFilter, ProviderMutation);
 
-	FHierarchicalKey Key{ RootInstanceHandle, SequenceID };
+			ProviderMutation.RemoveStaleComponents(Linker);
+		}
 
-	FHierarchicalInstanceData* ExistingInstanceData = PersistentHandleToEasingChannel.Find(Key);
-	if (ExistingInstanceData != nullptr)
-	{
-		++ExistingInstanceData->RefCount;
-		return *ExistingInstanceData;
-	}
+		// Step 5: Assign easing channels for anything else that is left and remove stale ones
+		//
+		{
+			FEntityComponentFilter NewFilter;
+			NewFilter.All({ BuiltInComponents->InstanceHandle });
+			NewFilter.None({ BuiltInComponents->Tags.NeedsUnlink, BuiltInComponents->HierarchicalEasingProvider });
 
-	FHierarchicalInstanceData ParentInstanceData;
+			FConsumerEasingChannelMutation ConsumerMutation(Linker, &CachedInstanceIDToChannel);
+			Linker->EntityManager.MutateConditional(NewFilter, ConsumerMutation);
 
-	if (SequenceID != MovieSceneSequenceID::Root)
-	{
-		const FMovieSceneSequenceHierarchyNode* Node = Hierarchy->FindNode(SequenceID);
-		check(Node);
-		ParentInstanceData = AllocateEasingChannelImpl(RootInstanceHandle, Node->ParentID, Hierarchy);
+			ConsumerMutation.RemoveStaleComponents(Linker);
+		}
 	}
-
-	FHierarchicalInstanceData NewInstanceData;
-	NewInstanceData.RefCount = 1;
-	NewInstanceData.HierarchicalDepth = ParentInstanceData.HierarchicalDepth + 1;
-	NewInstanceData.EasingChannelID = EvaluatorSystem->AllocateEasingChannel(ParentInstanceData.EasingChannelID, NewInstanceData.HierarchicalDepth);
-
-	PersistentHandleToEasingChannel.Add(Key, NewInstanceData);
-	NewEasingChannelKeys.Add(Key);
-
-	return NewInstanceData;
-}
-
-uint16 UMovieSceneHierarchicalEasingInstantiatorSystem::LocateEasingChannel(UE::MovieScene::FRootInstanceHandle RootInstanceHandle, FMovieSceneSequenceID SequenceID) const
-{
-	if (const FHierarchicalInstanceData* ExistingInstanceData = PersistentHandleToEasingChannel.Find(FHierarchicalKey{ RootInstanceHandle, SequenceID }))
+	else
 	{
-		return ExistingInstanceData->EasingChannelID;
-	}
-	return uint16(-1);
-}
-
-void UMovieSceneHierarchicalEasingInstantiatorSystem::ReleaseEasingChannel(UE::MovieScene::FRootInstanceHandle RootInstanceHandle, FMovieSceneSequenceID SequenceID)
-{
-	FHierarchicalKey Key{ RootInstanceHandle, SequenceID };
+		// Just add easing channels to NeedsLink entities
+		FEntityComponentFilter Filter;
+		Filter.All({ BuiltInComponents->InstanceHandle, BuiltInComponents->Tags.NeedsLink });
+		Filter.None({ BuiltInComponents->HierarchicalEasingProvider });
 
-	const FHierarchicalInstanceData* ExistingInstanceData = PersistentHandleToEasingChannel.Find(Key);
-	if (ExistingInstanceData != nullptr)
-	{
-		EvaluatorSystem->ReleaseEasingChannel(ExistingInstanceData->EasingChannelID);
+		FConsumerEasingChannelMutation ConsumerMutation(Linker, &CachedInstanceIDToChannel);
+		Linker->EntityManager.MutateConditional(Filter, ConsumerMutation);
 
-		PersistentHandleToEasingChannel.Remove(Key);
+		ensure(!ConsumerMutation.HasStaleComponents());
+		ConsumerMutation.RemoveStaleComponents(Linker);
 	}
 }
 
-
 UWeightAndEasingEvaluatorSystem::UWeightAndEasingEvaluatorSystem(const FObjectInitializer& ObjInit)
 	: Super(ObjInit)
 {
@@ -504,8 +716,6 @@ UWeightAndEasingEvaluatorSystem::UWeightAndEasingEvaluatorSystem(const FObjectIn
 
 	SystemCategories = UE::MovieScene::EEntitySystemCategory::ChannelEvaluators;
 
-	bResultsNeedResort = false;
-
 	if (HasAnyFlags(RF_ClassDefaultObject))
 	{
 		const FBuiltInComponentTypes* BuiltInComponents = FBuiltInComponentTypes::Get();
@@ -528,57 +738,17 @@ bool UWeightAndEasingEvaluatorSystem::IsRelevantImpl(UMovieSceneEntitySystemLink
 
 void UWeightAndEasingEvaluatorSystem::OnLink()
 {
-	// Sometimes there can be easing channels left open
-	// For instance if a sub sequence caused its parent to have a
-	// hierarchical channel allocated even though it wasn't explicitly weighted.
-	// Since there are no entities left that need the channels, we can just remove them now
-	EasingChannelToIndex.Empty();
 	PreAllocatedComputationData.Empty();
 }
 
 void UWeightAndEasingEvaluatorSystem::OnUnlink()
 {
-	// Sometimes there can be easing channels left open
-	// For instance if a sub sequence caused its parent to have a
-	// hierarchical channel allocated even though it wasn't explicitly weighted.
-	// Since there are no entities left that need the channels, we can just remove them now
-	EasingChannelToIndex.Empty();
 	PreAllocatedComputationData.Empty();
 }
 
-uint16 UWeightAndEasingEvaluatorSystem::AllocateEasingChannel(const uint16 ParentEasingChannel, const uint16 HierarchicaDepth)
-{
-	using namespace UE::MovieScene;
-
-	bResultsNeedResort = true;
-
-	// Allocate the result for this channel
-	const int32 ResultIndex = PreAllocatedComputationData.Num();
-	const int32 EasingChannelID = EasingChannelToIndex.Add(ResultIndex);
-
-	check(EasingChannelID >= 0 && EasingChannelID < uint16(-1));
-
-	FHierarchicalEasingChannelData Channel;
-	Channel.HierarchicalDepth = HierarchicaDepth;
-	Channel.ChannelID = EasingChannelID;
-	if (ParentEasingChannel != INVALID_EASING_CHANNEL)
-	{
-		Channel.ParentEasingIndex = EasingChannelToIndex[ParentEasingChannel];
-	}
-
-	PreAllocatedComputationData.Add(Channel);
-	return static_cast<uint16>(EasingChannelID);
-}
-
-void UWeightAndEasingEvaluatorSystem::ReleaseEasingChannel(uint16 EasingChannelID)
+void UWeightAndEasingEvaluatorSystem::SetComputationBuffer(TArray<UE::MovieScene::FHierarchicalEasingChannelData>&& InPreAllocatedComputationData)
 {
-	const int32 ComputationDataIndex = EasingChannelToIndex[EasingChannelID];
-
-	// Mark the channel ID as invalid, and give it the largest hierarchical depth so it gets sorted to the end
-	PreAllocatedComputationData[ComputationDataIndex].HierarchicalDepth = uint16(-1);
-	EasingChannelToIndex.RemoveAt(EasingChannelID);
-
-	bResultsNeedResort = true;
+	PreAllocatedComputationData = MoveTemp(InPreAllocatedComputationData);
 }
 
 void UWeightAndEasingEvaluatorSystem::OnRun(FSystemTaskPrerequisites& InPrerequisites, FSystemSubsequentTasks& Subsequents)
@@ -615,8 +785,6 @@ void UWeightAndEasingEvaluatorSystem::OnRun(FSystemTaskPrerequisites& InPrerequi
 	// If we have hierarchical easing, we initialize all the weights to their hierarchical defaults
 	if (PreAllocatedComputationData.Num() > 0)
 	{
-		ResortComputationBuffer();
-
 		FSystemTaskPrerequisites HarvestPrereqs {InPrerequisites};
 		HarvestPrereqs.AddComponentTask(Components->WeightAndEasingResult, EvaluateEasing);
 		HarvestPrereqs.AddComponentTask(Components->WeightAndEasingResult, AccumulateManualWeights);
@@ -628,7 +796,7 @@ void UWeightAndEasingEvaluatorSystem::OnRun(FSystemTaskPrerequisites& InPrerequi
 		.Read(Components->HierarchicalEasingChannel)
 		.FilterAll({ Components->HierarchicalEasingProvider })  // Only harvest results from entities that are providing results
 		.SetStat(GET_STATID(MovieSceneEval_HarvestEasingTask))
-		.Dispatch_PerEntity<FHarvestHierarchicalEasings>(&Linker->EntityManager, HarvestPrereqs, nullptr, &EasingChannelToIndex, &PreAllocatedComputationData);
+		.Dispatch_PerEntity<FHarvestHierarchicalEasings>(&Linker->EntityManager, HarvestPrereqs, nullptr, &PreAllocatedComputationData);
 
 		FSystemTaskPrerequisites PropagatePrereqs {InPrerequisites};
 		PropagatePrereqs.AddRootTask(HarvestTask);
@@ -641,47 +809,6 @@ void UWeightAndEasingEvaluatorSystem::OnRun(FSystemTaskPrerequisites& InPrerequi
 		.FilterNone({ Components->HierarchicalEasingProvider }) // Do not propagate hierarchical weights onto providers!
 		.SetStat(GET_STATID(MovieSceneEval_EvaluateEasingTask))
 		.Dispatch_PerAllocation<FPropagateHierarchicalEasings>(&Linker->EntityManager, PropagatePrereqs, &Subsequents,
-			PreAllocatedComputationData, EasingChannelToIndex.GetMaxIndex());
+			PreAllocatedComputationData);
 	}
-}
-
-void UWeightAndEasingEvaluatorSystem::ResortComputationBuffer()
-{
-	using namespace UE::MovieScene;
-
-	if (!bResultsNeedResort)
-	{
-		return;
-	}
-
-	bResultsNeedResort = false;
-
-	TSparseArray<int32> OldToNewIndex;
-
-	// Resort the results array by depth if it has been modified
-	Algo::SortBy(PreAllocatedComputationData, &FHierarchicalEasingChannelData::HierarchicalDepth);
-
-	for (int32 Index = 0; Index < PreAllocatedComputationData.Num(); ++Index)
-	{
-		FHierarchicalEasingChannelData& ComputationData = PreAllocatedComputationData[Index];
-
-		// As soon as we find an invalid hierarchical depth, everything proceeding this index is garbage
-		if (ComputationData.HierarchicalDepth == uint16(-1))
-		{
-			PreAllocatedComputationData.RemoveAt(Index, PreAllocatedComputationData.Num() - Index);
-			break;
-		}
-
-		// Reassign the channel ID to index mapping
-		int32& ChannelIndex = EasingChannelToIndex[ComputationData.ChannelID];
-		OldToNewIndex.Insert(ChannelIndex, Index);
-		ChannelIndex = Index;
-		
-		if (ComputationData.ParentEasingIndex != uint16(-1))
-		{
-			// Parent Index must have been added to the OldToNewIndex map by now because the results are always sorted
-			ComputationData.ParentEasingIndex = OldToNewIndex[ComputationData.ParentEasingIndex];
-		}
-	}
-}
-
+}
\ No newline at end of file
diff --git a/Engine/Source/Runtime/MovieSceneTracks/Public/Systems/WeightAndEasingEvaluatorSystem.h b/Engine/Source/Runtime/MovieSceneTracks/Public/Systems/WeightAndEasingEvaluatorSystem.h
--- a/Engine/Source/Runtime/MovieSceneTracks/Public/Systems/WeightAndEasingEvaluatorSystem.h
+++ b/Engine/Source/Runtime/MovieSceneTracks/Public/Systems/WeightAndEasingEvaluatorSystem.h
@@ -29,20 +29,15 @@ struct FMovieSceneSubSequenceData;
 namespace UE::MovieScene
 {
 
+struct FEasingMutationBase;
 struct FAddEasingChannelToProviderMutation;
+struct FAddEasingChannelToConsumerMutation;
 
 /** Computation data used for accumulating hierarchical weights for sub sequences */
 struct FHierarchicalEasingChannelData
 {
 	/** Our parent's computation data within PreAllocatedComputationData. Must only access if bResultsNeedResort is false. */
-	uint16 ParentEasingIndex = uint16(-1);
-	/** The accumulated hierarchical depth of this sequence within its root. */
-	uint16 HierarchicalDepth = 0;
-	/**
-	 * The easing channel ID for this data. This reprensents the index within EasingChannelToIndex that uniquely identifies our channel.
-	 * Any entity within this sequence that is subject to easing will contain a HierarchicalEasingChannel component with this ID
-	 */
-	uint16 ChannelID = uint16(-1);
+	uint16 ParentChannel = uint16(-1);
 	/** The final result of this easing channel, accumulated with all parents */
 	double FinalResult = 1.0;
 };
@@ -60,7 +55,8 @@ public:
 	struct FHierarchicalInstanceData
 	{
 		int32 HierarchicalDepth = -1;
-		uint16 RefCount = 0;
+		uint16 StrongRefCount = 0;
+		uint16 WeakRefCount = 0;
 		uint16 EasingChannelID = uint16(-1);
 	};
 
@@ -68,20 +64,6 @@ public:
 
 	UMovieSceneHierarchicalEasingInstantiatorSystem(const FObjectInitializer& ObjInit);
 
-	/** 
-	 * Locate an already allocated easing channel for the specified sequence instance and sub sequence ID.
-	 *
-	 * @param RootInstanceHandle The instance handle of the root sequence
-	 * @param SequenceID         The Sequence ID of the sub sequence
-	 * @return The sequence's easing channel or 0xFFFF if one does not exist
-	 */
-	uint16 LocateEasingChannel(UE::MovieScene::FRootInstanceHandle RootInstanceHandle, FMovieSceneSequenceID SequenceID) const;
-
-	/**
-	 * Remove any hierarchical easing channels for channels that are no longer needed because of unlinked easing providers
-	 */
-	void RemoveUnlinkedHierarchicalEasingChannels(FSystemTaskPrerequisites& InPrerequisites, FSystemSubsequentTasks& Subsequents);
-
 private:
 
 	virtual bool IsRelevantImpl(UMovieSceneEntitySystemLinker* InLinker) const;
@@ -91,58 +73,8 @@ private:
 
 private:
 
-	/** 
-	 * Allocate a new easing channel for the specified sequence instance and sub sequence ID. Will return an existing channel if it has already been allocated.
-	 *
-	 * @param InstanceRegistry   Pointer to the instance registry for retrievable of sequence instances (see UMovieSceneEntitySystemLinbker::GetInstanceRegistry)
-	 * @param RootInstanceHandle The instance handle of the root sequence
-	 * @param SequenceID         The Sequence ID of the sub sequence to make a channel for, or MovieSceneSequenceID::Root for the root sequence
-	 * @return Structure containing both the easing channel ID and this sequence's hierarchical depth
-	 */
-	FHierarchicalInstanceData AllocateEasingChannel(const UE::MovieScene::FInstanceRegistry* InstanceRegistry, UE::MovieScene::FRootInstanceHandle RootInstanceHandle, FMovieSceneSequenceID SequenceID);
-
-	/** 
-	 * Release an easing channel for the specified sequence instance and sub sequence ID if one exists
-	 *
-	 * @param RootInstanceHandle The instance handle of the root sequence
-	 * @param SequenceID         The Sequence ID of the sub sequence
-	 */
-	void ReleaseEasingChannel(UE::MovieScene::FRootInstanceHandle RootInstanceHandle, FMovieSceneSequenceID SequenceID);
-
-	/** 
-	 * Allocate a new easing channel for the specified sequence instance and sub sequence ID, and all its parents. Will return an existing channel if it has already been allocated.
-	 *
-	 * @param RootInstanceHandle The instance handle of the root sequence
-	 * @param SequenceID         The Sequence ID of the sub sequence to make a channel for, or MovieSceneSequenceID::Root for the root sequence
-	 * @param Hierarchy          The sequence hierarchy (or nullptr if this is a root sequence and one is not available)
-	 * @return Structure containing both the easing channel ID and this sequence's hierarchical depth
-	 */
-	FHierarchicalInstanceData AllocateEasingChannelImpl(UE::MovieScene::FRootInstanceHandle RootInstanceHandle, FMovieSceneSequenceID SequenceID, const FMovieSceneSequenceHierarchy* Hierarchy);
-
-private:
-
-	friend UE::MovieScene::FAddEasingChannelToProviderMutation;
-
-	/** Hierarchy key for associating channel IDs to pairs of instance handles and sequence IDs */
-	struct FHierarchicalKey
-	{
-		UE::MovieScene::FRootInstanceHandle RootInstanceHandle;
-		FMovieSceneSequenceID SequenceID;
-
-		friend uint32 GetTypeHash(const FHierarchicalKey& In)
-		{
-			return HashCombine(GetTypeHash(In.RootInstanceHandle), GetTypeHash(In.SequenceID));
-		}
-		friend bool operator==(const FHierarchicalKey& A, const FHierarchicalKey& B)
-		{
-			return A.RootInstanceHandle == B.RootInstanceHandle && A.SequenceID == B.SequenceID;
-		}
-	};
-
 	/** Map between a sub-sequence handle and the easing channel affecting it */
-	TMap<FHierarchicalKey, FHierarchicalInstanceData> PersistentHandleToEasingChannel;
-	/** Set of newly created easing channels this frame. This is used to add easing data to existing entities. */
-	TSet<FHierarchicalKey> NewEasingChannelKeys;
+	TSparseArray<uint16> CachedInstanceIDToChannel;
 
 	UPROPERTY()
 	TObjectPtr<UWeightAndEasingEvaluatorSystem> EvaluatorSystem;
@@ -160,8 +92,10 @@ public:
 
 	UWeightAndEasingEvaluatorSystem(const FObjectInitializer& ObjInit);
 
-	uint16 AllocateEasingChannel(const uint16 ParentEasingChannel, const uint16 HierarchicaDepth);
-	void ReleaseEasingChannel(uint16 EasingChannelID);
+	/**
+	 * Called from the instantiator system once channel IDs have been calculated
+	 */
+	void SetComputationBuffer(TArray<UE::MovieScene::FHierarchicalEasingChannelData>&& PreAllocatedComputationData);
 
 private:
 
@@ -170,17 +104,9 @@ private:
 	virtual void OnLink() override final;
 	virtual void OnUnlink() override final;
 
-	void ResortComputationBuffer();
-
 private:
 
-	/** Stable sparse array of indices into PreAllocatedEasingResults for each easing channel */
-	TSparseArray<int32> EasingChannelToIndex;
-
 	/** Unstable array of preallocated storage for computing easing results sorted by hierarchical depth. */
 	TArray<UE::MovieScene::FHierarchicalEasingChannelData> PreAllocatedComputationData;
-
-	/** True if the preallocated easing results need resorting. */
-	bool bResultsNeedResort;
 };
 
