Sequencer: Constraints:   More spawnable constraint fixes, may need to add the constraint during eval and when getting labels may need to use the cached spawnable constraint since it may not exist yet.

diff --git a/Engine/Plugins/Animation/ControlRig/Source/ControlRigEditor/Private/EditMode/ControlRigControlsProxy.h b/Engine/Plugins/Animation/ControlRig/Source/ControlRigEditor/Private/EditMode/ControlRigControlsProxy.h
--- a/Engine/Plugins/Animation/ControlRig/Source/ControlRigEditor/Private/EditMode/ControlRigControlsProxy.h
+++ b/Engine/Plugins/Animation/ControlRig/Source/ControlRigEditor/Private/EditMode/ControlRigControlsProxy.h
@@ -73,7 +73,7 @@ class UControlRigTransformControlProxy : public UControlRigControlsProxy
 
 public:
 	
-	UPROPERTY(EditAnywhere, Interp, Category = "Control")
+	UPROPERTY(EditAnywhere, Interp, Category = "Control", meta = (SliderExponent = "1.0"))
 	FEulerTransform Transform; //FTransform doesn't work with multiple values for some reason, so for now using eulertransform which does work
 };
 
@@ -96,7 +96,7 @@ class UControlRigEulerTransformControlProxy : public UControlRigControlsProxy
 
 public:
 
-	UPROPERTY(EditAnywhere, Interp, Category = "Control")
+	UPROPERTY(EditAnywhere, Interp, Category = "Control", meta = (SliderExponent = "1.0"))
 	FEulerTransform Transform;
 };
 
@@ -119,7 +119,7 @@ class UControlRigTransformNoScaleControlProxy : public UControlRigControlsProxy
 
 public:
 
-	UPROPERTY(EditAnywhere, Interp, Category = "Control")
+	UPROPERTY(EditAnywhere, Interp, Category = "Control", meta = (SliderExponent = "1.0"))
 	FTransformNoScale Transform;
 };
 
@@ -141,7 +141,7 @@ class UControlRigFloatControlProxy : public UControlRigControlsProxy
 
 public:
 
-	UPROPERTY(EditAnywhere, Interp, Category = "Control")
+	UPROPERTY(EditAnywhere, Interp, Category = "Control", meta = (SliderExponent = "1.0"))
 	float Float;
 };
 
@@ -249,7 +249,7 @@ class UControlRigVectorControlProxy : public UControlRigControlsProxy
 
 public:
 
-	UPROPERTY(EditAnywhere, Interp, Category = "Control")
+	UPROPERTY(EditAnywhere, Interp, Category = "Control", meta = (SliderExponent = "1.0"))
 	FVector3f Vector;
 };
 
@@ -271,7 +271,7 @@ class UControlRigVector2DControlProxy : public UControlRigControlsProxy
 
 public:
 
-	UPROPERTY(EditAnywhere, Interp, Category = "Control")
+	UPROPERTY(EditAnywhere, Interp, Category = "Control", meta = (SliderExponent = "1.0"))
 	FVector2D Vector2D;
 };
 
diff --git a/Engine/Source/Runtime/Experimental/Animation/Constraints/Private/ConstraintsManager.cpp b/Engine/Source/Runtime/Experimental/Animation/Constraints/Private/ConstraintsManager.cpp
--- a/Engine/Source/Runtime/Experimental/Animation/Constraints/Private/ConstraintsManager.cpp
+++ b/Engine/Source/Runtime/Experimental/Animation/Constraints/Private/ConstraintsManager.cpp
@@ -432,6 +432,14 @@ bool FConstraintsManagerController::AddConstraint(UTickableConstraint* InConstra
 	}
 
 	Manager->Modify();
+	//it's possible this constraint was actually in another ConstraintActor::ConstraintManager so we need to move it over via Rename.
+	//and clear out it ticks function since that may have been registered
+	UConstraintsManager* Outer = InConstraint->GetTypedOuter<UConstraintsManager>();
+	if (Outer && Outer != Manager)
+	{
+		InConstraint->ConstraintTick.UnRegisterTickFunction();
+		InConstraint->Rename(nullptr, Manager, REN_ForceNoResetLoaders | REN_DoNotDirty | REN_DontCreateRedirectors | REN_NonTransactional);
+	}
 	Manager->Constraints.Emplace(InConstraint);
 
 	InConstraint->ConstraintTick.RegisterFunction(InConstraint->GetFunction());
diff --git a/Engine/Source/Runtime/MovieSceneTracks/Private/Sections/MovieScene3DTransformSection.cpp b/Engine/Source/Runtime/MovieSceneTracks/Private/Sections/MovieScene3DTransformSection.cpp
--- a/Engine/Source/Runtime/MovieSceneTracks/Private/Sections/MovieScene3DTransformSection.cpp
+++ b/Engine/Source/Runtime/MovieSceneTracks/Private/Sections/MovieScene3DTransformSection.cpp
@@ -715,6 +715,18 @@ EMovieSceneChannelProxyType UMovieScene3DTransformSection::CacheChannelProxy()
 
 	FMovieSceneChannelProxyData Channels;
 
+#if WITH_EDITOR
+
+	auto GetValidConstraint = [](FConstraintAndActiveChannel& ConstraintChannel)
+	{
+		if (ConstraintChannel.Constraint.IsValid())
+		{
+			return ConstraintChannel.Constraint.Get();
+		}
+		return (ConstraintChannel.ConstraintCopyToSpawn.Get());
+	};
+#endif
+		
 	//Constraints go on top
 	int32 NumConstraints = 0;
 	if(Constraints)
@@ -722,11 +734,12 @@ EMovieSceneChannelProxyType UMovieScene3DTransformSection::CacheChannelProxy()
 		for (FConstraintAndActiveChannel& ConstraintChannel : Constraints->ConstraintsChannels)
 		{
 #if WITH_EDITOR
-			if (ConstraintChannel.Constraint.IsValid())
+			TWeakObjectPtr<UTickableConstraint> WeakConstraint = GetValidConstraint(ConstraintChannel);
+			if (WeakConstraint.IsValid())
 			{
 				FText ConstraintGroup = NSLOCTEXT("MovieSceneTransformSection", "Constraints", "Constraints");
-				const FName& ConstraintName = ConstraintChannel.Constraint->GetFName();
-				ConstraintChannel.ActiveChannel.ExtraLabel = [WeakConstraint = MakeWeakObjectPtr(ConstraintChannel.Constraint.Get())]
+				const FName& ConstraintName = WeakConstraint->GetFName();
+				ConstraintChannel.ActiveChannel.ExtraLabel = [WeakConstraint]
 				{
 					if (WeakConstraint.IsValid())
 					{
@@ -739,7 +752,7 @@ EMovieSceneChannelProxyType UMovieScene3DTransformSection::CacheChannelProxy()
 					static const FString DummyStr;
 					return DummyStr;
 				};
-				const FText DisplayText = FText::FromString(ConstraintChannel.Constraint->GetTypeLabel());
+				const FText DisplayText = FText::FromString(WeakConstraint->GetTypeLabel());
 				FMovieSceneChannelMetaData MetaData(ConstraintName, DisplayText, ConstraintGroup, true /*bInEnabled*/);
 				MetaData.SortOrder = NumConstraints++;
 				MetaData.bCanCollapseToTrack = false;
diff --git a/Engine/Source/Runtime/MovieSceneTracks/Private/Systems/MovieSceneConstraintSystem.cpp b/Engine/Source/Runtime/MovieSceneTracks/Private/Systems/MovieSceneConstraintSystem.cpp
--- a/Engine/Source/Runtime/MovieSceneTracks/Private/Systems/MovieSceneConstraintSystem.cpp
+++ b/Engine/Source/Runtime/MovieSceneTracks/Private/Systems/MovieSceneConstraintSystem.cpp
@@ -129,7 +129,8 @@ void UMovieSceneConstraintSystem::OnRun(FSystemTaskPrerequisites& InPrerequisite
 				}
 				else // it's possible that we have it but it's not in the manager, due to manager not being saved with it (due to spawning or undo/redo).
 				{
-					if(Controller->GetConstraint(ConstraintChannel.ConstraintName) != Constraint)
+					const TArray< TObjectPtr<UTickableConstraint>>& ConstraintsArray = Controller->GetConstraintsArray();
+					if(ConstraintsArray.Find(Constraint) == INDEX_NONE)
 					{
 						Controller->AddConstraint(Constraint);
 					}
