Sequencer: don't need to read the cached interpolation component when not using cached interpolation

diff --git a/Engine/Source/Runtime/MovieSceneTracks/Private/Systems/DoubleChannelEvaluatorSystem.cpp b/Engine/Source/Runtime/MovieSceneTracks/Private/Systems/DoubleChannelEvaluatorSystem.cpp
--- a/Engine/Source/Runtime/MovieSceneTracks/Private/Systems/DoubleChannelEvaluatorSystem.cpp
+++ b/Engine/Source/Runtime/MovieSceneTracks/Private/Systems/DoubleChannelEvaluatorSystem.cpp
@@ -40,7 +40,7 @@ TArray<FDoubleChannelTypeAssociation, TInlineAllocator<4>> GDoubleChannelTypeAss
 /** Entity-component task that evaluates using the non-cached codepath for testing parity */
 struct FEvaluateDoubleChannels_Uncached
 {
-	void ForEachEntity(FSourceDoubleChannel DoubleChannel, FFrameTime FrameTime, Interpolation::FCachedInterpolation& Cache, double& OutResult)
+	void ForEachEntity(FSourceDoubleChannel DoubleChannel, FFrameTime FrameTime, double& OutResult)
 	{
 		if (!DoubleChannel.Source->Evaluate(FrameTime, OutResult))
 		{
@@ -203,7 +203,6 @@ void UDoubleChannelEvaluatorSystem::OnRun(FSystemTaskPrerequisites& InPrerequisi
 				FEntityTaskBuilder()
 				.Read(ChannelType.ChannelType)
 				.Read(BuiltInComponents->EvalTime)
-				.Write(ChannelType.CachedInterpolationType)
 				.Write(ChannelType.ResultType)
 				.FilterNone({ BuiltInComponents->Tags.Ignored })
 				.SetStat(GET_STATID(MovieSceneEval_EvaluateDoubleChannelTask))
diff --git a/Engine/Source/Runtime/MovieSceneTracks/Private/Systems/FloatChannelEvaluatorSystem.cpp b/Engine/Source/Runtime/MovieSceneTracks/Private/Systems/FloatChannelEvaluatorSystem.cpp
--- a/Engine/Source/Runtime/MovieSceneTracks/Private/Systems/FloatChannelEvaluatorSystem.cpp
+++ b/Engine/Source/Runtime/MovieSceneTracks/Private/Systems/FloatChannelEvaluatorSystem.cpp
@@ -40,7 +40,7 @@ TArray<FFloatChannelTypeAssociation, TInlineAllocator<4>> GFloatChannelTypeAssoc
 /** Entity-component task that evaluates using the non-cached codepath for testing parity */
 struct FEvaluateFloatChannels_Uncached
 {
-	void ForEachEntity(FSourceFloatChannel FloatChannel, FFrameTime FrameTime, Interpolation::FCachedInterpolation& Cache, double& OutResult)
+	void ForEachEntity(FSourceFloatChannel FloatChannel, FFrameTime FrameTime, double& OutResult)
 	{
 		float Result;
 		if (FloatChannel.Source->Evaluate(FrameTime, Result))
@@ -214,7 +214,6 @@ void UFloatChannelEvaluatorSystem::OnRun(FSystemTaskPrerequisites& InPrerequisit
 				FEntityTaskBuilder()
 				.Read(ChannelType.ChannelType)
 				.Read(BuiltInComponents->EvalTime)
-				.Write(ChannelType.CachedInterpolationType)
 				.Write(ChannelType.ResultType)
 				.FilterNone({ BuiltInComponents->Tags.Ignored })
 				.SetStat(GET_STATID(MovieSceneEval_EvaluateFloatChannelTask))
