Sequencer: Use cached filter result for resolving bound materials

 - This avoids having to refilter the ecs every frame when calling the OnPostSpawn callback

diff --git a/Engine/Source/Runtime/MovieScene/Private/EntitySystem/MovieSceneCachedEntityFilterResult.cpp b/Engine/Source/Runtime/MovieScene/Private/EntitySystem/MovieSceneCachedEntityFilterResult.cpp
--- a/Engine/Source/Runtime/MovieScene/Private/EntitySystem/MovieSceneCachedEntityFilterResult.cpp
+++ b/Engine/Source/Runtime/MovieScene/Private/EntitySystem/MovieSceneCachedEntityFilterResult.cpp
@@ -40,12 +40,12 @@ void FCachedEntityFilterResult_Match::Invalidate()
 
 
 
-TArrayView<const FEntityAllocation* const> FCachedEntityFilterResult_Allocations::GetMatchingAllocations(const FEntityManager& InEntityManager)
+TArrayView<FEntityAllocation* const> FCachedEntityFilterResult_Allocations::GetMatchingAllocations(const FEntityManager& InEntityManager)
 {
 	if (Cache.Update(InEntityManager) == ECachedEntityManagerState::Stale)
 	{
 		MatchedEntityAllocations.Empty();
-		for (const FEntityAllocation* Allocation : InEntityManager.Iterate(&Filter))
+		for (FEntityAllocation* Allocation : InEntityManager.Iterate(&Filter))
 		{
 			MatchedEntityAllocations.Add(Allocation);
 		}
diff --git a/Engine/Source/Runtime/MovieScene/Public/EntitySystem/MovieSceneCachedEntityFilterResult.h b/Engine/Source/Runtime/MovieScene/Public/EntitySystem/MovieSceneCachedEntityFilterResult.h
--- a/Engine/Source/Runtime/MovieScene/Public/EntitySystem/MovieSceneCachedEntityFilterResult.h
+++ b/Engine/Source/Runtime/MovieScene/Public/EntitySystem/MovieSceneCachedEntityFilterResult.h
@@ -88,7 +88,7 @@ struct MOVIESCENE_API FCachedEntityFilterResult_Allocations
 {
 	FEntityComponentFilter Filter;
 
-	TArrayView<const FEntityAllocation* const> GetMatchingAllocations(const FEntityManager& InEntityManager);
+	TArrayView<FEntityAllocation* const> GetMatchingAllocations(const FEntityManager& InEntityManager);
 
 	void Invalidate();
 
@@ -97,7 +97,7 @@ private:
 	FCachedEntityManagerState Cache;
 
 	/** Set bits indicate matching allocations */
-	TArray<const FEntityAllocation*> MatchedEntityAllocations;
+	TArray<FEntityAllocation*> MatchedEntityAllocations;
 };
 
 
diff --git a/Engine/Source/Runtime/MovieSceneTracks/Private/MovieSceneTracksModule.cpp b/Engine/Source/Runtime/MovieSceneTracks/Private/MovieSceneTracksModule.cpp
--- a/Engine/Source/Runtime/MovieSceneTracks/Private/MovieSceneTracksModule.cpp
+++ b/Engine/Source/Runtime/MovieSceneTracks/Private/MovieSceneTracksModule.cpp
@@ -8,12 +8,15 @@
 #include "Systems/MovieScene3DTransformPropertySystem.h"
 
 #include "EntitySystem/MovieSceneEntityManager.h"
+#include "Systems/MovieSceneMaterialSystem.h"
 #include "MovieSceneTracksComponentTypes.h"
 
 #if !IS_MONOLITHIC
 	UE::MovieScene::FEntityManager*& GEntityManagerForDebugging = UE::MovieScene::GEntityManagerForDebuggingVisualizers;
 #endif
 
+DEFINE_STAT(MovieSceneEval_ReinitializeBoundMaterials);
+
 /**
  * Implements the MovieSceneTracks module.
  */
diff --git a/Engine/Source/Runtime/MovieSceneTracks/Public/Systems/MovieSceneMaterialSystem.h b/Engine/Source/Runtime/MovieSceneTracks/Public/Systems/MovieSceneMaterialSystem.h
--- a/Engine/Source/Runtime/MovieSceneTracks/Public/Systems/MovieSceneMaterialSystem.h
+++ b/Engine/Source/Runtime/MovieSceneTracks/Public/Systems/MovieSceneMaterialSystem.h
@@ -12,11 +12,15 @@
 #include "EntitySystem/BuiltInComponentTypes.h"
 #include "EntitySystem/MovieSceneEntitySystemRunner.h"
 #include "EntitySystem/MovieSceneComponentTypeInfo.h"
+#include "EntitySystem/MovieSceneCachedEntityFilterResult.h"
 #include "Evaluation/PreAnimatedState/MovieScenePreAnimatedObjectStorage.h"
 #include "MovieSceneTracksComponentTypes.h"
 
 #include "MovieSceneMaterialSystem.generated.h"
 
+
+DECLARE_CYCLE_STAT_EXTERN(TEXT("Reinitialize Bound Materials"), MovieSceneEval_ReinitializeBoundMaterials, STATGROUP_MovieSceneECS, MOVIESCENETRACKS_API);
+
 USTRUCT()
 struct FMovieScenePreAnimatedMaterialParameters
 {
@@ -118,8 +122,10 @@ public:
 
 protected:
 
-	UE::MovieScene::FEntityComponentFilter MaterialSwitcherFilter;
-	UE::MovieScene::FEntityComponentFilter MaterialParameterFilter;
+	FEntityComponentFilter MaterialSwitcherFilter;
+	FEntityComponentFilter MaterialParameterFilter;
+
+	FCachedEntityFilterResult_Allocations ReinitializeBoundMaterials;
 };
 
 template<typename AccessorType, typename... RequiredComponents>
@@ -206,6 +212,14 @@ struct TReinitializeBoundMaterials
 		: Linker(InLinker)
 	{}
 
+	void ForEachAllocation(int32 Num, const FMovieSceneEntityID* EntityIDs, const RequiredComponents*... Inputs, UObject** Objects)
+	{
+		for (int32 Index = 0; Index < Num; ++Index)
+		{
+			ForEachEntity(EntityIDs[Index], Inputs[Index]..., Objects[Index]);
+		}
+	}
+
 	void ForEachEntity(FMovieSceneEntityID EntityID, typename TCallTraits<RequiredComponents>::ParamType... Inputs, UObject*& OutDynamicMaterial)
 	{
 		if (TInitializeBoundMaterials<AccessorType, RequiredComponents...>::InitializeBoundMaterial(Inputs..., OutDynamicMaterial))
@@ -281,6 +295,9 @@ void TMovieSceneMaterialSystem<AccessorType, RequiredComponents...>::OnLink(UMov
 	MaterialParameterFilter.Any({ TracksComponents->ScalarParameterName, TracksComponents->ColorParameterName, TracksComponents->VectorParameterName });
 
 	Linker->Events.PostSpawnEvent.AddRaw(this, &TMovieSceneMaterialSystem<AccessorType, RequiredComponents...>::OnPostSpawn, InRequiredComponents...);
+
+	ReinitializeBoundMaterials.Filter.All({ TracksComponents->BoundMaterial, InRequiredComponents... });
+	ReinitializeBoundMaterials.Filter.None({ BuiltInComponents->Tags.NeedsUnlink });
 }
 
 template<typename AccessorType, typename... RequiredComponents>
@@ -294,19 +311,25 @@ void TMovieSceneMaterialSystem<AccessorType, RequiredComponents...>::OnPostSpawn
 {
 	using namespace UE::MovieScene;
 
+	SCOPE_CYCLE_COUNTER(MovieSceneEval_ReinitializeBoundMaterials)
+
 	FBuiltInComponentTypes*          BuiltInComponents = FBuiltInComponentTypes::Get();
 	FMovieSceneTracksComponentTypes* TracksComponents  = FMovieSceneTracksComponentTypes::Get();
 
+	FEntityAllocationWriteContext WriteContext(InLinker->EntityManager);
 	TReinitializeBoundMaterials<AccessorType, RequiredComponents...> ReinitializeBoundMaterialsTask(InLinker);
 
 	// Reinitialize bound dynamic materials, adding NeedsLink during PostTask to any that changed
 	// This will cause the instantiation phase to be re-run for these entities (and any other new or expired ones)
-	FEntityTaskBuilder()
-	.ReadEntityIDs()
-	.ReadAllOf(InRequiredComponents...)
-	.Write(TracksComponents->BoundMaterial)
-	.FilterNone({ BuiltInComponents->Tags.NeedsUnlink })
-	.RunInline_PerEntity(&InLinker->EntityManager, ReinitializeBoundMaterialsTask);
+	for (FEntityAllocation* Allocation : ReinitializeBoundMaterials.GetMatchingAllocations(InLinker->EntityManager))
+	{
+		FEntityAllocationMutexGuard LockGuard(Allocation, EComponentHeaderLockMode::LockFree);
+
+		const int32 Num = Allocation->Num();
+		ReinitializeBoundMaterialsTask.ForEachAllocation(Num, Allocation->GetRawEntityIDs(), Allocation->ReadComponents(InRequiredComponents)..., Allocation->WriteComponents(TracksComponents->BoundMaterial, WriteContext));
+	}
+
+	ReinitializeBoundMaterialsTask.PostTask();
 }
 
 template<typename AccessorType, typename... RequiredComponents>
