Sequencer: convert fade track to ECS evaluation

diff --git a/Engine/Source/Runtime/MovieSceneTracks/Private/Evaluation/MovieSceneFadeTemplate.cpp b/Engine/Source/Runtime/MovieSceneTracks/Private/Evaluation/MovieSceneFadeTemplate.cpp
deleted file mode 100644
diff --git a/Engine/Source/Runtime/MovieSceneTracks/Private/Evaluation/MovieSceneFadeTemplate.h b/Engine/Source/Runtime/MovieSceneTracks/Private/Evaluation/MovieSceneFadeTemplate.h
deleted file mode 100644
diff --git a/Engine/Source/Runtime/MovieSceneTracks/Private/MovieSceneTracksComponentTypes.cpp b/Engine/Source/Runtime/MovieSceneTracks/Private/MovieSceneTracksComponentTypes.cpp
--- a/Engine/Source/Runtime/MovieSceneTracks/Private/MovieSceneTracksComponentTypes.cpp
+++ b/Engine/Source/Runtime/MovieSceneTracks/Private/MovieSceneTracksComponentTypes.cpp
@@ -598,6 +598,8 @@ FMovieSceneTracksComponentTypes::FMovieSceneTracksComponentTypes()
 	ComponentRegistry->NewComponentType(&ColorParameterName,     TEXT("Color Parameter Name"), EComponentTypeFlags::CopyToChildren | EComponentTypeFlags::CopyToOutput);
 	ComponentRegistry->NewComponentType(&TransformParameterName, TEXT("Transform Parameter Name"), EComponentTypeFlags::CopyToChildren | EComponentTypeFlags::CopyToOutput);
 
+	ComponentRegistry->NewComponentType(&Fade,                   TEXT("Fade"), EComponentTypeFlags::CopyToChildren);
+
 	Tags.BoundMaterialChanged = ComponentRegistry->NewTag(TEXT("Bound Material Changed"));
 	FBuiltInComponentTypes::Get()->RequiresInstantiationMask.Set(Tags.BoundMaterialChanged);
 
diff --git a/Engine/Source/Runtime/MovieSceneTracks/Private/Sections/MovieSceneFadeSection.cpp b/Engine/Source/Runtime/MovieSceneTracks/Private/Sections/MovieSceneFadeSection.cpp
--- a/Engine/Source/Runtime/MovieSceneTracks/Private/Sections/MovieSceneFadeSection.cpp
+++ b/Engine/Source/Runtime/MovieSceneTracks/Private/Sections/MovieSceneFadeSection.cpp
@@ -3,6 +3,8 @@
 #include "Sections/MovieSceneFadeSection.h"
 #include "UObject/SequencerObjectVersion.h"
 #include "Channels/MovieSceneChannelProxy.h"
+#include "EntitySystem/BuiltInComponentTypes.h"
+#include "MovieSceneTracksComponentTypes.h"
 
 #include UE_INLINE_GENERATED_CPP_BY_NAME(MovieSceneFadeSection)
 
@@ -39,3 +41,23 @@ UMovieSceneFadeSection::UMovieSceneFadeSection()
 #endif
 }
 
+void UMovieSceneFadeSection::ImportEntityImpl(UMovieSceneEntitySystemLinker* EntityLinker, const FEntityImportParams& Params, FImportedEntity* OutImportedEntity)
+{
+	using namespace UE::MovieScene;
+
+	if (!FloatCurve.HasAnyData())
+	{
+		return;
+	}
+
+	const FBuiltInComponentTypes* BuiltInComponents = FBuiltInComponentTypes::Get();
+	const FMovieSceneTracksComponentTypes* TrackComponents = FMovieSceneTracksComponentTypes::Get();
+
+	OutImportedEntity->AddBuilder(
+		FEntityBuilder()
+		.AddTag(BuiltInComponents->Tags.Root)
+		.Add(BuiltInComponents->FloatChannel[0], &FloatCurve)
+		.Add(TrackComponents->Fade, FFadeComponentData{ FadeColor, (bool)bFadeAudio })
+	);
+}
+
diff --git a/Engine/Source/Runtime/MovieSceneTracks/Private/Systems/MovieSceneFadeSystem.cpp b/Engine/Source/Runtime/MovieSceneTracks/Private/Systems/MovieSceneFadeSystem.cpp
new file mode 100644
--- /dev/null
+++ b/Engine/Source/Runtime/MovieSceneTracks/Private/Systems/MovieSceneFadeSystem.cpp
@@ -0,0 +1,178 @@
+// Copyright Epic Games, Inc. All Rights Reserved.
+
+#include "Systems/MovieSceneFadeSystem.h"
+#include "Async/TaskGraphInterfaces.h"
+#include "Engine/Engine.h"
+#include "Engine/GameInstance.h"
+#include "EngineGlobals.h"
+#include "EntitySystem/BuiltInComponentTypes.h"
+#include "Evaluation/PreAnimatedState/MovieScenePreAnimatedStateStorage.h"
+#include "Evaluation/PreAnimatedState/MovieScenePreAnimatedStorageID.inl"
+#include "GameFramework/PlayerController.h"
+#include "IMovieScenePlayer.h"
+#include "MovieSceneExecutionToken.h"
+#include "MovieSceneSequence.h"
+#include "Sections/MovieSceneFadeSection.h"
+#include "EntitySystem/BuiltInComponentTypes.h"
+#include "MovieSceneTracksComponentTypes.h"
+
+#include UE_INLINE_GENERATED_CPP_BY_NAME(MovieSceneFadeSystem)
+
+namespace UE::MovieScene
+{
+
+struct FFadeUtil
+{
+	static void ApplyFade(IMovieScenePlayer& Player, float FadeValue, const FLinearColor& FadeColor, bool bFadeAudio)
+	{
+		// Set editor preview/fade
+		EMovieSceneViewportParams ViewportParams;
+		ViewportParams.SetWhichViewportParam = ( EMovieSceneViewportParams::SetViewportParam )( EMovieSceneViewportParams::SVP_FadeAmount | EMovieSceneViewportParams::SVP_FadeColor );
+		ViewportParams.FadeAmount = FadeValue;
+		ViewportParams.FadeColor = FadeColor;
+
+		TMap<FViewportClient*, EMovieSceneViewportParams> ViewportParamsMap;
+		Player.GetViewportSettings( ViewportParamsMap );
+		for( auto ViewportParamsPair : ViewportParamsMap )
+		{
+			ViewportParamsMap[ ViewportParamsPair.Key ] = ViewportParams;
+		}
+		Player.SetViewportSettings( ViewportParamsMap );
+
+		// Set runtime fade
+		UObject* Context = Player.GetPlaybackContext();
+		UWorld* World = Context ? Context->GetWorld() : nullptr;
+		if( World && ( World->WorldType == EWorldType::Game || World->WorldType == EWorldType::PIE ) )
+		{
+			APlayerController* PlayerController = World->GetGameInstance()->GetFirstLocalPlayerController();
+			if( PlayerController != nullptr && IsValid(PlayerController->PlayerCameraManager) )
+			{
+				PlayerController->PlayerCameraManager->SetManualCameraFade( FadeValue, FadeColor, bFadeAudio );
+			}
+		}
+	}
+};
+
+struct FPreAnimatedFadeState
+{
+	float FadeValue;
+	FLinearColor FadeColor;
+	bool bFadeAudio;
+
+	static FPreAnimatedFadeState SaveState(IMovieScenePlayer* Player)
+	{
+		float FadeAmount = 0.f;
+		FLinearColor FadeColor = FLinearColor::Black;
+		bool bFadeAudio = false;
+
+		UObject* Context = Player->GetPlaybackContext();
+		UWorld* World = Context ? Context->GetWorld() : nullptr;
+		if (World && (World->WorldType == EWorldType::Game || World->WorldType == EWorldType::PIE))
+		{
+			APlayerController* PlayerController = World->GetGameInstance()->GetFirstLocalPlayerController();
+			if (PlayerController != nullptr && IsValid(PlayerController->PlayerCameraManager))
+			{
+				FadeAmount = PlayerController->PlayerCameraManager->FadeAmount;
+				FadeColor = PlayerController->PlayerCameraManager->FadeColor;
+				bFadeAudio = PlayerController->PlayerCameraManager->bFadeAudio;
+			}
+		}
+
+		return FPreAnimatedFadeState{ FadeAmount, FadeColor, bFadeAudio };
+	}
+
+	void RestoreState(const FMovieSceneAnimTypeID& Unused, const FRestoreStateParams& Params)
+	{
+		IMovieScenePlayer* Player = Params.GetTerminalPlayer();
+		if (!ensure(Player))
+		{
+			return;
+		}
+		
+		FFadeUtil::ApplyFade(*Player, FadeValue, FadeColor, bFadeAudio);
+	}
+};
+
+struct FPreAnimatedFadeStateStorage : TSimplePreAnimatedStateStorage<FMovieSceneAnimTypeID, FPreAnimatedFadeState>
+{
+	static TAutoRegisterPreAnimatedStorageID<FPreAnimatedFadeStateStorage> StorageID;
+
+	FPreAnimatedStorageID GetStorageType() const override { return StorageID; }
+};
+
+TAutoRegisterPreAnimatedStorageID<FPreAnimatedFadeStateStorage> FPreAnimatedFadeStateStorage::StorageID;
+
+struct FEvaluateFade
+{
+	const FInstanceRegistry* InstanceRegistry;
+	TSharedPtr<FPreAnimatedFadeStateStorage> PreAnimatedStorage;
+
+	FEvaluateFade(const FInstanceRegistry* InInstanceRegistry, TSharedPtr<FPreAnimatedFadeStateStorage> InPreAnimatedStorage)
+		: InstanceRegistry(InInstanceRegistry)
+		, PreAnimatedStorage(InPreAnimatedStorage)
+	{}
+
+	void ForEachAllocation(const FEntityAllocation* Allocation, TRead<FMovieSceneEntityID> EntityIDs, TRead<FRootInstanceHandle> RootInstanceHandles, TRead<FFadeComponentData> FadeComponents, TRead<double> FadeAmounts)
+	{
+		static const TMovieSceneAnimTypeID<FEvaluateFade> AnimTypeID;
+
+		const FBuiltInComponentTypes* BuiltInComponents = FBuiltInComponentTypes::Get();
+
+		const int32 Num = Allocation->Num();
+		const bool bWantsRestoreState = Allocation->HasComponent(BuiltInComponents->Tags.RestoreState);
+		const FMovieSceneAnimTypeID Key = AnimTypeID;
+
+		for (int32 Index = 0; Index < Num; ++Index)
+		{
+			FRootInstanceHandle RootInstanceHandle = RootInstanceHandles[Index];
+			const FSequenceInstance& Instance = InstanceRegistry->GetInstance(RootInstanceHandle);
+			IMovieScenePlayer* Player = Instance.GetPlayer();
+
+			PreAnimatedStorage->BeginTrackingEntity(EntityIDs[Index], bWantsRestoreState, RootInstanceHandle, Key);
+			PreAnimatedStorage->CachePreAnimatedValue(Key, [Player](const FMovieSceneAnimTypeID&) { return FPreAnimatedFadeState::SaveState(Player); });
+
+			const FFadeComponentData& FadeComponent(FadeComponents[Index]);
+			FFadeUtil::ApplyFade(*Player, FadeAmounts[Index], FadeComponent.FadeColor, FadeComponent.bFadeAudio);
+		}
+	}
+};
+
+} // namespace UE::MovieScene
+
+UMovieSceneFadeSystem::UMovieSceneFadeSystem(const FObjectInitializer& ObjInit)
+	: Super(ObjInit)
+{
+	using namespace UE::MovieScene;
+
+	RelevantComponent = FMovieSceneTracksComponentTypes::Get()->Fade;
+
+	if (HasAnyFlags(RF_ClassDefaultObject))
+	{
+		const FBuiltInComponentTypes* BuiltInComponents = FBuiltInComponentTypes::Get();
+		DefineComponentConsumer(GetClass(), BuiltInComponents->DoubleResult[0]);
+	}
+}
+
+void UMovieSceneFadeSystem::OnLink()
+{
+	using namespace UE::MovieScene;
+
+	PreAnimatedStorage = Linker->PreAnimatedState.GetOrCreateStorage<FPreAnimatedFadeStateStorage>();
+}
+
+void UMovieSceneFadeSystem::OnRun(FSystemTaskPrerequisites& InPrerequisites, FSystemSubsequentTasks& Subsequents)
+{
+	using namespace UE::MovieScene;
+
+	const FBuiltInComponentTypes* BuiltInComponents = FBuiltInComponentTypes::Get();
+	const FMovieSceneTracksComponentTypes* TrackComponents = FMovieSceneTracksComponentTypes::Get();
+
+	FEntityTaskBuilder()
+	.ReadEntityIDs()
+	.Read(BuiltInComponents->RootInstanceHandle)
+	.Read(TrackComponents->Fade)
+	.Read(BuiltInComponents->DoubleResult[0])
+	.SetDesiredThread(Linker->EntityManager.GetGatherThread())
+	.Dispatch_PerAllocation<FEvaluateFade>(& Linker->EntityManager, InPrerequisites, &Subsequents, Linker->GetInstanceRegistry(), PreAnimatedStorage);
+}
+
diff --git a/Engine/Source/Runtime/MovieSceneTracks/Private/Systems/MovieSceneFadeSystem.h b/Engine/Source/Runtime/MovieSceneTracks/Private/Systems/MovieSceneFadeSystem.h
new file mode 100644
--- /dev/null
+++ b/Engine/Source/Runtime/MovieSceneTracks/Private/Systems/MovieSceneFadeSystem.h
@@ -0,0 +1,29 @@
+// Copyright Epic Games, Inc. All Rights Reserved.
+
+#pragma once
+
+#include "CoreTypes.h"
+#include "EntitySystem/MovieSceneEntitySystem.h"
+
+#include "MovieSceneFadeSystem.generated.h"
+
+namespace UE::MovieScene
+{
+	struct FPreAnimatedFadeStateStorage;
+}
+
+UCLASS()
+class UMovieSceneFadeSystem : public UMovieSceneEntitySystem
+{
+	GENERATED_BODY()
+
+public:
+	UMovieSceneFadeSystem(const FObjectInitializer& ObjInit);
+
+	virtual void OnLink() override;
+	virtual void OnRun(FSystemTaskPrerequisites& InPrerequisites, FSystemSubsequentTasks& Subsequents) override;
+
+private:
+	TSharedPtr<UE::MovieScene::FPreAnimatedFadeStateStorage> PreAnimatedStorage;
+};
+
diff --git a/Engine/Source/Runtime/MovieSceneTracks/Private/Tracks/MovieSceneFadeTrack.cpp b/Engine/Source/Runtime/MovieSceneTracks/Private/Tracks/MovieSceneFadeTrack.cpp
--- a/Engine/Source/Runtime/MovieSceneTracks/Private/Tracks/MovieSceneFadeTrack.cpp
+++ b/Engine/Source/Runtime/MovieSceneTracks/Private/Tracks/MovieSceneFadeTrack.cpp
@@ -2,8 +2,6 @@
 
 #include "Tracks/MovieSceneFadeTrack.h"
 #include "Sections/MovieSceneFadeSection.h"
-#include "Evaluation/MovieSceneFadeTemplate.h"
-#include "Evaluation/MovieSceneEvaluationTrack.h"
 #include "Templates/Casts.h"
 
 #include UE_INLINE_GENERATED_CPP_BY_NAME(MovieSceneFadeTrack)
@@ -29,11 +27,6 @@ UMovieSceneSection* UMovieSceneFadeTrack::CreateNewSection()
 	return NewObject<UMovieSceneFadeSection>(this, NAME_None, RF_Transactional);
 }
 
-FMovieSceneEvalTemplatePtr UMovieSceneFadeTrack::CreateTemplateForSection(const UMovieSceneSection& InSection) const
-{
-	return FMovieSceneFadeSectionTemplate(*CastChecked<UMovieSceneFadeSection>(&InSection));
-}
-
 #if WITH_EDITORONLY_DATA
 FText UMovieSceneFadeTrack::GetDefaultDisplayName() const
 {
diff --git a/Engine/Source/Runtime/MovieSceneTracks/Public/MovieSceneTracksComponentTypes.h b/Engine/Source/Runtime/MovieSceneTracks/Public/MovieSceneTracksComponentTypes.h
--- a/Engine/Source/Runtime/MovieSceneTracks/Public/MovieSceneTracksComponentTypes.h
+++ b/Engine/Source/Runtime/MovieSceneTracks/Public/MovieSceneTracksComponentTypes.h
@@ -130,6 +130,12 @@ struct FAttachmentComponent
 	FComponentDetachParams DetachParams;
 };
 
+struct FFadeComponentData
+{
+	FLinearColor FadeColor;
+	bool bFadeAudio;
+};
+
 struct FFloatPropertyTraits
 {
 	using StorageType  = double;
@@ -455,6 +461,8 @@ struct MOVIESCENETRACKS_API FMovieSceneTracksComponentTypes
 	TComponentTypeID<UObject*> BoundMaterial;
 	TComponentTypeID<UMaterialParameterCollection*> MPC;
 
+	TComponentTypeID<FFadeComponentData> Fade;
+
 	struct
 	{
 		TCustomPropertyRegistration<FBoolPropertyTraits> Bool;
diff --git a/Engine/Source/Runtime/MovieSceneTracks/Public/Sections/MovieSceneFadeSection.h b/Engine/Source/Runtime/MovieSceneTracks/Public/Sections/MovieSceneFadeSection.h
--- a/Engine/Source/Runtime/MovieSceneTracks/Public/Sections/MovieSceneFadeSection.h
+++ b/Engine/Source/Runtime/MovieSceneTracks/Public/Sections/MovieSceneFadeSection.h
@@ -7,7 +7,7 @@
 #include "HAL/Platform.h"
 #include "Math/Color.h"
 #include "MovieSceneSection.h"
-#include "Sections/MovieSceneFloatSection.h"
+#include "EntitySystem/IMovieSceneEntityProvider.h"
 #include "UObject/ObjectMacros.h"
 #include "UObject/UObjectGlobals.h"
 
@@ -22,12 +22,16 @@ class UObject;
 UCLASS(MinimalAPI)
 class UMovieSceneFadeSection
 	: public UMovieSceneSection
+	, public IMovieSceneEntityProvider
 {
 	GENERATED_BODY()
 
 	/** Default constructor. */
 	UMovieSceneFadeSection();
 
+	/** IMovieSceneEntityProvider interface */
+	void ImportEntityImpl(UMovieSceneEntitySystemLinker* EntityLinker, const FEntityImportParams& Params, FImportedEntity* OutImportedEntity) override;
+
 public:
 
 	/** Float data */
diff --git a/Engine/Source/Runtime/MovieSceneTracks/Public/Tracks/MovieSceneFadeTrack.h b/Engine/Source/Runtime/MovieSceneTracks/Public/Tracks/MovieSceneFadeTrack.h
--- a/Engine/Source/Runtime/MovieSceneTracks/Public/Tracks/MovieSceneFadeTrack.h
+++ b/Engine/Source/Runtime/MovieSceneTracks/Public/Tracks/MovieSceneFadeTrack.h
@@ -5,7 +5,6 @@
 #include "CoreMinimal.h"
 #include "UObject/ObjectMacros.h"
 #include "Tracks/MovieSceneFloatTrack.h"
-#include "Compilation/IMovieSceneTrackTemplateProducer.h"
 #include "MovieSceneFadeTrack.generated.h"
 
 /**
@@ -14,7 +13,6 @@
 UCLASS(MinimalAPI)
 class UMovieSceneFadeTrack
 	: public UMovieSceneFloatTrack
-	, public IMovieSceneTrackTemplateProducer
 {
 	GENERATED_BODY()
 
@@ -26,7 +24,6 @@ public:
 
 	virtual bool SupportsType(TSubclassOf<UMovieSceneSection> SectionClass) const override;
 	virtual UMovieSceneSection* CreateNewSection() override;
-	virtual FMovieSceneEvalTemplatePtr CreateTemplateForSection(const UMovieSceneSection& InSection) const override;
 
 #if WITH_EDITORONLY_DATA
 	virtual FText GetDefaultDisplayName() const override;
