PR #9968: Fix quat interpolation for consecutive keyframes in sequencer at low tick resolutions (Contributed by [REDACTED_PERSON])

diff --git a/Engine/Source/Runtime/MovieSceneTracks/Private/Systems/MovieSceneQuaternionInterpolationRotationSystem.cpp b/Engine/Source/Runtime/MovieSceneTracks/Private/Systems/MovieSceneQuaternionInterpolationRotationSystem.cpp
--- a/Engine/Source/Runtime/MovieSceneTracks/Private/Systems/MovieSceneQuaternionInterpolationRotationSystem.cpp
+++ b/Engine/Source/Runtime/MovieSceneTracks/Private/Systems/MovieSceneQuaternionInterpolationRotationSystem.cpp
@@ -31,11 +31,11 @@ struct FEvaluateQuaternionInterpolationRotationChannels
 		Index2 = Index2 < Times.Num() ? Index2 : INDEX_NONE;
 		if (Index1 != INDEX_NONE && Index2 != INDEX_NONE)
 		{
-			if (Times[Index1] != FrameTime.FrameNumber && Times[Index1] > FrameRange.GetLowerBoundValue()) 
+			if (Times[Index1] > FrameRange.GetLowerBoundValue()) 
 			{
 				FrameRange.SetLowerBoundValue(Times[Index1]);
 			}
-			if (Times[Index2] != FrameTime.FrameNumber && Times[Index1] < FrameRange.GetUpperBoundValue())
+			if (Times[Index2] != FrameTime.FrameNumber && Times[Index2] < FrameRange.GetUpperBoundValue())
 			{
 				FrameRange.SetUpperBoundValue(Times[Index2]);
 			}
