Adding includes and forward declarations.

diff --git a/Engine/Source/Runtime/Engine/Classes/Particles/Collision/ParticleModuleCollision.h b/Engine/Source/Runtime/Engine/Classes/Particles/Collision/ParticleModuleCollision.h
--- a/Engine/Source/Runtime/Engine/Classes/Particles/Collision/ParticleModuleCollision.h
+++ b/Engine/Source/Runtime/Engine/Classes/Particles/Collision/ParticleModuleCollision.h
@@ -14,6 +14,7 @@
 class UParticleEmitter;
 class UParticleLODLevel;
 class UParticleModuleTypeDataBase;
+struct FHitResult;
 struct FParticleEmitterInstance;
 
 UCLASS(editinlinenew, hidecategories=Object, MinimalAPI, meta=(DisplayName = "Actor Collision"))
diff --git a/Engine/Source/Runtime/Engine/Classes/Particles/ParticleEmitter.h b/Engine/Source/Runtime/Engine/Classes/Particles/ParticleEmitter.h
--- a/Engine/Source/Runtime/Engine/Classes/Particles/ParticleEmitter.h
+++ b/Engine/Source/Runtime/Engine/Classes/Particles/ParticleEmitter.h
@@ -18,6 +18,8 @@ class UInterpCurveEdSetup;
 class UParticleLODLevel;
 class UParticleSystemComponent;
 
+enum EDetailMode : int;
+
 //~=============================================================================
 //	Burst emissions
 //~=============================================================================
diff --git a/Engine/Source/Runtime/Engine/Private/Components/PointLightComponent.cpp b/Engine/Source/Runtime/Engine/Private/Components/PointLightComponent.cpp
--- a/Engine/Source/Runtime/Engine/Private/Components/PointLightComponent.cpp
+++ b/Engine/Source/Runtime/Engine/Private/Components/PointLightComponent.cpp
@@ -10,6 +10,7 @@
 #include "Engine/Texture2D.h"
 #include "PointLightSceneProxy.h"
 #include "UObject/UnrealType.h"
+#include "SceneInterface.h"
 #include "SceneView.h"
 
 #include UE_INLINE_GENERATED_CPP_BY_NAME(PointLightComponent)
diff --git a/Engine/Source/Runtime/Engine/Private/Components/SpotLightComponent.cpp b/Engine/Source/Runtime/Engine/Private/Components/SpotLightComponent.cpp
--- a/Engine/Source/Runtime/Engine/Private/Components/SpotLightComponent.cpp
+++ b/Engine/Source/Runtime/Engine/Private/Components/SpotLightComponent.cpp
@@ -10,6 +10,7 @@
 #include "Engine/Texture2D.h"
 #include "PointLightSceneProxy.h"
 #include "UObject/UnrealType.h"
+#include "SceneInterface.h"
 #include "SceneView.h"
 
 #include UE_INLINE_GENERATED_CPP_BY_NAME(SpotLightComponent)
diff --git a/Engine/Source/Runtime/Engine/Private/LocalLightSceneProxy.cpp b/Engine/Source/Runtime/Engine/Private/LocalLightSceneProxy.cpp
--- a/Engine/Source/Runtime/Engine/Private/LocalLightSceneProxy.cpp
+++ b/Engine/Source/Runtime/Engine/Private/LocalLightSceneProxy.cpp
@@ -7,6 +7,7 @@
 #include "LocalLightSceneProxy.h"
 #include "Components/LocalLightComponent.h"
 #include "Math/InverseRotationMatrix.h"
+#include "SceneManagement.h"
 #include "SceneView.h"
 
 /** Initialization constructor. */
diff --git a/Engine/Source/Runtime/Engine/Private/Particles/ParticleSystemRender.cpp b/Engine/Source/Runtime/Engine/Private/Particles/ParticleSystemRender.cpp
--- a/Engine/Source/Runtime/Engine/Private/Particles/ParticleSystemRender.cpp
+++ b/Engine/Source/Runtime/Engine/Private/Particles/ParticleSystemRender.cpp
@@ -10,6 +10,7 @@
 #include "Materials/MaterialRenderProxy.h"
 #include "Particles/Orientation/ParticleModuleOrientationAxisLock.h"
 #include "UObject/UObjectIterator.h"
+#include "MaterialDomain.h"
 #include "MeshParticleVertexFactory.h"
 #include "Particles/ParticleEmitter.h"
 #include "Particles/ParticleSystemComponent.h"
diff --git a/Engine/Source/Runtime/MovieSceneTracks/Private/Systems/MovieSceneDataLayerSystem.cpp b/Engine/Source/Runtime/MovieSceneTracks/Private/Systems/MovieSceneDataLayerSystem.cpp
--- a/Engine/Source/Runtime/MovieSceneTracks/Private/Systems/MovieSceneDataLayerSystem.cpp
+++ b/Engine/Source/Runtime/MovieSceneTracks/Private/Systems/MovieSceneDataLayerSystem.cpp
@@ -6,6 +6,7 @@
 #include "EntitySystem/MovieSceneEntitySystemLinker.h"
 #include "EntitySystem/BuiltInComponentTypes.h"
 #include "Evaluation/PreAnimatedState/MovieScenePreAnimatedEntityCaptureSource.h"
+#include "Evaluation/PreAnimatedState/MovieScenePreAnimatedObjectStorage.h"
 #include "Evaluation/PreAnimatedState/MovieScenePreAnimatedStateStorage.h"
 #include "Evaluation/PreAnimatedState/MovieScenePreAnimatedStorageID.inl"
 #include "Evaluation/MovieSceneEvaluationTemplateInstance.h"
diff --git a/Engine/Source/Runtime/Renderer/Private/DeferredShadingRenderer.h b/Engine/Source/Runtime/Renderer/Private/DeferredShadingRenderer.h
--- a/Engine/Source/Runtime/Renderer/Private/DeferredShadingRenderer.h
+++ b/Engine/Source/Runtime/Renderer/Private/DeferredShadingRenderer.h
@@ -54,6 +54,7 @@ struct FBuildHZBAsyncComputeParams;
 struct FForwardBasePassTextures;
 struct FTranslucentLightInjectionCollector;
 struct FRayTracingPickingFeedback;
+struct FDBufferTextures;
 
 #if RHI_RAYTRACING
 struct FRayTracingRelevantPrimitiveTaskData;
diff --git a/Engine/Source/Runtime/Renderer/Private/PostProcess/PostProcessCombineLUTs.cpp b/Engine/Source/Runtime/Renderer/Private/PostProcess/PostProcessCombineLUTs.cpp
--- a/Engine/Source/Runtime/Renderer/Private/PostProcess/PostProcessCombineLUTs.cpp
+++ b/Engine/Source/Runtime/Renderer/Private/PostProcess/PostProcessCombineLUTs.cpp
@@ -2,6 +2,7 @@
 
 #include "PostProcess/PostProcessCombineLUTs.h"
 #include "DataDrivenShaderPlatformInfo.h"
+#include "Engine/Texture.h"
 #include "PostProcess/PostProcessTonemap.h"
 #include "ScenePrivate.h"
 #include "VolumeRendering.h"
diff --git a/Engine/Source/Runtime/Renderer/Private/PostProcess/PostProcessEyeAdaptation.cpp b/Engine/Source/Runtime/Renderer/Private/PostProcess/PostProcessEyeAdaptation.cpp
--- a/Engine/Source/Runtime/Renderer/Private/PostProcess/PostProcessEyeAdaptation.cpp
+++ b/Engine/Source/Runtime/Renderer/Private/PostProcess/PostProcessEyeAdaptation.cpp
@@ -2,6 +2,7 @@
 
 #include "PostProcess/PostProcessEyeAdaptation.h"
 
+#include "Engine/Texture.h"
 #include "TranslucentLighting.h"
 
 #include "SceneTextureParameters.h"
diff --git a/Engine/Source/Runtime/Renderer/Private/PostProcess/PostProcessMaterial.cpp b/Engine/Source/Runtime/Renderer/Private/PostProcess/PostProcessMaterial.cpp
--- a/Engine/Source/Runtime/Renderer/Private/PostProcess/PostProcessMaterial.cpp
+++ b/Engine/Source/Runtime/Renderer/Private/PostProcess/PostProcessMaterial.cpp
@@ -9,6 +9,7 @@
 #include "RendererModule.h"
 #include "Materials/Material.h"
 #include "Materials/MaterialRenderProxy.h"
+#include "MaterialDomain.h"
 #include "MaterialShaderType.h"
 #include "MaterialShader.h"
 #include "RenderUtils.h"
diff --git a/Engine/Source/Runtime/Renderer/Private/PostProcess/PostProcessMobile.cpp b/Engine/Source/Runtime/Renderer/Private/PostProcess/PostProcessMobile.cpp
--- a/Engine/Source/Runtime/Renderer/Private/PostProcess/PostProcessMobile.cpp
+++ b/Engine/Source/Runtime/Renderer/Private/PostProcess/PostProcessMobile.cpp
@@ -6,6 +6,7 @@
 
 #include "PostProcess/PostProcessMobile.h"
 #include "DataDrivenShaderPlatformInfo.h"
+#include "Engine/Texture.h"
 #include "StaticBoundShaderState.h"
 #include "SceneUtils.h"
 #include "SceneRenderTargetParameters.h"
diff --git a/Engine/Source/Runtime/Renderer/Private/Shadows/ShadowSceneRenderer.h b/Engine/Source/Runtime/Renderer/Private/Shadows/ShadowSceneRenderer.h
--- a/Engine/Source/Runtime/Renderer/Private/Shadows/ShadowSceneRenderer.h
+++ b/Engine/Source/Runtime/Renderer/Private/Shadows/ShadowSceneRenderer.h
@@ -17,6 +17,7 @@ class FDeferredShadingSceneRenderer;
 class FWholeSceneProjectedShadowInitializer;
 class FRDGBuilder;
 class FVirtualShadowMapPerLightCacheEntry;
+struct FNaniteVisibilityQuery;
 
 /**
  * Transient scope for per-frame rendering resources for the shadow rendering.
diff --git a/Engine/Source/Runtime/Renderer/Private/VT/VirtualTextureSystem.cpp b/Engine/Source/Runtime/Renderer/Private/VT/VirtualTextureSystem.cpp
--- a/Engine/Source/Runtime/Renderer/Private/VT/VirtualTextureSystem.cpp
+++ b/Engine/Source/Runtime/Renderer/Private/VT/VirtualTextureSystem.cpp
@@ -4,6 +4,7 @@
 
 #include "AllocatedVirtualTexture.h"
 #include "Engine/Canvas.h"
+#include "Engine/Texture2D.h"
 #include "Debug/DebugDrawService.h"
 #include "HAL/IConsoleManager.h"
 #include "Materials/MaterialRenderProxy.h"
diff --git a/Engine/Source/Runtime/Renderer/Private/VirtualShadowMaps/VirtualShadowMapCacheManager.h b/Engine/Source/Runtime/Renderer/Private/VirtualShadowMaps/VirtualShadowMapCacheManager.h
--- a/Engine/Source/Runtime/Renderer/Private/VirtualShadowMaps/VirtualShadowMapCacheManager.h
+++ b/Engine/Source/Runtime/Renderer/Private/VirtualShadowMaps/VirtualShadowMapCacheManager.h
@@ -16,6 +16,8 @@ class FRHIGPUBufferReadback;
 class FGPUScene;
 class FVirtualShadowMapPerLightCacheEntry;
 
+namespace Nanite { struct FPackedViewParams; }
+
 #define VSM_LOG_INVALIDATIONS 0
 
 class FVirtualShadowMapCacheEntry
