Sequencer: NLA Fix for time start calculation and regression frix from some UAnimDataModel changes, need to check non zero translation not non zero scale.
#trivial

diff --git a/Engine/Source/Runtime/MovieSceneTracks/Private/Sections/MovieSceneSkeletalAnimationSection.cpp b/Engine/Source/Runtime/MovieSceneTracks/Private/Sections/MovieSceneSkeletalAnimationSection.cpp
--- a/Engine/Source/Runtime/MovieSceneTracks/Private/Sections/MovieSceneSkeletalAnimationSection.cpp
+++ b/Engine/Source/Runtime/MovieSceneTracks/Private/Sections/MovieSceneSkeletalAnimationSection.cpp
@@ -545,14 +545,16 @@ int32 UMovieSceneSkeletalAnimationSection::SetBoneIndexForRootMotionCalculations
 				if (DataModel->IsValidBoneTrackName(BoneName))
 				{
 					DataModel->GetBoneTrackTransforms(BoneName, OutTransforms);
-					
+					TOptional<FTransform> LocalPreviousTransform;
 					for (const FTransform& Transform : OutTransforms)
 					{
-						if (Transform.GetScale3D().IsNearlyZero() == false)
+						if (Transform.GetLocation().IsNearlyZero() == false
+							&& LocalPreviousTransform.IsSet() && LocalPreviousTransform.GetValue().GetLocation() != Transform.GetLocation())
 						{
 							TempRootBoneIndex = BoneIndex;
 							break;
 						}
+						LocalPreviousTransform = Transform;
 					}
 					
 					OutTransforms.Reset();
@@ -612,7 +614,9 @@ FTransform UMovieSceneSkeletalAnimationSection::GetRootMotionStartOffset() const
 	UAnimSequence* AnimSequence = Cast<UAnimSequence>(Params.Animation);
 	if (AnimSequence && TempRootBoneIndex.IsSet() && TempRootBoneIndex.GetValue() != 0)
 	{
-		return AnimSequence->ExtractRootTrackTransform(0, nullptr);
+		UMovieScene* MovieScene = GetTypedOuter<UMovieScene>();
+		const double StartSeconds = MovieScene ? (MapTimeToAnimation(FFrameNumber(0), MovieScene->GetTickResolution())): 0.0;
+		return AnimSequence->ExtractRootTrackTransform(StartSeconds, nullptr);
 	}
 	return FTransform::Identity;;
 }
