Sequencer: Numeric slider now obeys property metadata

diff --git a/Engine/Source/Editor/MovieSceneTools/Private/CurveKeyEditors/SNumericKeyEditor.h b/Engine/Source/Editor/MovieSceneTools/Private/CurveKeyEditors/SNumericKeyEditor.h
--- a/Engine/Source/Editor/MovieSceneTools/Private/CurveKeyEditors/SNumericKeyEditor.h
+++ b/Engine/Source/Editor/MovieSceneTools/Private/CurveKeyEditors/SNumericKeyEditor.h
@@ -11,6 +11,7 @@
 #include "Widgets/Input/SSpinBox.h"
 #include "Styling/AppStyle.h"
 #include "CurveKeyEditors/SequencerKeyEditor.h"
+#include "NumericPropertyParams.h"
 
 #define LOCTEXT_NAMESPACE "NumericKeyEditor"
 
@@ -41,15 +42,46 @@ public:
 	void Construct(const FArguments& InArgs, const TSequencerKeyEditor<ChannelType, NumericType>& InKeyEditor)
 	{
 		KeyEditor = InKeyEditor;
+
+		const FProperty* Property = nullptr;
+		ISequencer* Sequencer = InKeyEditor.GetSequencer();
+		FTrackInstancePropertyBindings* PropertyBindings = InKeyEditor.GetPropertyBindings();
+		if (Sequencer && PropertyBindings)
+		{
+			for (TWeakObjectPtr<> WeakObject : Sequencer->FindBoundObjects(InKeyEditor.GetObjectBindingID(), Sequencer->GetFocusedTemplateID()))
+			{
+				if (UObject* Object = WeakObject.Get())
+				{
+					Property = PropertyBindings->GetProperty(*Object);
+					if (Property)
+					{
+						break;
+					}
+				}
+			}
+		}
+
+		const typename TNumericPropertyParams<NumericType>::FMetaDataGetter MetaDataGetter = TNumericPropertyParams<NumericType>::FMetaDataGetter::CreateLambda([&](const FName& Key)
+		{
+			return InKeyEditor.GetMetaData(Key);
+		});
+
+		TNumericPropertyParams<NumericType> NumericPropertyParams(Property, MetaDataGetter);
+
 		ChildSlot
 		[
 			SNew(SNonThrottledSpinBox<NumericType>)
 			.Style(&FAppStyle::GetWidgetStyle<FSpinBoxStyle>("Sequencer.HyperlinkSpinBox"))
 			.Font(FAppStyle::GetFontStyle("Sequencer.AnimationOutliner.RegularFont"))
-			.MinValue(TOptional<NumericType>())
-			.MaxValue(TOptional<NumericType>())
-			.MaxSliderValue(TOptional<NumericType>())
-			.MinSliderValue(TOptional<NumericType>())
+			.MinValue(NumericPropertyParams.MinValue)
+			.MaxValue(NumericPropertyParams.MaxValue)
+			.MinSliderValue(NumericPropertyParams.MinSliderValue)
+			.MaxSliderValue(NumericPropertyParams.MaxSliderValue)
+			.SliderExponent(NumericPropertyParams.SliderExponent)
+			.Delta(NumericPropertyParams.Delta)
+			// LinearDeltaSensitivity needs to be left unset if not provided, rather than being set to some default
+			.LinearDeltaSensitivity(NumericPropertyParams.GetLinearDeltaSensitivityAttribute())
+			.WheelStep(NumericPropertyParams.WheelStep)
 			.Value_Raw(&KeyEditor, &decltype(KeyEditor)::GetCurrentValue)
 			.OnValueChanged(this, &SNumericKeyEditor::OnValueChanged)
 			.OnValueCommitted(this, &SNumericKeyEditor::OnValueCommitted)
diff --git a/Engine/Source/Editor/MovieSceneTools/Private/CurveKeyEditors/SequencerKeyEditor.h b/Engine/Source/Editor/MovieSceneTools/Private/CurveKeyEditors/SequencerKeyEditor.h
--- a/Engine/Source/Editor/MovieSceneTools/Private/CurveKeyEditors/SequencerKeyEditor.h
+++ b/Engine/Source/Editor/MovieSceneTools/Private/CurveKeyEditors/SequencerKeyEditor.h
@@ -187,6 +187,32 @@ struct TSequencerKeyEditor
 		return WeakPropertyBindings.Pin().Get();
 	}
 
+	FString GetMetaData(const FName& Key) const
+	{
+		ISequencer* Sequencer = GetSequencer();
+		FTrackInstancePropertyBindings* PropertyBindings = GetPropertyBindings();
+		if (Sequencer && PropertyBindings)
+		{
+			for (TWeakObjectPtr<> WeakObject : Sequencer->FindBoundObjects(ObjectBindingID, Sequencer->GetFocusedTemplateID()))
+			{
+				if (UObject* Object = WeakObject.Get())
+				{
+					if (FProperty* Property = PropertyBindings->GetProperty(*Object))
+					{
+						return Property->GetMetaData(Key);
+					}
+				}
+			}
+		}
+
+		if (const FMovieSceneChannelMetaData* MetaData = ChannelHandle.GetMetaData())
+		{
+			return MetaData->GetPropertyMetaData(Key);
+		}
+
+		return FString();
+	}
+
 private:
 
 	FGuid ObjectBindingID;
diff --git a/Engine/Source/Runtime/MovieScene/Private/Channels/MovieSceneChannelEditorData.cpp b/Engine/Source/Runtime/MovieScene/Private/Channels/MovieSceneChannelEditorData.cpp
--- a/Engine/Source/Runtime/MovieScene/Private/Channels/MovieSceneChannelEditorData.cpp
+++ b/Engine/Source/Runtime/MovieScene/Private/Channels/MovieSceneChannelEditorData.cpp
@@ -42,4 +42,9 @@ void FMovieSceneChannelMetaData::SetIdentifiers(FName InName, FText InDisplayTex
 	DisplayText = InDisplayText;
 }
 
+FString FMovieSceneChannelMetaData::GetPropertyMetaData(const FName& InKey) const
+{
+	return PropertyMetaData.FindRef(InKey);
+}
+
 #endif	// WITH_EDITOR
\ No newline at end of file
diff --git a/Engine/Source/Runtime/MovieScene/Public/Channels/MovieSceneChannelEditorData.h b/Engine/Source/Runtime/MovieScene/Public/Channels/MovieSceneChannelEditorData.h
--- a/Engine/Source/Runtime/MovieScene/Public/Channels/MovieSceneChannelEditorData.h
+++ b/Engine/Source/Runtime/MovieScene/Public/Channels/MovieSceneChannelEditorData.h
@@ -49,6 +49,13 @@ struct FMovieSceneChannelMetaData
 	 */
 	MOVIESCENE_API void SetIdentifiers(FName InName, FText InDisplayText, FText InGroup = FText());
 
+	/*
+	 * Get property metadata that corresponds to the given key.
+	 * 
+	 * @param InKey The requested key to get metadata for
+	 */
+	MOVIESCENE_API FString GetPropertyMetaData(const FName& InKey) const;
+
 	/** Whether this channel is enabled or not */
 	uint8 bEnabled : 1;
 	/** True if this channel can be collapsed onto the top level track node */
@@ -67,6 +74,8 @@ struct FMovieSceneChannelMetaData
 	FText LongIntentNameFormat;
 	/** Optional color to draw underneath the keys on this channel */
 	TOptional<FLinearColor> Color;
+	/** Property meta data */
+	TMap<FName, FString> PropertyMetaData;
 };
 
 
diff --git a/Engine/Source/Runtime/MovieSceneTracks/Private/Sections/MovieScene3DTransformSection.cpp b/Engine/Source/Runtime/MovieSceneTracks/Private/Sections/MovieScene3DTransformSection.cpp
--- a/Engine/Source/Runtime/MovieSceneTracks/Private/Sections/MovieScene3DTransformSection.cpp
+++ b/Engine/Source/Runtime/MovieSceneTracks/Private/Sections/MovieScene3DTransformSection.cpp
@@ -29,6 +29,12 @@ struct F3DTransformChannelEditorData
 {
 	F3DTransformChannelEditorData(EMovieSceneTransformChannel Mask,int SortOrderStart)
 	{
+		static const TSet<FName> PropertyMetaDataKeys = { "UIMin", "UIMax", "SliderExponent", "LinearDeltaSensitivity", "Delta", "ClampMin", "ClampMax", "ForceUnits", "WheelStep" };
+
+		const FProperty* RelativeLocationProperty = USceneComponent::StaticClass()->FindPropertyByName(USceneComponent::GetRelativeLocationPropertyName());
+		const FProperty* RelativeRotationProperty = USceneComponent::StaticClass()->FindPropertyByName(USceneComponent::GetRelativeRotationPropertyName());
+		const FProperty* RelativeScale3DProperty = USceneComponent::StaticClass()->FindPropertyByName(USceneComponent::GetRelativeScale3DPropertyName());
+
 		FText LocationGroup = NSLOCTEXT("MovieSceneTransformSection", "Location", "Location");
 		FText RotationGroup = NSLOCTEXT("MovieSceneTransformSection", "Rotation", "Rotation");
 		FText ScaleGroup    = NSLOCTEXT("MovieSceneTransformSection", "Scale",    "Scale");
@@ -38,18 +44,30 @@ struct F3DTransformChannelEditorData
 			MetaData[0].Color = FCommonChannelData::RedChannelColor;
 			MetaData[0].SortOrder = SortOrderStart;
 			MetaData[0].bCanCollapseToTrack = false;
+			for (const FName& PropertyMetaDataKey : PropertyMetaDataKeys)
+			{
+				MetaData[0].PropertyMetaData.Add(PropertyMetaDataKey, RelativeLocationProperty->GetMetaData(PropertyMetaDataKey));
+			}
 
 			MetaData[1].SetIdentifiers("Location.Y", FCommonChannelData::ChannelY, LocationGroup);
 			MetaData[1].bEnabled = EnumHasAllFlags(Mask, EMovieSceneTransformChannel::TranslationY);
 			MetaData[1].Color = FCommonChannelData::GreenChannelColor;
 			MetaData[1].SortOrder = SortOrderStart + 1;
 			MetaData[1].bCanCollapseToTrack = false;
+			for (const FName& PropertyMetaDataKey : PropertyMetaDataKeys)
+			{
+				MetaData[1].PropertyMetaData.Add(PropertyMetaDataKey, RelativeLocationProperty->GetMetaData(PropertyMetaDataKey));
+			}
 
 			MetaData[2].SetIdentifiers("Location.Z", FCommonChannelData::ChannelZ, LocationGroup);
 			MetaData[2].bEnabled = EnumHasAllFlags(Mask, EMovieSceneTransformChannel::TranslationZ);
 			MetaData[2].Color = FCommonChannelData::BlueChannelColor;
 			MetaData[2].SortOrder = SortOrderStart + 2;
 			MetaData[2].bCanCollapseToTrack = false;
+			for (const FName& PropertyMetaDataKey : PropertyMetaDataKeys)
+			{
+				MetaData[2].PropertyMetaData.Add(PropertyMetaDataKey, RelativeLocationProperty->GetMetaData(PropertyMetaDataKey));
+			}
 		}
 		{
 			MetaData[3].SetIdentifiers("Rotation.X", NSLOCTEXT("MovieSceneTransformSection", "RotationX", "Roll"), RotationGroup);
@@ -57,18 +75,30 @@ struct F3DTransformChannelEditorData
 			MetaData[3].Color = FCommonChannelData::RedChannelColor;
 			MetaData[3].SortOrder = SortOrderStart + 3;
 			MetaData[3].bCanCollapseToTrack = false;
+			for (const FName& PropertyMetaDataKey : PropertyMetaDataKeys)
+			{
+				MetaData[3].PropertyMetaData.Add(PropertyMetaDataKey, RelativeRotationProperty->GetMetaData(PropertyMetaDataKey));
+			}
 
 			MetaData[4].SetIdentifiers("Rotation.Y", NSLOCTEXT("MovieSceneTransformSection", "RotationY", "Pitch"), RotationGroup);
 			MetaData[4].bEnabled = EnumHasAllFlags(Mask, EMovieSceneTransformChannel::RotationY);
 			MetaData[4].Color = FCommonChannelData::GreenChannelColor;
 			MetaData[4].SortOrder = SortOrderStart + 4;
 			MetaData[4].bCanCollapseToTrack = false;
+			for (const FName& PropertyMetaDataKey : PropertyMetaDataKeys)
+			{
+				MetaData[4].PropertyMetaData.Add(PropertyMetaDataKey, RelativeRotationProperty->GetMetaData(PropertyMetaDataKey));
+			}
 
 			MetaData[5].SetIdentifiers("Rotation.Z", NSLOCTEXT("MovieSceneTransformSection", "RotationZ", "Yaw"), RotationGroup);
 			MetaData[5].bEnabled = EnumHasAllFlags(Mask, EMovieSceneTransformChannel::RotationZ);
 			MetaData[5].Color = FCommonChannelData::BlueChannelColor;
 			MetaData[5].SortOrder = SortOrderStart + 5;
 			MetaData[5].bCanCollapseToTrack = false;
+			for (const FName& PropertyMetaDataKey : PropertyMetaDataKeys)
+			{
+				MetaData[5].PropertyMetaData.Add(PropertyMetaDataKey, RelativeRotationProperty->GetMetaData(PropertyMetaDataKey));
+			}
 		}
 		{
 			MetaData[6].SetIdentifiers("Scale.X", FCommonChannelData::ChannelX, ScaleGroup);
@@ -76,18 +106,30 @@ struct F3DTransformChannelEditorData
 			MetaData[6].Color = FCommonChannelData::RedChannelColor;
 			MetaData[6].SortOrder = SortOrderStart + 6;
 			MetaData[6].bCanCollapseToTrack = false;
+			for (const FName& PropertyMetaDataKey : PropertyMetaDataKeys)
+			{
+				MetaData[6].PropertyMetaData.Add(PropertyMetaDataKey, RelativeScale3DProperty->GetMetaData(PropertyMetaDataKey));
+			}
 
 			MetaData[7].SetIdentifiers("Scale.Y", FCommonChannelData::ChannelY, ScaleGroup);
 			MetaData[7].bEnabled = EnumHasAllFlags(Mask, EMovieSceneTransformChannel::ScaleY);
 			MetaData[7].Color = FCommonChannelData::GreenChannelColor;
 			MetaData[7].SortOrder = SortOrderStart + 7;
 			MetaData[7].bCanCollapseToTrack = false;
+			for (const FName& PropertyMetaDataKey : PropertyMetaDataKeys)
+			{
+				MetaData[7].PropertyMetaData.Add(PropertyMetaDataKey, RelativeScale3DProperty->GetMetaData(PropertyMetaDataKey));
+			}
 
 			MetaData[8].SetIdentifiers("Scale.Z", FCommonChannelData::ChannelZ, ScaleGroup);
 			MetaData[8].bEnabled = EnumHasAllFlags(Mask, EMovieSceneTransformChannel::ScaleZ);
 			MetaData[8].Color = FCommonChannelData::BlueChannelColor;
 			MetaData[8].SortOrder = SortOrderStart + 8;
 			MetaData[8].bCanCollapseToTrack = false;
+			for (const FName& PropertyMetaDataKey : PropertyMetaDataKeys)
+			{
+				MetaData[8].PropertyMetaData.Add(PropertyMetaDataKey, RelativeScale3DProperty->GetMetaData(PropertyMetaDataKey));
+			}
 		}
 		{
 			MetaData[9].SetIdentifiers("Weight", NSLOCTEXT("MovieSceneTransformSection", "Weight", "Weight"));
