Sequencer: Constraints: Handle duplication of contraints so it recreates proper constraint duplicate with the new duplicated objects and add preanimated state for component constraints. Put this all in one review since the fixes are interrelated.

For duplication, copy/paste of tracks, levels we need to not replace the ChildTRS since we can set that in the system evaluation, and it would resolve the original child trs to the new copy on paste which is not what we want, we want to keep the original.
Also added a PostDuplicate so we don't actually use a copy of the Constraint.
Fixed up places were would we get the constraint from the manager by it's FName, which may not be unique. Also make sure the Component Handles are owned by the constraint, not the manager since it duplication could return an existing handle if the parameters sent to DuplicateObject match.

diff --git a/Engine/Plugins/Animation/ControlRig/Source/ControlRig/Private/Sequencer/MovieSceneControlRigParameterSection.cpp b/Engine/Plugins/Animation/ControlRig/Source/ControlRig/Private/Sequencer/MovieSceneControlRigParameterSection.cpp
--- a/Engine/Plugins/Animation/ControlRig/Source/ControlRig/Private/Sequencer/MovieSceneControlRigParameterSection.cpp
+++ b/Engine/Plugins/Animation/ControlRig/Source/ControlRig/Private/Sequencer/MovieSceneControlRigParameterSection.cpp
@@ -975,9 +975,28 @@ void UMovieSceneControlRigParameterSection::SetBlendType(EMovieSceneBlendType In
 void UMovieSceneControlRigParameterSection::Serialize(FArchive& Ar)
 {
 	Super::Serialize(Ar);
-
 }
- 
+#if WITH_EDITOR
+void UMovieSceneControlRigParameterSection::PostDuplicate(bool bDuplicateForPIE)
+{
+	Super::PostDuplicate(bDuplicateForPIE);
+	//if not duplicating for PIE clear out the soft object ptrs to the constraints we dont' want to hold pointers to the old constraints.
+	//also make the copy to spawn so we can duplicate it if it doesn't exist in the level
+	//to the old ones but make new duplicates
+	if (bDuplicateForPIE == false)
+	{
+		for (FConstraintAndActiveChannel& ConstraintChannel : ConstraintsChannels)
+		{
+			if (ConstraintChannel.Constraint.IsValid() && !ConstraintChannel.ConstraintCopyToSpawn)
+			{
+				ConstraintChannel.ConstraintCopyToSpawn = ConstraintChannel.Constraint->Duplicate(this);
+			}
+			ConstraintChannel.Constraint.Reset();
+		}
+	}
+}
+#endif
+
 void UMovieSceneControlRigParameterSection::PostEditImport()
 {
 	Super::PostEditImport();
@@ -1417,9 +1436,10 @@ bool UMovieSceneControlRigParameterSection::HasConstraintChannel(const FName& In
 FConstraintAndActiveChannel* UMovieSceneControlRigParameterSection::GetConstraintChannel(const FName& InConstraintName)
 {
 	const int32 Index = ConstraintsChannels.IndexOfByPredicate([InConstraintName](const FConstraintAndActiveChannel& InChannel)
-	{
-		return InChannel.Constraint.IsValid() ? InChannel.Constraint->GetFName() == InConstraintName : false;
-	});
+		{
+			return InChannel.Constraint.IsValid() ? InChannel.Constraint->GetFName() == InConstraintName : (InChannel.ConstraintCopyToSpawn ?
+			InChannel.ConstraintCopyToSpawn->GetFName() == InConstraintName : false);
+		});
 	return (Index != INDEX_NONE) ? &ConstraintsChannels[Index] : nullptr;	
 }
 
@@ -1467,11 +1487,15 @@ void UMovieSceneControlRigParameterSection::AddConstraintChannel(UTickableConstr
 	}
 }
 
-void UMovieSceneControlRigParameterSection::RemoveConstraintChannel(const FName& InConstraintName)
+void UMovieSceneControlRigParameterSection::RemoveConstraintChannel(const UTickableConstraint* InConstraint)
 {
-	const int32 Index = ConstraintsChannels.IndexOfByPredicate([InConstraintName](const FConstraintAndActiveChannel& InChannel)
+	if (bDoNotRemoveChannel == true)
 	{
-		return InChannel.Constraint.IsValid() ? InChannel.Constraint->GetFName() == InConstraintName : false;
+		return;
+	}
+	const int32 Index = ConstraintsChannels.IndexOfByPredicate([InConstraint](const FConstraintAndActiveChannel& InChannel)
+	{
+		return InChannel.Constraint.IsValid() ? InChannel.Constraint == InConstraint : false;
 	});
 
 	if (ConstraintsChannels.IsValidIndex(Index))
diff --git a/Engine/Plugins/Animation/ControlRig/Source/ControlRig/Public/Sequencer/MovieSceneControlRigParameterSection.h b/Engine/Plugins/Animation/ControlRig/Source/ControlRig/Public/Sequencer/MovieSceneControlRigParameterSection.h
--- a/Engine/Plugins/Animation/ControlRig/Source/ControlRig/Public/Sequencer/MovieSceneControlRigParameterSection.h
+++ b/Engine/Plugins/Animation/ControlRig/Source/ControlRig/Public/Sequencer/MovieSceneControlRigParameterSection.h
@@ -280,7 +280,7 @@ public:
 	/*
 	*  Remove Constraint channel
 	*/
-	virtual void RemoveConstraintChannel(const FName& InConstraintName) override;
+	virtual void RemoveConstraintChannel(const UTickableConstraint* InConstraint) override;
 
 	/*
 	*  Get The channels
@@ -290,7 +290,7 @@ public:
 	/*
 	*  Replace the constraint with the specified name with the new one
 	*/
-	virtual void ReplaceConstraint(const FName InName, UTickableConstraint* InConstraint)  override;
+	virtual void ReplaceConstraint(const FName InConstraintName, UTickableConstraint* InConstraint)  override;
 
 	/*
 	*  What to do if the constraint object has been changed, for example by an undo or redo.
@@ -503,6 +503,10 @@ protected:
 	virtual void GetReferencedBindings(TArray<FGuid>& OutBindings) override;
 	virtual void PreSave(FObjectPreSaveContext SaveContext) override;
 
+#if WITH_EDITOR
+	virtual void PostDuplicate(bool bDuplicateForPIE) override;
+#endif
+
 
 	// When true we do not set a key on the section, since it will be set because we changed the value
 	// We need this because control rig notifications are set on every change even when just changing sequencer time
diff --git a/Engine/Plugins/Animation/ControlRig/Source/ControlRigEditor/Private/Sequencer/ControlRigParameterTrackEditor.cpp b/Engine/Plugins/Animation/ControlRig/Source/ControlRigEditor/Private/Sequencer/ControlRigParameterTrackEditor.cpp
--- a/Engine/Plugins/Animation/ControlRig/Source/ControlRigEditor/Private/Sequencer/ControlRigParameterTrackEditor.cpp
+++ b/Engine/Plugins/Animation/ControlRig/Source/ControlRigEditor/Private/Sequencer/ControlRigParameterTrackEditor.cpp
@@ -2780,9 +2780,8 @@ void FControlRigParameterTrackEditor::HandleConstraintRemoved(IMovieSceneConstra
 								return;
 							}
 
-							const FName ConstraintName = Constraint->GetFName();
-							const FConstraintAndActiveChannel* ConstraintChannel = InSection->GetConstraintChannel(ConstraintName);
-							if (!ConstraintChannel)
+							const FConstraintAndActiveChannel* ConstraintChannel = InSection->GetConstraintChannel(Constraint->GetFName());
+							if (!ConstraintChannel || ConstraintChannel->Constraint != Constraint)
 							{
 								return;
 							}
@@ -2797,7 +2796,7 @@ void FControlRigParameterTrackEditor::HandleConstraintRemoved(IMovieSceneConstra
 									Section);
 							}
 
-							InSection->RemoveConstraintChannel(ConstraintName);
+							InSection->RemoveConstraintChannel(Constraint);
 						}
 						break;
 					case EConstraintsManagerNotifyType::ManagerUpdated:
@@ -4668,7 +4667,7 @@ bool FControlRigParameterSection::RequestDeleteCategory(const TArray<FName>& Cat
 		{
 			if (ParameterSection->TryModify())
 			{
-				ParameterSection->RemoveConstraintChannel(Channel);
+				ParameterSection->RemoveConstraintChannel(ConstraintChannel->Constraint.Get());
 				SequencerPtr->NotifyMovieSceneDataChanged(EMovieSceneDataChangeType::MovieSceneStructureItemsChanged);
 				return true;
 			}
diff --git a/Engine/Source/Editor/MovieSceneTools/Private/TrackEditors/TransformTrackEditor.cpp b/Engine/Source/Editor/MovieSceneTools/Private/TrackEditors/TransformTrackEditor.cpp
--- a/Engine/Source/Editor/MovieSceneTools/Private/TrackEditors/TransformTrackEditor.cpp
+++ b/Engine/Source/Editor/MovieSceneTools/Private/TrackEditors/TransformTrackEditor.cpp
@@ -1540,9 +1540,8 @@ void F3DTransformTrackEditor::HandleConstraintRemoved(IMovieSceneConstrainedSect
 									return;
 								}
 
-								const FName ConstraintName = Constraint->GetFName();
-								const FConstraintAndActiveChannel* ConstraintChannel = InSection->GetConstraintChannel(ConstraintName);
-								if (!ConstraintChannel)
+								const FConstraintAndActiveChannel* ConstraintChannel = InSection->GetConstraintChannel(Constraint->GetFName());
+								if (!ConstraintChannel || ConstraintChannel->Constraint != Constraint)
 								{
 									return;
 								}
@@ -1557,7 +1556,7 @@ void F3DTransformTrackEditor::HandleConstraintRemoved(IMovieSceneConstrainedSect
 										Section);
 								}
 
-								InSection->RemoveConstraintChannel(ConstraintName);
+								InSection->RemoveConstraintChannel(Constraint);
 							}
 							break;
 						case EConstraintsManagerNotifyType::ManagerUpdated:
diff --git a/Engine/Source/Runtime/Experimental/Animation/Constraints/Private/ConstraintsManager.cpp b/Engine/Source/Runtime/Experimental/Animation/Constraints/Private/ConstraintsManager.cpp
--- a/Engine/Source/Runtime/Experimental/Animation/Constraints/Private/ConstraintsManager.cpp
+++ b/Engine/Source/Runtime/Experimental/Animation/Constraints/Private/ConstraintsManager.cpp
@@ -466,6 +466,7 @@ UTickableConstraint* FConstraintsManagerController::AddConstraintFromCopy(UTicka
 	}
 
 	UTickableConstraint* OurCopy = CopyOfConstraint->Duplicate(Manager);
+	OurCopy->SetActive(false); //always set false
 	if (AddConstraint(OurCopy))
 	{
 		return OurCopy;
@@ -488,6 +489,25 @@ int32 FConstraintsManagerController::GetConstraintIndex(const FName& InConstrain
 	} );
 }
 	
+bool FConstraintsManagerController::RemoveConstraint(UTickableConstraint* InConstraint, bool bDoNotCompensate) const
+{
+	const UConstraintsManager* Manager = FindManager();
+	if (!Manager)
+	{
+		return false;
+	}
+
+	int32 Index = Manager->Constraints.IndexOfByPredicate([InConstraint](const TObjectPtr<UTickableConstraint>& Constraint)
+		{
+			return 	(Constraint && Constraint == InConstraint);
+		});
+	if (Index == INDEX_NONE)
+	{
+		return false;
+	}
+	return RemoveConstraint(Index, bDoNotCompensate);
+}
+
 bool  FConstraintsManagerController::RemoveConstraint(const FName& InConstraintName, bool bDoNotCompensate) const
 {
 	const int32 Index = GetConstraintIndex(InConstraintName);
diff --git a/Engine/Source/Runtime/Experimental/Animation/Constraints/Private/TransformConstraint.cpp b/Engine/Source/Runtime/Experimental/Animation/Constraints/Private/TransformConstraint.cpp
--- a/Engine/Source/Runtime/Experimental/Animation/Constraints/Private/TransformConstraint.cpp
+++ b/Engine/Source/Runtime/Experimental/Animation/Constraints/Private/TransformConstraint.cpp
@@ -226,12 +226,12 @@ UTickableConstraint* UTickableTransformConstraint::Duplicate(UObject* NewOuter)
 	UTickableTransformConstraint* Dup = DuplicateObject<UTickableTransformConstraint>(this, NewOuter);
 	if (ChildTRSHandle)
 	{
-		UTransformableHandle* HandleCopy = ChildTRSHandle->Duplicate(NewOuter);
+		UTransformableHandle* HandleCopy = ChildTRSHandle->Duplicate(Dup);
 		Dup->ChildTRSHandle = HandleCopy;
 	}
 	if (ParentTRSHandle)
 	{
-		UTransformableHandle* HandleCopy = ParentTRSHandle->Duplicate(NewOuter);
+		UTransformableHandle* HandleCopy = ParentTRSHandle->Duplicate(Dup);
 		Dup->ParentTRSHandle = HandleCopy;
 	}
 	Dup->SetupDependencies();
diff --git a/Engine/Source/Runtime/Experimental/Animation/Constraints/Private/TransformableHandle.cpp b/Engine/Source/Runtime/Experimental/Animation/Constraints/Private/TransformableHandle.cpp
--- a/Engine/Source/Runtime/Experimental/Animation/Constraints/Private/TransformableHandle.cpp
+++ b/Engine/Source/Runtime/Experimental/Animation/Constraints/Private/TransformableHandle.cpp
@@ -543,7 +543,7 @@ void UTransformableComponentHandle::ResolveBoundObjects(FMovieSceneSequenceID Lo
 
 UTransformableHandle* UTransformableComponentHandle::Duplicate(UObject* NewOuter) const
 {
-	UTransformableComponentHandle* HandleCopy = DuplicateObject<UTransformableComponentHandle>(this, NewOuter, GetFName());
+	UTransformableComponentHandle* HandleCopy = DuplicateObject<UTransformableComponentHandle>(this, NewOuter);
 	HandleCopy->Component = Component;
 	HandleCopy->SocketName = SocketName;
 	return HandleCopy;
diff --git a/Engine/Source/Runtime/Experimental/Animation/Constraints/Public/ConstraintsManager.h b/Engine/Source/Runtime/Experimental/Animation/Constraints/Public/ConstraintsManager.h
--- a/Engine/Source/Runtime/Experimental/Animation/Constraints/Public/ConstraintsManager.h
+++ b/Engine/Source/Runtime/Experimental/Animation/Constraints/Public/ConstraintsManager.h
@@ -231,6 +231,9 @@ public:
 	/** Remove the constraint at the given index. */
 	bool RemoveConstraint(const int32 InConstraintIndex, bool bDoNotCompensate = false) const;
 
+	/** Remove constraint by ptr*/
+	bool RemoveConstraint(UTickableConstraint* InConstraint, bool bDoNotCompensate = false) const;
+
 	/** Returns the constraint based on it's name within the manager's constraints array. */
 	UTickableConstraint* GetConstraint(const FName& InConstraintName) const;
 
diff --git a/Engine/Source/Runtime/MovieSceneTracks/Private/Sections/MovieScene3DTransformSection.cpp b/Engine/Source/Runtime/MovieSceneTracks/Private/Sections/MovieScene3DTransformSection.cpp
--- a/Engine/Source/Runtime/MovieSceneTracks/Private/Sections/MovieScene3DTransformSection.cpp
+++ b/Engine/Source/Runtime/MovieSceneTracks/Private/Sections/MovieScene3DTransformSection.cpp
@@ -402,10 +402,7 @@ void UMovieScene3DTransformSection::OnBindingIDsUpdated(const TMap<UE::MovieScen
 		{
 			if (UTickableTransformConstraint* TransformConstraint = Cast< UTickableTransformConstraint>(ConstraintChannel.Constraint.Get()))
 			{
-				if (TransformConstraint->ChildTRSHandle)
-				{
-					TransformConstraint->ChildTRSHandle->OnBindingIDsUpdated(OldFixedToNewFixedMap, LocalSequenceID, Hierarchy, Player);
-				}
+				//Don't do child's we do that in the system, needed for duplication
 				if (TransformConstraint->ParentTRSHandle)
 				{
 					TransformConstraint->ParentTRSHandle->OnBindingIDsUpdated(OldFixedToNewFixedMap, LocalSequenceID, Hierarchy, Player);
@@ -413,10 +410,7 @@ void UMovieScene3DTransformSection::OnBindingIDsUpdated(const TMap<UE::MovieScen
 			}
 			if (UTickableTransformConstraint* SpawnCopy = Cast< UTickableTransformConstraint>(ConstraintChannel.ConstraintCopyToSpawn))
 			{
-				if (SpawnCopy->ChildTRSHandle)
-				{
-					SpawnCopy->ChildTRSHandle->OnBindingIDsUpdated(OldFixedToNewFixedMap, LocalSequenceID, Hierarchy, Player);
-				}
+				//Don't do child's we do that in the system, needed for duplication
 				if (SpawnCopy->ParentTRSHandle)
 				{
 					SpawnCopy->ParentTRSHandle->OnBindingIDsUpdated(OldFixedToNewFixedMap, LocalSequenceID, Hierarchy, Player);
@@ -636,7 +630,6 @@ void UMovieScene3DTransformSection::ImportConstraintEntity(UMovieSceneEntitySyst
 
 			FConstraintComponentData ComponentData;
 			ComponentData.ConstraintName = ConstraintName;
-			ComponentData.ConstraintAndActiveChannel = &(Constraints->ConstraintsChannels[ConstraintIndex]);
 			ComponentData.Section = this;
 			OutImportedEntity->AddBuilder(
 				FEntityBuilder()
@@ -1000,7 +993,8 @@ FConstraintAndActiveChannel* UMovieScene3DTransformSection::GetConstraintChannel
 	{
 		const int32 Index = Constraints->ConstraintsChannels.IndexOfByPredicate([InConstraintName](const FConstraintAndActiveChannel& InChannel)
 			{
-				return InChannel.Constraint.IsValid() ? InChannel.Constraint->GetFName() == InConstraintName : false;
+				return InChannel.Constraint.IsValid() ? InChannel.Constraint->GetFName() == InConstraintName : (InChannel.ConstraintCopyToSpawn ?
+				InChannel.ConstraintCopyToSpawn->GetFName() == InConstraintName : false);
 			});
 		return (Index != INDEX_NONE) ? &(Constraints->ConstraintsChannels[Index]) : nullptr;
 	}
@@ -1029,6 +1023,25 @@ bool UMovieScene3DTransformSection::Modify(bool bAlwaysMarkDirty)
 	
 	return bModified;
 }
+
+void UMovieScene3DTransformSection::PostDuplicate(bool bDuplicateForPIE)
+{
+	Super::PostDuplicate(bDuplicateForPIE);
+	//if not duplicating for PIE clear out the soft object ptrs to the constraints we dont' want to hold pointers to the old constraints.
+	//also make the copy to spawn so we can duplicate it if it doesn't exist in the level
+	//to the old ones but make new duplicates
+	if (bDuplicateForPIE == false && Constraints)
+	{
+		for (FConstraintAndActiveChannel& ConstraintChannel : Constraints->ConstraintsChannels)
+		{
+			if (ConstraintChannel.Constraint.IsValid() && !ConstraintChannel.ConstraintCopyToSpawn)
+			{
+				ConstraintChannel.ConstraintCopyToSpawn = ConstraintChannel.Constraint->Duplicate(this);
+			}
+			ConstraintChannel.Constraint.Reset();
+		}
+	}
+}
 #endif
 void UMovieScene3DTransformSection::AddConstraintChannel(UTickableConstraint* InConstraint)
 {
@@ -1075,13 +1088,17 @@ void UMovieScene3DTransformSection::ReplaceConstraint(const FName InConstraintNa
 	}
 }
 
-void UMovieScene3DTransformSection::RemoveConstraintChannel(const FName& InConstraintName)
+void UMovieScene3DTransformSection::RemoveConstraintChannel(const UTickableConstraint* InConstraint)
 {
+	if (bDoNotRemoveChannel == true)
+	{
+		return;
+	}
 	if (Constraints)
 	{
-		const int32 Index = Constraints->ConstraintsChannels.IndexOfByPredicate([InConstraintName](const FConstraintAndActiveChannel& InChannel)
+		const int32 Index = Constraints->ConstraintsChannels.IndexOfByPredicate([InConstraint](const FConstraintAndActiveChannel& InChannel)
 			{
-				return InChannel.Constraint.IsValid() ? InChannel.Constraint->GetFName() == InConstraintName : false;
+				return InChannel.Constraint.IsValid() ? InChannel.Constraint == InConstraint : false;
 			});
 
 		if (Constraints->ConstraintsChannels.IsValidIndex(Index))
@@ -1143,6 +1160,7 @@ void UMovieScene3DTransformSection::OnChannelOverridesChanged()
 	ChannelProxy = nullptr;
 }
 
+
 #if WITH_EDITOR
 
 void UMovieScene3DTransformSection::PostPaste()
diff --git a/Engine/Source/Runtime/MovieSceneTracks/Private/Sections/MovieSceneConstrainedSection.cpp b/Engine/Source/Runtime/MovieSceneTracks/Private/Sections/MovieSceneConstrainedSection.cpp
--- a/Engine/Source/Runtime/MovieSceneTracks/Private/Sections/MovieSceneConstrainedSection.cpp
+++ b/Engine/Source/Runtime/MovieSceneTracks/Private/Sections/MovieSceneConstrainedSection.cpp
@@ -5,7 +5,14 @@
 #include UE_INLINE_GENERATED_CPP_BY_NAME(MovieSceneConstrainedSection)
 
 
-IMovieSceneConstrainedSection::IMovieSceneConstrainedSection()
+IMovieSceneConstrainedSection::IMovieSceneConstrainedSection() : bDoNotRemoveChannel(false)
 {
 }
 
+
+void IMovieSceneConstrainedSection::SetDoNoRemoveChannel(bool bInDoNotRemoveChannel)
+{
+	bDoNotRemoveChannel = bInDoNotRemoveChannel;
+}
+
+
diff --git a/Engine/Source/Runtime/MovieSceneTracks/Private/Systems/MovieSceneConstraintSystem.cpp b/Engine/Source/Runtime/MovieSceneTracks/Private/Systems/MovieSceneConstraintSystem.cpp
--- a/Engine/Source/Runtime/MovieSceneTracks/Private/Systems/MovieSceneConstraintSystem.cpp
+++ b/Engine/Source/Runtime/MovieSceneTracks/Private/Systems/MovieSceneConstraintSystem.cpp
@@ -16,41 +16,131 @@
 #include "Systems/MovieSceneComponentTransformSystem.h"
 #include "TransformConstraint.h"
 #include "TransformableHandle.h"
+#include "Sections/MovieSceneConstrainedSection.h"
 
 #include UE_INLINE_GENERATED_CPP_BY_NAME(MovieSceneConstraintSystem)
 
 class UTickableConstraint;
 
+
+UTickableConstraint* CreateConstraintIfNeeded(FConstraintsManagerController& Controller, FConstraintAndActiveChannel* ConstraintAndActiveChannel,const  FConstraintComponentData& ConstraintChannel)
+{
+	if (ConstraintAndActiveChannel->Constraint.IsPending())
+	{
+		ConstraintAndActiveChannel->Constraint.LoadSynchronous();
+	}
+	UTickableConstraint* Constraint = ConstraintAndActiveChannel->Constraint.Get();
+	if (!Constraint) //if constraint doesn't exist it probably got unspawned so recreate it and add it
+	{
+		Constraint = Controller.AddConstraintFromCopy(ConstraintAndActiveChannel->ConstraintCopyToSpawn);
+		ConstraintChannel.Section->ReplaceConstraint(ConstraintChannel.ConstraintName, Constraint);
+	}
+	else // it's possible that we have it but it's not in the manager, due to manager not being saved with it (due to spawning or undo/redo).
+	{
+		const TArray< TObjectPtr<UTickableConstraint>>& ConstraintsArray = Controller.GetConstraintsArray();
+		if (ConstraintsArray.Find(Constraint) == INDEX_NONE)
+		{
+			Controller.AddConstraint(Constraint);
+		}
+	}
+	return Constraint;
+}
+
+
 namespace UE::MovieScene
 {
 
 struct FPreAnimatedConstraint
 {
-	TWeakObjectPtr<UTickableConstraint> WeakConstraint;
+	FName ConstraintName;
+	bool bDeleteIt;
 	bool bPreviouslyEnabled;
+	TWeakObjectPtr<USceneComponent> SceneComponent;
+	TWeakObjectPtr<UMovieScene3DTransformSection> Section;
 };
 
 struct FPreAnimatedConstraintTraits : FBoundObjectPreAnimatedStateTraits
 {
-	using KeyType = TTuple<FObjectKey, FName>;
+	/** Key type that stores pre-animated state associated with the object and its constraint name */
+	struct FKeyType
+	{
+		FObjectKey Object;
+		FName ConstraintName;
+
+		/** Constructor that takes a BoundObject and ConstraintChannel component */
+		FKeyType(UObject* InObject, const FConstraintComponentData& ComponentData)
+			: Object(InObject)
+			, ConstraintName(ComponentData.ConstraintName)
+		{}
+
+		/** Hashing and equality required for storage within a map */
+		friend uint32 GetTypeHash(const FKeyType& InKey)
+		{
+			return HashCombine(GetTypeHash(InKey.Object), GetTypeHash(InKey.ConstraintName));
+		}
+		friend bool operator==(const FKeyType& A, const FKeyType& B)
+		{
+			return A.Object == B.Object && A.ConstraintName == B.ConstraintName;
+		}
+	};
+
+	using KeyType = FKeyType;
 	using StorageType = FPreAnimatedConstraint;
 
-	static FPreAnimatedConstraint CachePreAnimatedValue(UObject* InBoundObject, const FName& ConstraintName)
+	static FPreAnimatedConstraint CachePreAnimatedValue(UObject* InBoundObject, const FConstraintComponentData& ConstraintData)
 	{
 		USceneComponent* SceneComponent = CastChecked<USceneComponent>(InBoundObject);
 		FConstraintsManagerController& Controller = FConstraintsManagerController::Get(SceneComponent->GetWorld());
-
-		// @todo: I don't know how to store pre-animated values - do we need to enable/disable constraints or add/remove them?
-		//        Maybe we just need to also cache some other state other than Constraint->Active
-		UTickableConstraint* Constraint = Controller.GetConstraint(ConstraintName);
-		return FPreAnimatedConstraint{ Constraint, Constraint ? Constraint->Active : false };
+		if (FConstraintAndActiveChannel* ConstraintAndActiveChannel = ConstraintData.Section->GetConstraintChannel(ConstraintData.ConstraintName))
+		{
+			if (ConstraintAndActiveChannel->Constraint.IsPending())
+			{
+				ConstraintAndActiveChannel->Constraint.LoadSynchronous();
+			}
+			UTickableConstraint* Constraint = ConstraintAndActiveChannel->Constraint.Get();
+			const bool bIsActive = Constraint ? Constraint->Active : false;
+			const FName Name = Constraint ? Constraint->GetFName() : ConstraintData.ConstraintName;
+			//also need to create it
+			CreateConstraintIfNeeded(Controller, ConstraintAndActiveChannel, ConstraintData);
+			return FPreAnimatedConstraint{ Name, (Constraint == nullptr), bIsActive, SceneComponent, ConstraintData.Section };
+		}
+		return FPreAnimatedConstraint{ ConstraintData.ConstraintName, false,false, SceneComponent, ConstraintData.Section };
 	}
 
-	static void RestorePreAnimatedValue(const TTuple<FObjectKey, FName>& InKey, const FPreAnimatedConstraint& OldValue, const FRestoreStateParams& Params)
+	static void RestorePreAnimatedValue(const FKeyType& InKey, const FPreAnimatedConstraint& OldValue, const FRestoreStateParams& Params)
 	{
-		if (UTickableConstraint* Constraint = OldValue.WeakConstraint.Get())
+		if (USceneComponent* SceneComponent = OldValue.SceneComponent.Get())
 		{
-			Constraint->SetActive(OldValue.bPreviouslyEnabled);
+			if (FConstraintAndActiveChannel* ConstraintAndActiveChannel = OldValue.Section->GetConstraintChannel(OldValue.ConstraintName))
+			{				
+				if (ConstraintAndActiveChannel->Constraint.IsPending())
+				{
+					ConstraintAndActiveChannel->Constraint.LoadSynchronous();
+				}
+				UTickableConstraint* Constraint = ConstraintAndActiveChannel->Constraint.Get();
+				if (Constraint)
+				{
+					Constraint->SetActive(OldValue.bPreviouslyEnabled);
+
+					//we don't delete it since for some reason on first creation we get a restore state and so it immediately get's deleted
+					if (OldValue.bDeleteIt)
+					{
+						UMovieScene3DTransformSection* ConstrainedSection = OldValue.Section.Get();
+						if (ConstrainedSection)
+						{
+							ConstrainedSection->SetDoNoRemoveChannel(true);
+						}
+						FConstraintsManagerController& Controller = FConstraintsManagerController::Get(SceneComponent->GetWorld());
+						Controller.RemoveConstraint(Constraint, true);
+
+						if (ConstrainedSection)
+						{
+							ConstrainedSection->SetDoNoRemoveChannel(false);
+						}
+					}
+				
+				}
+			}
 		}
 	}
 };
@@ -97,10 +187,12 @@ void UMovieSceneConstraintSystem::OnRun(FSystemTaskPrerequisites& InPrerequisite
 	if (CurrentPhase == ESystemPhase::Instantiation)
 	{
 		// Save pre-animated state
-		//TSharedPtr<FPreAnimatedConstraintStorage> PreAnimatedStorage = Linker->PreAnimatedState.GetOrCreateStorage<FPreAnimatedConstraintStorage>();
+		TSharedPtr<FPreAnimatedConstraintStorage> PreAnimatedStorage = Linker->PreAnimatedState.GetOrCreateStorage<FPreAnimatedConstraintStorage>();
+
+		FMovieSceneTracksComponentTypes* TracksComponents = FMovieSceneTracksComponentTypes::Get();
+		FBuiltInComponentTypes* BuiltInComponents = FBuiltInComponentTypes::Get();
+		PreAnimatedStorage->BeginTrackingAndCachePreAnimatedValues(Linker, BuiltInComponents->BoundObject, TracksComponents->ConstraintChannel);
 
-		//andrew, had to remove this, last parameter not the correct variadic type
-		//PreAnimatedStorage->BeginTrackingAndCachePreAnimatedValues(Linker, BuiltInComponents->BoundObject, TracksComponents->ConstraintName);	
 	}
 	else if (CurrentPhase == ESystemPhase::Evaluation)
 	{
@@ -115,40 +207,27 @@ void UMovieSceneConstraintSystem::OnRun(FSystemTaskPrerequisites& InPrerequisite
 			void ForEachEntity(UObject* BoundObject, UE::MovieScene::FInstanceHandle InstanceHandle, const FConstraintComponentData& ConstraintChannel, FFrameTime FrameTime)
 			{
 				const FSequenceInstance& TargetInstance = InstanceRegistry->GetInstance(InstanceHandle);
-
-				if (ConstraintChannel.ConstraintAndActiveChannel->Constraint.IsPending())
+				if (FConstraintAndActiveChannel* ConstraintAndActiveChannel = ConstraintChannel.Section->GetConstraintChannel(ConstraintChannel.ConstraintName))
 				{
-					ConstraintChannel.ConstraintAndActiveChannel->Constraint.LoadSynchronous();
-				}
-				UTickableConstraint* Constraint = ConstraintChannel.ConstraintAndActiveChannel->Constraint.Get();
-				if (!Constraint) //if constraint doesn't exist it probably got unspawned so recreate it and add it
-				{
-					UTickableConstraint* NewOne = Controller->AddConstraintFromCopy(ConstraintChannel.ConstraintAndActiveChannel->ConstraintCopyToSpawn);
-					Constraint = Controller->GetConstraint(ConstraintChannel.ConstraintName);
-					ConstraintChannel.Section->ReplaceConstraint(ConstraintChannel.ConstraintName, Constraint);
-				}
-				else // it's possible that we have it but it's not in the manager, due to manager not being saved with it (due to spawning or undo/redo).
-				{
-					const TArray< TObjectPtr<UTickableConstraint>>& ConstraintsArray = Controller->GetConstraintsArray();
-					if(ConstraintsArray.Find(Constraint) == INDEX_NONE)
-					{
-						Controller->AddConstraint(Constraint);
-					}
-				}
-				if (Constraint)
-				{
-					bool Result = false;
-					ConstraintChannel.ConstraintAndActiveChannel->ActiveChannel.Evaluate(FrameTime, Result);
-					Constraint->SetActive(Result);
-					Constraint->ResolveBoundObjects(TargetInstance.GetSequenceID(), *TargetInstance.GetPlayer());
-					if (UTickableTransformConstraint* TransformConstraint = Cast< UTickableTransformConstraint>(Constraint))
+					UTickableConstraint* Constraint = CreateConstraintIfNeeded(*Controller, ConstraintAndActiveChannel, ConstraintChannel);
+
+					if (Constraint)
 					{
-						if (UTransformableComponentHandle* ComponentHandle = Cast<UTransformableComponentHandle>(TransformConstraint->ChildTRSHandle))
+						bool Result = false;
+						ConstraintAndActiveChannel->ActiveChannel.Evaluate(FrameTime, Result);
+						Constraint->SetActive(Result);
+						Constraint->ResolveBoundObjects(TargetInstance.GetSequenceID(), *TargetInstance.GetPlayer());
+						if (UTickableTransformConstraint* TransformConstraint = Cast< UTickableTransformConstraint>(Constraint))
 						{
-							FUpdateHandleForConstraint UpdateHandle;
-							UpdateHandle.Constraint = TransformConstraint;
-							UpdateHandle.TransformHandle = ComponentHandle;
-							System->DynamicOffsets.Add(UpdateHandle);
+							if (UTransformableComponentHandle* ComponentHandle = Cast<UTransformableComponentHandle>(TransformConstraint->ChildTRSHandle))
+							{
+								//bound component may change so need to update constraint
+								ComponentHandle->Component = Cast<USceneComponent>(BoundObject);
+								FUpdateHandleForConstraint UpdateHandle;
+								UpdateHandle.Constraint = TransformConstraint;
+								UpdateHandle.TransformHandle = ComponentHandle;
+								System->DynamicOffsets.Add(UpdateHandle);
+							}
 						}
 					}
 				}
@@ -178,7 +257,4 @@ void UMovieSceneConstraintSystem::OnRun(FSystemTaskPrerequisites& InPrerequisite
 		}
 		DynamicOffsets.Reset();
 	}
-}
-
-
-
+}
\ No newline at end of file
diff --git a/Engine/Source/Runtime/MovieSceneTracks/Public/MovieSceneTracksComponentTypes.h b/Engine/Source/Runtime/MovieSceneTracks/Public/MovieSceneTracksComponentTypes.h
--- a/Engine/Source/Runtime/MovieSceneTracks/Public/MovieSceneTracksComponentTypes.h
+++ b/Engine/Source/Runtime/MovieSceneTracks/Public/MovieSceneTracksComponentTypes.h
@@ -77,7 +77,6 @@ struct FConstraintComponentData
 
 	UPROPERTY()
 	FName ConstraintName;
-	FConstraintAndActiveChannel* ConstraintAndActiveChannel;
 	UMovieScene3DTransformSection* Section;
 };
 
diff --git a/Engine/Source/Runtime/MovieSceneTracks/Public/Sections/MovieScene3DTransformSection.h b/Engine/Source/Runtime/MovieSceneTracks/Public/Sections/MovieScene3DTransformSection.h
--- a/Engine/Source/Runtime/MovieSceneTracks/Public/Sections/MovieScene3DTransformSection.h
+++ b/Engine/Source/Runtime/MovieSceneTracks/Public/Sections/MovieScene3DTransformSection.h
@@ -225,7 +225,7 @@ public:
 #if WITH_EDITOR
 	/* From UObject*/
 	virtual bool Modify(bool bAlwaysMarkDirty = true) override;
-
+	virtual void PostDuplicate(bool bDuplicateForPIE) override;
 #endif
 	/* From UMovieSection*/
 	
@@ -344,7 +344,7 @@ public:
 	/*
 	*  Remove Constraint channel
 	*/
-	virtual void RemoveConstraintChannel(const FName& InConstraintName) override;
+	virtual void RemoveConstraintChannel(const UTickableConstraint* InConstraint) override;
 
 	/*
 	*  Get The channels by value
@@ -354,7 +354,7 @@ public:
 	/*
 	*  Replace the constraint with the specified name with the new one
 	*/
-	virtual void ReplaceConstraint(const FName InName, UTickableConstraint* InConstraint)  override;
+	virtual void ReplaceConstraint(const FName InConstraintName, UTickableConstraint* InConstraint)  override;
 
 	/*
 	* Clear proxy if changed
diff --git a/Engine/Source/Runtime/MovieSceneTracks/Public/Sections/MovieSceneConstrainedSection.h b/Engine/Source/Runtime/MovieSceneTracks/Public/Sections/MovieSceneConstrainedSection.h
--- a/Engine/Source/Runtime/MovieSceneTracks/Public/Sections/MovieSceneConstrainedSection.h
+++ b/Engine/Source/Runtime/MovieSceneTracks/Public/Sections/MovieSceneConstrainedSection.h
@@ -55,7 +55,7 @@ public:
 	/*
 	*  Remove Constraint channel
 	*/
-	virtual void RemoveConstraintChannel(const FName& InConstraintName) = 0;
+	virtual void RemoveConstraintChannel(const UTickableConstraint* InConstraint) = 0;
 
 	/*
 	*  Get The channels
@@ -65,7 +65,7 @@ public:
 	/*
 	*  Replace the constraint with the specified name with the new one
 	*/
-	virtual void ReplaceConstraint(const FName InName, UTickableConstraint* InConstraint) = 0;
+	virtual void ReplaceConstraint(const FName InConstraintName, UTickableConstraint* InConstraint) = 0;
 
 	/*
 	* Added Delegate
@@ -77,6 +77,11 @@ public:
 	 */
 	 virtual void OnConstraintsChanged() {};
 
+	/*
+	*  SetToNotRemoveChannel when we are told a constraint is removed, we need this sometimes sincet his will be destructive
+	*/
+	 void SetDoNoRemoveChannel(bool bInDoNotRemoveChannel);
+
 	/*
 	*  Removed delegate that get's added by the track editor
 	*/
@@ -86,6 +91,6 @@ public:
 protected:
 
 	FConstraintChannelAddedEvent OnConstraintChannelAdded;
-
+	bool bDoNotRemoveChannel = false;
 
 };
