Sequencer: Expanded blender support to allow for choosing blenders for any given combination of entities

  - Blend Channel IDs will now always be accompanied by a component tag that uniquely identifies the blender system that it applies to. This allows blenders to be sure they are only reading channels that are intended for the specific system.
  - Blenders now have a selection priority which is used for choosing a blender when multiple are relevant to a given channel.

diff --git a/Engine/Source/Runtime/MovieScene/Private/EntitySystem/MovieSceneBlenderSystem.cpp b/Engine/Source/Runtime/MovieScene/Private/EntitySystem/MovieSceneBlenderSystem.cpp
--- a/Engine/Source/Runtime/MovieScene/Private/EntitySystem/MovieSceneBlenderSystem.cpp
+++ b/Engine/Source/Runtime/MovieScene/Private/EntitySystem/MovieSceneBlenderSystem.cpp
@@ -30,16 +30,28 @@ UMovieSceneBlenderSystem::UMovieSceneBlenderSystem(const FObjectInitializer& Obj
 	using namespace UE::MovieScene;
 
 	SystemCategories = UE::MovieScene::EEntitySystemCategory::BlenderSystems;
+	SelectionPriority = DefaultPriority;
 
 	if (HasAnyFlags(RF_ClassDefaultObject))
 	{
-		// You can only ever register blender systems, never unregister them.
-		SystemID = FMovieSceneBlenderSystemID(GBlenderSystemRegistry.Num());
-		GBlenderSystemRegistry.Add(SystemID, GetClass());
+		// No need to register these for the base class
+		if (GetClass() != UMovieSceneBlenderSystem::StaticClass())
+		{
+			// You can only ever register blender systems, never unregister them.
+			SystemID = FMovieSceneBlenderSystemID(GBlenderSystemRegistry.Num());
+			GBlenderSystemRegistry.Add(SystemID, GetClass());
+
+			FComponentRegistry* ComponentRegistry = UMovieSceneEntitySystemLinker::GetComponents();
+			BlenderTypeTag = ComponentRegistry->NewTag(*GetName());
+		}
 	}
 	else 
 	{
-		SystemID = GetClass()->GetDefaultObject<UMovieSceneBlenderSystem>()->SystemID;
+		UMovieSceneBlenderSystem* CDO = GetClass()->GetDefaultObject<UMovieSceneBlenderSystem>();
+
+		SystemID = CDO->SystemID;
+		BlenderTypeTag = CDO->BlenderTypeTag;
+
 		checkf(SystemID.IsValid(), TEXT("Blender system wasn't registered correctly on init!"));
 	}
 }
diff --git a/Engine/Source/Runtime/MovieScene/Public/EntitySystem/MovieSceneBlenderSystem.h b/Engine/Source/Runtime/MovieScene/Public/EntitySystem/MovieSceneBlenderSystem.h
--- a/Engine/Source/Runtime/MovieScene/Public/EntitySystem/MovieSceneBlenderSystem.h
+++ b/Engine/Source/Runtime/MovieScene/Public/EntitySystem/MovieSceneBlenderSystem.h
@@ -47,6 +47,9 @@ public:
 
 	GENERATED_BODY()
 
+	/** Default blender system priority. Higher priorities win, negative priorities are supported. */
+	static constexpr int32 DefaultPriority = 0;
+
 	/**
 	 * Get a blender system class, given a blender system ID.
 	 */
@@ -79,6 +82,21 @@ public:
 	 */
 	void ReleaseBlendChannel(FMovieSceneBlendChannelID BlendChannelID);
 
+	/**
+	 * Retrieve the selection priority to use when choosing a blender system. Highest wins.
+	 */
+	int32 GetSelectionPriority() const
+	{
+		return SelectionPriority;
+	}
+
+	/**
+	 * Retrieve the tag that should accompany any blend input and output components to be used by this blender
+	 */
+	UE::MovieScene::FComponentTypeID GetBlenderTypeTag() const
+	{
+		return BlenderTypeTag;
+	}
 
 protected:
 
@@ -93,6 +111,12 @@ protected:
 
 private:
 
+	/** Priority to use when selecting a blender system for a given set of entities. Highest priority will be used. */
+	int32 SelectionPriority;
+
 	/** Cached blender system ID */
 	FMovieSceneBlenderSystemID SystemID;
+
+	/** Type that is unique to this blender system type that distinquishes its BlendInput and Output components */
+	UE::MovieScene::FComponentTypeID BlenderTypeTag;
 };
diff --git a/Engine/Source/Runtime/MovieSceneTracks/Private/Systems/MovieScenePiecewiseDoubleBlenderSystem.cpp b/Engine/Source/Runtime/MovieSceneTracks/Private/Systems/MovieScenePiecewiseDoubleBlenderSystem.cpp
--- a/Engine/Source/Runtime/MovieSceneTracks/Private/Systems/MovieScenePiecewiseDoubleBlenderSystem.cpp
+++ b/Engine/Source/Runtime/MovieSceneTracks/Private/Systems/MovieScenePiecewiseDoubleBlenderSystem.cpp
@@ -475,7 +475,7 @@ void UMovieScenePiecewiseDoubleBlenderSystem::OnRun(FSystemTaskPrerequisites& In
 		FGraphEventRef Task = FEntityTaskBuilder()
 		.Read(BuiltInComponents->BlendChannelInput)
 		.ReadOptional(BuiltInComponents->WeightAndEasingResult)
-		.FilterAll({ BuiltInComponents->Tags.AbsoluteBlend })
+		.FilterAll({ BuiltInComponents->Tags.AbsoluteBlend, GetBlenderTypeTag() })
 		.FilterAny(BlendedResultMask)
 		.FilterNone({ BuiltInComponents->Tags.Ignored })
 		.SetStat(BlendValuesStatId)
@@ -492,7 +492,7 @@ void UMovieScenePiecewiseDoubleBlenderSystem::OnRun(FSystemTaskPrerequisites& In
 		FGraphEventRef Task = FEntityTaskBuilder()
 		.Read(BuiltInComponents->BlendChannelInput)
 		.ReadOptional(BuiltInComponents->WeightAndEasingResult)
-		.FilterAll({ BuiltInComponents->Tags.RelativeBlend })
+		.FilterAll({ BuiltInComponents->Tags.RelativeBlend, GetBlenderTypeTag() })
 		.FilterAny(BlendedResultMask)
 		.FilterNone({ BuiltInComponents->Tags.Ignored })
 		.SetStat(BlendValuesStatId)
@@ -509,7 +509,7 @@ void UMovieScenePiecewiseDoubleBlenderSystem::OnRun(FSystemTaskPrerequisites& In
 		FGraphEventRef Task = FEntityTaskBuilder()
 		.Read(BuiltInComponents->BlendChannelInput)
 		.ReadOptional(BuiltInComponents->WeightAndEasingResult)
-		.FilterAll({ BuiltInComponents->Tags.AdditiveBlend })
+		.FilterAll({ BuiltInComponents->Tags.AdditiveBlend, GetBlenderTypeTag() })
 		.FilterAny(BlendedResultMask)
 		.FilterNone({ BuiltInComponents->Tags.Ignored })
 		.SetStat(BlendValuesStatId)
@@ -526,7 +526,7 @@ void UMovieScenePiecewiseDoubleBlenderSystem::OnRun(FSystemTaskPrerequisites& In
 		FGraphEventRef Task = FEntityTaskBuilder()
 		.Read(BuiltInComponents->BlendChannelInput)
 		.ReadOptional(BuiltInComponents->WeightAndEasingResult)
-		.FilterAll({ BuiltInComponents->Tags.AdditiveFromBaseBlend })
+		.FilterAll({ BuiltInComponents->Tags.AdditiveFromBaseBlend, GetBlenderTypeTag() })
 		.FilterAny(BlendedResultMask)
 		.FilterNone({ BuiltInComponents->Tags.Ignored })
 		.SetStat(BlendValuesStatId)
@@ -543,6 +543,7 @@ void UMovieScenePiecewiseDoubleBlenderSystem::OnRun(FSystemTaskPrerequisites& In
 	{
 		FEntityTaskBuilder()
 		.Read(BuiltInComponents->BlendChannelOutput)
+		.FilterAll({ GetBlenderTypeTag() })
 		.FilterAny(BlendedPropertyMask)
 		.SetStat(CombineBlendsStatId)
 		.template Dispatch_PerAllocation<FCombineBlendsForProperties>(&EntityManager, Prereqs, &Subsequents, CachedRelevantProperties, &AccumulationBuffers, FEntityAllocationWriteContext(EntityManager));
@@ -553,6 +554,7 @@ void UMovieScenePiecewiseDoubleBlenderSystem::OnRun(FSystemTaskPrerequisites& In
 		// Blend task that combines vanilla (non-property-based) components
 		FEntityTaskBuilder()
 		.Read(BuiltInComponents->BlendChannelOutput)
+		.FilterAll({ GetBlenderTypeTag() })
 		.FilterAny(BlendedResultMask)
 		.FilterNone(BlendedPropertyMask)
 		.SetStat(CombineBlendsStatId)
@@ -581,10 +583,10 @@ void UMovieScenePiecewiseDoubleBlenderSystem::ReinitializeAccumulationBuffers()
 	{
 		TComponentTypeID<double> Component = ResultComponents[Index];
 
-		const bool bHasAbsolutes         = EntityManager.Contains(FEntityComponentFilter().All({ Component, BuiltInComponents->BlendChannelInput, BuiltInComponents->Tags.AbsoluteBlend }));
-		const bool bHasRelatives         = EntityManager.Contains(FEntityComponentFilter().All({ Component, BuiltInComponents->BlendChannelInput, BuiltInComponents->Tags.RelativeBlend }));
-		const bool bHasAdditives         = EntityManager.Contains(FEntityComponentFilter().All({ Component, BuiltInComponents->BlendChannelInput, BuiltInComponents->Tags.AdditiveBlend }));
-		const bool bHasAdditivesFromBase = EntityManager.Contains(FEntityComponentFilter().All({ Component, BuiltInComponents->BlendChannelInput, BuiltInComponents->Tags.AdditiveFromBaseBlend }));
+		const bool bHasAbsolutes         = EntityManager.Contains(FEntityComponentFilter().All({ Component, GetBlenderTypeTag(), BuiltInComponents->BlendChannelInput, BuiltInComponents->Tags.AbsoluteBlend }));
+		const bool bHasRelatives         = EntityManager.Contains(FEntityComponentFilter().All({ Component, GetBlenderTypeTag(), BuiltInComponents->BlendChannelInput, BuiltInComponents->Tags.RelativeBlend }));
+		const bool bHasAdditives         = EntityManager.Contains(FEntityComponentFilter().All({ Component, GetBlenderTypeTag(), BuiltInComponents->BlendChannelInput, BuiltInComponents->Tags.AdditiveBlend }));
+		const bool bHasAdditivesFromBase = EntityManager.Contains(FEntityComponentFilter().All({ Component, GetBlenderTypeTag(), BuiltInComponents->BlendChannelInput, BuiltInComponents->Tags.AdditiveFromBaseBlend }));
 
 		if (!(bHasAbsolutes || bHasRelatives || bHasAdditives || bHasAdditivesFromBase))
 		{
@@ -639,7 +641,7 @@ void UMovieScenePiecewiseDoubleBlenderSystem::ReinitializeAccumulationBuffers()
 			AllPropertyTypes.Set(PropertyDefinition.PropertyType);
 
 			InclusionFilter.Reset();
-			InclusionFilter.All({ BuiltInComponents->BlendChannelOutput, PropertyDefinition.PropertyType });
+			InclusionFilter.All({ GetBlenderTypeTag(), BuiltInComponents->BlendChannelOutput, PropertyDefinition.PropertyType });
 			if (EntityManager.Contains(InclusionFilter))
 			{
 				CachedRelevantProperties.PadToNum(PropertyTypeIndex + 1, false);
@@ -650,7 +652,7 @@ void UMovieScenePiecewiseDoubleBlenderSystem::ReinitializeAccumulationBuffers()
 		}
 	}
 
-	bContainsNonPropertyBlends = EntityManager.Contains(FEntityComponentFilter().All({ BuiltInComponents->BlendChannelOutput }).None(AllPropertyTypes));
+	bContainsNonPropertyBlends = EntityManager.Contains(FEntityComponentFilter().All({ GetBlenderTypeTag(), BuiltInComponents->BlendChannelOutput }).None(AllPropertyTypes));
 }
 
 void UMovieScenePiecewiseDoubleBlenderSystem::ZeroAccumulationBuffers()
@@ -774,7 +776,7 @@ FGraphEventRef UMovieScenePiecewiseDoubleBlenderSystem::DispatchDecomposeTask(co
 			.Read(ResultComponentType)
 			.ReadOptional(BaseValueComponentType)
 			.ReadOptional(BuiltInComponents->WeightAndEasingResult)
-			.FilterAll({ Params.PropertyTag })
+			.FilterAll({ Params.PropertyTag, GetBlenderTypeTag() })
 			.template Dispatch_PerAllocation<FChannelResultTask>(&Linker->EntityManager, FSystemTaskPrerequisites(), nullptr, Params, Output);
 	}
 	else
@@ -785,7 +787,7 @@ FGraphEventRef UMovieScenePiecewiseDoubleBlenderSystem::DispatchDecomposeTask(co
 			.Read(ResultComponentType)
 			.ReadOptional(BaseValueComponentType)
 			.ReadOptional(BuiltInComponents->WeightAndEasingResult)
-			.FilterAll({ Params.PropertyTag })
+			.FilterAll({ Params.PropertyTag, GetBlenderTypeTag() })
 			.template Dispatch_PerAllocation<FChannelResultTask>(&Linker->EntityManager, FSystemTaskPrerequisites(), nullptr, Params, Output);
 	}
 }
diff --git a/Engine/Source/Runtime/MovieSceneTracks/Private/Systems/MovieScenePiecewiseIntegerBlenderSystem.cpp b/Engine/Source/Runtime/MovieSceneTracks/Private/Systems/MovieScenePiecewiseIntegerBlenderSystem.cpp
--- a/Engine/Source/Runtime/MovieSceneTracks/Private/Systems/MovieScenePiecewiseIntegerBlenderSystem.cpp
+++ b/Engine/Source/Runtime/MovieSceneTracks/Private/Systems/MovieScenePiecewiseIntegerBlenderSystem.cpp
@@ -279,7 +279,7 @@ void UMovieScenePiecewiseIntegerBlenderSystem::OnRun(FSystemTaskPrerequisites& I
 		.Read(BuiltInComponents->BlendChannelInput)
 		.Read(BuiltInComponents->IntegerResult)
 		.ReadOptional(BuiltInComponents->WeightAndEasingResult)
-		.FilterAll({ BuiltInComponents->Tags.AbsoluteBlend })
+		.FilterAll({ BuiltInComponents->Tags.AbsoluteBlend, GetBlenderTypeTag() })
 		.FilterNone({ BuiltInComponents->Tags.Ignored })
 		.Dispatch_PerAllocation<FIntegerAccumulationTask>(&Linker->EntityManager, InPrerequisites, nullptr, &AccumulationBuffers.Absolute);
 
@@ -295,7 +295,7 @@ void UMovieScenePiecewiseIntegerBlenderSystem::OnRun(FSystemTaskPrerequisites& I
 		.Read(BuiltInComponents->BlendChannelInput)
 		.Read(BuiltInComponents->IntegerResult)
 		.ReadOptional(BuiltInComponents->WeightAndEasingResult)
-		.FilterAll({ BuiltInComponents->Tags.RelativeBlend })
+		.FilterAll({ BuiltInComponents->Tags.RelativeBlend, GetBlenderTypeTag() })
 		.FilterNone({ BuiltInComponents->Tags.Ignored })
 		.Dispatch_PerAllocation<FIntegerAccumulationTask>(&Linker->EntityManager, InPrerequisites, nullptr, &AccumulationBuffers.Relative);
 
@@ -311,7 +311,7 @@ void UMovieScenePiecewiseIntegerBlenderSystem::OnRun(FSystemTaskPrerequisites& I
 		.Read(BuiltInComponents->BlendChannelInput)
 		.Read(BuiltInComponents->IntegerResult)
 		.ReadOptional(BuiltInComponents->WeightAndEasingResult)
-		.FilterAll({ BuiltInComponents->Tags.AdditiveBlend })
+		.FilterAll({ BuiltInComponents->Tags.AdditiveBlend, GetBlenderTypeTag() })
 		.FilterNone({ BuiltInComponents->Tags.Ignored })
 		.Dispatch_PerAllocation<FIntegerAccumulationTask>(&Linker->EntityManager, InPrerequisites, nullptr, &AccumulationBuffers.Additive);
 
@@ -328,7 +328,7 @@ void UMovieScenePiecewiseIntegerBlenderSystem::OnRun(FSystemTaskPrerequisites& I
 		.Read(BuiltInComponents->IntegerResult)
 		.Read(BuiltInComponents->BaseInteger)
 		.ReadOptional(BuiltInComponents->WeightAndEasingResult)
-		.FilterAll({ BuiltInComponents->Tags.AdditiveFromBaseBlend })
+		.FilterAll({ BuiltInComponents->Tags.AdditiveFromBaseBlend, GetBlenderTypeTag() })
 		.FilterNone({ BuiltInComponents->Tags.Ignored })
 		.Dispatch_PerAllocation<FIntegerAdditiveFromBaseBlendTask>(&Linker->EntityManager, InPrerequisites, nullptr, &AccumulationBuffers.AdditiveFromBase);
 
@@ -343,6 +343,7 @@ void UMovieScenePiecewiseIntegerBlenderSystem::OnRun(FSystemTaskPrerequisites& I
 	.Read(BuiltInComponents->BlendChannelOutput)
 	.ReadOptional(TracksComponents->Integer.InitialValue)
 	.Write(BuiltInComponents->IntegerResult)
+	.FilterAll({ GetBlenderTypeTag() })
 	.Dispatch_PerAllocation<FIntegerCombineBlends>(&Linker->EntityManager, Prereqs, &Subsequents, &AccumulationBuffers);
 }
 
@@ -359,10 +360,10 @@ void UMovieScenePiecewiseIntegerBlenderSystem::ReinitializeAccumulationBuffers()
 	// Find if we have any integer results to blend using any supported blend types.
 	const TComponentTypeID<int32> Component = BuiltInComponents->IntegerResult;
 
-	const bool bHasAbsolutes         = Linker->EntityManager.Contains(FEntityComponentFilter().All({ Component, BuiltInComponents->BlendChannelInput, BuiltInComponents->Tags.AbsoluteBlend }));
-	const bool bHasRelatives         = Linker->EntityManager.Contains(FEntityComponentFilter().All({ Component, BuiltInComponents->BlendChannelInput, BuiltInComponents->Tags.RelativeBlend }));
-	const bool bHasAdditives         = Linker->EntityManager.Contains(FEntityComponentFilter().All({ Component, BuiltInComponents->BlendChannelInput, BuiltInComponents->Tags.AdditiveBlend }));
-	const bool bHasAdditivesFromBase = Linker->EntityManager.Contains(FEntityComponentFilter().All({ Component, BuiltInComponents->BlendChannelInput, BuiltInComponents->Tags.AdditiveFromBaseBlend }));
+	const bool bHasAbsolutes         = Linker->EntityManager.Contains(FEntityComponentFilter().All({ GetBlenderTypeTag(), Component, BuiltInComponents->BlendChannelInput, BuiltInComponents->Tags.AbsoluteBlend }));
+	const bool bHasRelatives         = Linker->EntityManager.Contains(FEntityComponentFilter().All({ GetBlenderTypeTag(), Component, BuiltInComponents->BlendChannelInput, BuiltInComponents->Tags.RelativeBlend }));
+	const bool bHasAdditives         = Linker->EntityManager.Contains(FEntityComponentFilter().All({ GetBlenderTypeTag(), Component, BuiltInComponents->BlendChannelInput, BuiltInComponents->Tags.AdditiveBlend }));
+	const bool bHasAdditivesFromBase = Linker->EntityManager.Contains(FEntityComponentFilter().All({ GetBlenderTypeTag(), Component, BuiltInComponents->BlendChannelInput, BuiltInComponents->Tags.AdditiveFromBaseBlend }));
 
 	if (bHasAbsolutes)
 	{
diff --git a/Engine/Source/Runtime/MovieSceneTracks/Private/Systems/MovieScenePropertyInstantiator.cpp b/Engine/Source/Runtime/MovieSceneTracks/Private/Systems/MovieScenePropertyInstantiator.cpp
--- a/Engine/Source/Runtime/MovieSceneTracks/Private/Systems/MovieScenePropertyInstantiator.cpp
+++ b/Engine/Source/Runtime/MovieSceneTracks/Private/Systems/MovieScenePropertyInstantiator.cpp
@@ -562,53 +562,159 @@ void UMovieScenePropertyInstantiatorSystem::InitializeFastPath(const FPropertyPa
 	}
 }
 
-void UMovieScenePropertyInstantiatorSystem::InitializeBlendPath(const FPropertyParameters& Params)
+UMovieScenePropertyInstantiatorSystem::FSetupBlenderSystemResult UMovieScenePropertyInstantiatorSystem::SetupBlenderSystem(const FPropertyParameters& Params)
 {
 	using namespace UE::MovieScene;
 
-	TArrayView<const FPropertyCompositeDefinition> Composites = BuiltInComponents->PropertyRegistry.GetComposites(*Params.PropertyDefinition);
+	FSetupBlenderSystemResult Result;
 
-	UClass* BlenderClass = Params.PropertyDefinition->BlenderSystemClass;
+	UClass* NewBlenderClass = nullptr;
+	int32 BlenderClassPriority = TNumericLimits<int32>::Lowest();
 
-	// Ensure contributors all have the necessary blend inputs and tags
+	// Iterate all the contributors to locate the correct blender system by-priority
 	for (auto ContributorIt = Contributors.CreateConstKeyIterator(Params.PropertyInfoIndex); ContributorIt; ++ContributorIt)
 	{
 		FMovieSceneEntityID Contributor = ContributorIt.Value();
-
 		TOptionalComponentReader<TSubclassOf<UMovieSceneBlenderSystem>> BlenderTypeComponent = Linker->EntityManager.ReadComponent(Contributor, BuiltInComponents->BlenderType);
-		if (BlenderTypeComponent)
+		if (!BlenderTypeComponent)
 		{
-			BlenderClass = BlenderTypeComponent->Get();
-			break;
+			continue;
+		}
+
+		UClass* ProspectiveBlenderClass = BlenderTypeComponent->Get();
+		if (NewBlenderClass == ProspectiveBlenderClass)
+		{
+			// If it's already the same, don't waste time getting the CDO or anything like that
+			continue;
+		}
+
+		const UMovieSceneBlenderSystem* CDO = GetDefault<UMovieSceneBlenderSystem>(ProspectiveBlenderClass);
+
+		if (!NewBlenderClass || CDO->GetSelectionPriority() > BlenderClassPriority)
+		{
+			NewBlenderClass = ProspectiveBlenderClass;
+			BlenderClassPriority = CDO->GetSelectionPriority();
+		}
+		else
+		{
+#if DO_CHECK
+			ensureMsgf(CDO->GetSelectionPriority() != BlenderClassPriority,
+				TEXT("Encountered 2 different blender classes being used with the same priority - this is undefined behavior. Please check the system classes to ensure they have different priorities (%s and %s)."),
+				*ProspectiveBlenderClass->GetName(), *NewBlenderClass->GetName());
+#endif
 		}
 	}
 
-	if (!ensureMsgf(BlenderClass, TEXT("No default blender class specified on property, and no custom blender specified on entities. Falling back to float blender.")))
+	if (!NewBlenderClass)
 	{
-		BlenderClass = UMovieScenePiecewiseDoubleBlenderSystem::StaticClass();
+		NewBlenderClass = Params.PropertyDefinition->BlenderSystemClass;
 	}
 
+	if (!ensureMsgf(NewBlenderClass, TEXT("No default blender class specified on property, and no custom blender specified on entities. Falling back to double blender.")))
+	{
+		NewBlenderClass = UMovieScenePiecewiseDoubleBlenderSystem::StaticClass();
+	}
+
+	FComponentTypeID BlenderTypeTag = GetDefault<UMovieSceneBlenderSystem>(NewBlenderClass)->GetBlenderTypeTag();
+	ensureMsgf(BlenderTypeTag, TEXT("Encountered a blender system (%s) with an invalid type tag."), *NewBlenderClass->GetName());
+
+	FBlenderSystemInfo NewBlenderInfo{ NewBlenderClass, BlenderTypeTag };
+	FBlenderSystemInfo OldBlenderInfo;
+
 	UMovieSceneBlenderSystem* ExistingBlender = Params.PropertyInfo->Blender.Get();
-	if (ExistingBlender && BlenderClass != ExistingBlender->GetClass())
+	if (ExistingBlender)
 	{
-		const FMovieSceneBlendChannelID ExistingBlendChannel(ExistingBlender->GetBlenderSystemID(), Params.PropertyInfo->BlendChannel);
-		ExistingBlender->ReleaseBlendChannel(ExistingBlendChannel);
-		Params.PropertyInfo->BlendChannel = INVALID_BLEND_CHANNEL;
+		UClass* OldBlenderClass = ExistingBlender->GetClass();
+		if (OldBlenderClass != NewBlenderClass)
+		{
+			OldBlenderInfo = FBlenderSystemInfo{ OldBlenderClass, ExistingBlender->GetBlenderTypeTag() };
+		}
+		else
+		{
+			// It's the same - keep the same info
+			OldBlenderInfo = NewBlenderInfo;
+		}
 	}
 
-	UMovieSceneBlenderSystem* const Blender = CastChecked<UMovieSceneBlenderSystem>(Linker->LinkSystem(BlenderClass));
-	Params.PropertyInfo->Blender = Blender;
+	return FSetupBlenderSystemResult{NewBlenderInfo, OldBlenderInfo};
+}
+
+void UMovieScenePropertyInstantiatorSystem::InitializeBlendPath(const FPropertyParameters& Params)
+{
+	using namespace UE::MovieScene;
+
+	TArrayView<const FPropertyCompositeDefinition> Composites = BuiltInComponents->PropertyRegistry.GetComposites(*Params.PropertyDefinition);
+
+	FSetupBlenderSystemResult SetupResult = SetupBlenderSystem(Params);
 
-	const bool bWasAlreadyBlended = Params.PropertyInfo->BlendChannel != INVALID_BLEND_CHANNEL;
-	if (!bWasAlreadyBlended)
+	// -----------------------------------------------------------------------------------------------------
+	// Situation 1: New or modified contributors (inputs) but we're already set up for blending using the same system
+	if (SetupResult.CurrentInfo.BlenderSystemClass == SetupResult.PreviousInfo.BlenderSystemClass)
 	{
-		const FMovieSceneBlendChannelID NewBlendChannel = Blender->AllocateBlendChannel();
-		Params.PropertyInfo->BlendChannel = NewBlendChannel.ChannelID;
+		UMovieSceneBlenderSystem* Blender = Params.PropertyInfo->Blender.Get();
+		check(Blender);
+
+		// Ensure the output entity still matches the correct set of channels of all inputs
+		FComponentMask NewEntityType;
+		Params.MakeOutputComponentType(Linker->EntityManager, Composites, NewEntityType);
+		Linker->EntityManager.ChangeEntityType(Params.PropertyInfo->PropertyEntityID, NewEntityType);
+
+		const FMovieSceneBlendChannelID BlendChannel(Blender->GetBlenderSystemID(), Params.PropertyInfo->BlendChannel);
+
+		// Change new contributors to include the blend input components
+		for (auto ContributorIt = NewContributors.CreateConstKeyIterator(Params.PropertyInfoIndex); ContributorIt; ++ContributorIt)
+		{
+			FEntityBuilder()
+			.Add(BuiltInComponents->BlendChannelInput, BlendChannel)
+			.AddTag(SetupResult.CurrentInfo.BlenderTypeTag)
+			.MutateExisting(&Linker->EntityManager, ContributorIt.Value());
+		}
+
+		// Nothing more to do
+		return;
 	}
 
-	const FMovieSceneBlendChannelID BlendChannel(Blender->GetBlenderSystemID(), Params.PropertyInfo->BlendChannel);
+	// -----------------------------------------------------------------------------------------------------
+	// Situation 2: Never used blending before, or the blender type has changed
+	UMovieSceneBlenderSystem* OldBlender = Params.PropertyInfo->Blender.Get();
+	UMovieSceneBlenderSystem* NewBlender = CastChecked<UMovieSceneBlenderSystem>(Linker->LinkSystem(SetupResult.CurrentInfo.BlenderSystemClass));
+
+	const FMovieSceneBlendChannelID NewBlendChannel = NewBlender->AllocateBlendChannel();
+
+	Params.PropertyInfo->Blender = NewBlender;
+	Params.PropertyInfo->BlendChannel = NewBlendChannel.ChannelID;
 
-	if (!bWasAlreadyBlended)
+	FTypelessMutation InputMutation;
+
+	// -----------------------------------------------------------------------------------------------------
+	// Situation 2.1: We're already set up for blending, but with a different blender system
+	if (OldBlender)
+	{
+		const FMovieSceneBlendChannelID OldBlendChannel(OldBlender->GetBlenderSystemID(), Params.PropertyInfo->BlendChannel);
+
+		OldBlender->ReleaseBlendChannel(OldBlendChannel);
+
+		// Change the output entity by adding the new blend channel and tag, while simultaneously
+		// updating the channels and restore state flags etc added by MakeOutputComponentType
+		{
+			FTypelessMutation OutputMutation;
+			OutputMutation.RemoveAll();
+
+			Params.MakeOutputComponentType(Linker->EntityManager, Composites, OutputMutation.AddMask);
+
+			FEntityBuilder()
+			.Add(BuiltInComponents->BlendChannelOutput, NewBlendChannel)
+			.AddTag(SetupResult.CurrentInfo.BlenderTypeTag)
+			.MutateExisting(&Linker->EntityManager, Params.PropertyInfo->PropertyEntityID, OutputMutation);
+		}
+
+		// Ensure that the old blend tag is removed from inputs
+		InputMutation.Remove({ SetupResult.PreviousInfo.BlenderTypeTag });
+	}
+
+	// -----------------------------------------------------------------------------------------------------
+	// Situation 2.2: Never encountered blending before - need to create a new output entity to receive the blend result
+	else
 	{
 		FComponentMask NewMask;
 		NewMask.Set(Params.PropertyDefinition->InitialValueType);
@@ -633,92 +739,66 @@ void UMovieScenePropertyInstantiatorSystem::InitializeBlendPath(const FPropertyP
 		}
 		NewMask.Set(Params.PropertyDefinition->PropertyType);
 
-		FMovieSceneEntityID NewEntityID;
+		FMovieSceneEntityID NewOutputEntityID;
+
+		auto NewOutputEntity = FEntityBuilder()
+		.Add(BuiltInComponents->BlendChannelOutput,      NewBlendChannel)
+		.Add(BuiltInComponents->PropertyBinding,         Params.PropertyInfo->PropertyBinding)
+		.Add(BuiltInComponents->BoundObject,             Params.PropertyInfo->BoundObject)
+		.Add(BuiltInComponents->BlendChannelOutput,      NewBlendChannel)
+		.AddTagConditional(BuiltInComponents->Tags.MigratedFromFastPath, Params.PropertyInfo->PropertyEntityID.IsValid())
+		.AddTagConditional(BuiltInComponents->Tags.RestoreState, Params.PropertyInfo->bWantsRestoreState)
+		.AddTag(SetupResult.CurrentInfo.BlenderTypeTag)
+		.AddTag(BuiltInComponents->Tags.NeedsLink)
+		.AddMutualComponents();
+
 		switch (Params.PropertyInfo->Property.GetIndex())
 		{
 		// Never seen this property before
 		case 0:
-			NewEntityID = FEntityBuilder()
-			.Add(BuiltInComponents->FastPropertyOffset,      Params.PropertyInfo->Property.template Get<uint16>())
-			.Add(BuiltInComponents->PropertyBinding,         Params.PropertyInfo->PropertyBinding)
-			.Add(BuiltInComponents->BoundObject,             Params.PropertyInfo->BoundObject)
-			.Add(BuiltInComponents->BlendChannelOutput,      BlendChannel)
-			.AddTagConditional(BuiltInComponents->Tags.MigratedFromFastPath, Params.PropertyInfo->PropertyEntityID.IsValid())
-			.AddTagConditional(BuiltInComponents->Tags.RestoreState, Params.PropertyInfo->bWantsRestoreState)
-			.AddTag(BuiltInComponents->Tags.NeedsLink)
-			.AddMutualComponents()
+			NewOutputEntityID = NewOutputEntity
+			.Add(BuiltInComponents->FastPropertyOffset, Params.PropertyInfo->Property.template Get<uint16>())
 			.CreateEntity(&Linker->EntityManager, NewMask);
 			break;
 
 		case 1:
-			NewEntityID = FEntityBuilder()
+			NewOutputEntityID = NewOutputEntity
 			.Add(BuiltInComponents->CustomPropertyIndex, Params.PropertyInfo->Property.template Get<FCustomPropertyIndex>())
-			.Add(BuiltInComponents->PropertyBinding,     Params.PropertyInfo->PropertyBinding)
-			.Add(BuiltInComponents->BoundObject,         Params.PropertyInfo->BoundObject)
-			.Add(BuiltInComponents->BlendChannelOutput,  BlendChannel)
-			.AddTagConditional(BuiltInComponents->Tags.MigratedFromFastPath, Params.PropertyInfo->PropertyEntityID.IsValid())
-			.AddTagConditional(BuiltInComponents->Tags.RestoreState, Params.PropertyInfo->bWantsRestoreState)
-			.AddTag(BuiltInComponents->Tags.NeedsLink)
-			.AddMutualComponents()
 			.CreateEntity(&Linker->EntityManager, NewMask);
 			break;
 
 		case 2:
-			NewEntityID = FEntityBuilder()
-			.Add(BuiltInComponents->SlowProperty,            Params.PropertyInfo->Property.template Get<FSlowPropertyPtr>())
-			.Add(BuiltInComponents->PropertyBinding,         Params.PropertyInfo->PropertyBinding)
-			.Add(BuiltInComponents->BoundObject,             Params.PropertyInfo->BoundObject)
-			.Add(BuiltInComponents->BlendChannelOutput,      BlendChannel)
-			.AddTagConditional(BuiltInComponents->Tags.MigratedFromFastPath, Params.PropertyInfo->PropertyEntityID.IsValid())
-			.AddTagConditional(BuiltInComponents->Tags.RestoreState, Params.PropertyInfo->bWantsRestoreState)
-			.AddTag(BuiltInComponents->Tags.NeedsLink)
-			.AddMutualComponents()
+			NewOutputEntityID = NewOutputEntity
+			.Add(BuiltInComponents->SlowProperty, Params.PropertyInfo->Property.template Get<FSlowPropertyPtr>())
 			.CreateEntity(&Linker->EntityManager, NewMask);
 			break;
 		}
 
+		// Params.PropertyInfo->PropertyEntityID will be set if the entity was previously using the fast path
 		if (Params.PropertyInfo->PropertyEntityID)
 		{
 			// Move any copiable/migratable components over from the existing fast-path entity
-			Linker->EntityManager.CopyComponents(Params.PropertyInfo->PropertyEntityID, NewEntityID, Linker->EntityManager.GetComponents()->GetCopyAndMigrationMask());
-
-			// Add blend inputs on the first contributor, which was using the fast-path
-			Linker->EntityManager.AddComponent(Params.PropertyInfo->PropertyEntityID, BuiltInComponents->BlendChannelInput, BlendChannel);
-			Linker->EntityManager.RemoveComponents(Params.PropertyInfo->PropertyEntityID, CleanFastPathMask);
-		}
-
-		Params.PropertyInfo->PropertyEntityID = NewEntityID;
-	}
-	else
-	{
-		FComponentMask NewEntityType = Linker->EntityManager.GetEntityType(Params.PropertyInfo->PropertyEntityID);
-
-		// Ensure the property has only the exact combination of components that constitute its animation
-		for (int32 Index = 0; Index < Composites.Num(); ++Index)
-		{
-			FComponentTypeID Composite = Composites[Index].ComponentTypeID;
-			NewEntityType[Composite] = (Params.PropertyInfo->EmptyChannels[Index] != true);
-		}
-		NewEntityType.Set(Params.PropertyDefinition->PropertyType);
+			Linker->EntityManager.CopyComponents(Params.PropertyInfo->PropertyEntityID, NewOutputEntityID, Linker->EntityManager.GetComponents()->GetCopyAndMigrationMask());
 
-		if (Params.PropertyInfo->bWantsRestoreState)
-		{
-			NewEntityType.Set(BuiltInComponents->Tags.RestoreState);
-		}
-		else
-		{
-			NewEntityType.Remove(BuiltInComponents->Tags.RestoreState);
+			// Clean the previously-fast-path entity
+			InputMutation.RemoveMask = CleanFastPathMask;
 		}
 
-		Linker->EntityManager.ChangeEntityType(Params.PropertyInfo->PropertyEntityID, NewEntityType);
+		// The property entity ID is now the blend output entity
+		Params.PropertyInfo->PropertyEntityID = NewOutputEntityID;
 	}
 
-	// Ensure contributors all have the necessary blend inputs and tags
-	for (auto ContributorIt = NewContributors.CreateConstKeyIterator(Params.PropertyInfoIndex); ContributorIt; ++ContributorIt)
+	// Change *all* contributors (not just new ones because the old ones will have the old blender's channel and tag on them)
+	// to include the new blend channel input, and remove the old blender type. No need to remove the clean fast path mask because
+	// that will have already happened as part of the 'completely new blending' branch below
+	for (auto ContributorIt = Contributors.CreateConstKeyIterator(Params.PropertyInfoIndex); ContributorIt; ++ContributorIt)
 	{
 		const FMovieSceneEntityID Contributor = ContributorIt.Value();
-		Linker->EntityManager.AddComponent(Contributor, BuiltInComponents->BlendChannelInput, BlendChannel);
-		Linker->EntityManager.RemoveComponents(Contributor, CleanFastPathMask);
+
+		FEntityBuilder()
+		.Add(BuiltInComponents->BlendChannelInput, NewBlendChannel)
+		.AddTag(SetupResult.CurrentInfo.BlenderTypeTag)
+		.MutateExisting(&Linker->EntityManager, Contributor, InputMutation);
 	}
 }
 
@@ -791,3 +871,31 @@ void UMovieScenePropertyInstantiatorSystem::InitializePropertyMetaData(FSystemTa
 	InitializePropertyMetaDataTasks.Empty();
 }
 
+void UMovieScenePropertyInstantiatorSystem::FPropertyParameters::MakeOutputComponentType(
+	const UE::MovieScene::FEntityManager& EntityManager,
+	TArrayView<const UE::MovieScene::FPropertyCompositeDefinition> Composites,
+	UE::MovieScene::FComponentMask& OutComponentType) const
+{
+	using namespace UE::MovieScene;
+
+	// Get the existing type
+	OutComponentType = EntityManager.GetEntityType(PropertyInfo->PropertyEntityID);
+
+	// Ensure the property has only the exact combination of channels that constitute its animation
+	for (int32 Index = 0; Index < Composites.Num(); ++Index)
+	{
+		FComponentTypeID Composite = Composites[Index].ComponentTypeID;
+		OutComponentType[Composite] = (PropertyInfo->EmptyChannels[Index] != true);
+	}
+	OutComponentType.Set(PropertyDefinition->PropertyType);
+
+	// Set the restore state tag appropriately
+	if (PropertyInfo->bWantsRestoreState)
+	{
+		OutComponentType.Set(FBuiltInComponentTypes::Get()->Tags.RestoreState);
+	}
+	else
+	{
+		OutComponentType.Remove(FBuiltInComponentTypes::Get()->Tags.RestoreState);
+	}
+}
diff --git a/Engine/Source/Runtime/MovieSceneTracks/Public/Systems/MovieSceneBlenderSystemHelper.h b/Engine/Source/Runtime/MovieSceneTracks/Public/Systems/MovieSceneBlenderSystemHelper.h
--- a/Engine/Source/Runtime/MovieSceneTracks/Public/Systems/MovieSceneBlenderSystemHelper.h
+++ b/Engine/Source/Runtime/MovieSceneTracks/Public/Systems/MovieSceneBlenderSystemHelper.h
@@ -156,6 +156,7 @@ public:
 
 			const bool bHasChannels = Linker->EntityManager.Contains(FEntityComponentFilter().All({
 						ResultComponentID, 
+						BlenderSystem->GetBlenderTypeTag(),
 						BuiltInComponents->BlendChannelInput }));
 			if (bHasChannels)
 			{
@@ -174,6 +175,7 @@ public:
 		FGraphEventRef Task = FEntityTaskBuilder()
 			.Read(BuiltInComponents->BlendChannelInput)
 			.Read(ResultComponentID)
+			.FilterAll({ BlenderSystem->GetBlenderTypeTag() })
 			.FilterNone({ BuiltInComponents->Tags.Ignored })
 			.template Dispatch_PerEntity<TSimpleBlenderGatherResults<PropertyType>>(
 					&Linker->EntityManager, InPrerequisites, nullptr, BlenderSystem->GetBlenderSystemID(), BlendChannelResults);
@@ -186,6 +188,7 @@ public:
 		FEntityTaskBuilder()
 			.Read(BuiltInComponents->BlendChannelOutput)
 			.Write(ResultComponentID)
+			.FilterAll({ BlenderSystem->GetBlenderTypeTag() })
 			.FilterNone({ BuiltInComponents->Tags.Ignored })
 			.template Dispatch_PerEntity<TSimpleBlenderCombineResults<PropertyType>>(
 					&Linker->EntityManager, Prereqs, &Subsequents, BlenderSystem->GetBlenderSystemID(), BlendChannelResults);
diff --git a/Engine/Source/Runtime/MovieSceneTracks/Public/Systems/MovieScenePropertyInstantiator.h b/Engine/Source/Runtime/MovieSceneTracks/Public/Systems/MovieScenePropertyInstantiator.h
--- a/Engine/Source/Runtime/MovieSceneTracks/Public/Systems/MovieScenePropertyInstantiator.h
+++ b/Engine/Source/Runtime/MovieSceneTracks/Public/Systems/MovieScenePropertyInstantiator.h
@@ -139,6 +139,18 @@ private:
 
 private:
 
+	struct FBlenderSystemInfo
+	{
+		UClass* BlenderSystemClass = nullptr;
+		UE::MovieScene::FComponentTypeID BlenderTypeTag;
+	};
+
+	struct FSetupBlenderSystemResult
+	{
+		FBlenderSystemInfo CurrentInfo;
+		FBlenderSystemInfo PreviousInfo;
+	};
+
 	/* Parameter structure passed around when instantiating a specific instance of a property */
 	struct FPropertyParameters
 	{
@@ -148,6 +160,11 @@ private:
 		const UE::MovieScene::FPropertyDefinition* PropertyDefinition;
 		/** The index of the PropertyInfo member within UMovieScenePropertyInstantiatorSystem::ResolvedProperties */
 		int32 PropertyInfoIndex;
+
+		void MakeOutputComponentType(
+			const UE::MovieScene::FEntityManager& EntityManager,
+			TArrayView<const UE::MovieScene::FPropertyCompositeDefinition> Composites,
+			UE::MovieScene::FComponentMask& OutComponentType) const;
 	};
 	void DiscoverInvalidatedProperties(TBitArray<>& OutInvalidatedProperties);
 	void UpgradeFloatToDoubleProperties(const TBitArray<>& InvalidatedProperties);
@@ -156,6 +173,9 @@ private:
 	bool PropertySupportsFastPath(const FPropertyParameters& Params) const;
 	void InitializeFastPath(const FPropertyParameters& Params);
 	void InitializeBlendPath(const FPropertyParameters& Params);
+
+	FSetupBlenderSystemResult SetupBlenderSystem(const FPropertyParameters& Params);
+
 	int32 ResolveProperty(UE::MovieScene::FCustomAccessorView CustomAccessors, UObject* Object, const FMovieScenePropertyBinding& PropertyBinding, int32 PropertyDefinitionIndex);
 
 	UE::MovieScene::FPropertyRecomposerPropertyInfo FindPropertyFromSource(FMovieSceneEntityID EntityID, UObject* Object) const;
