Sequencer: Added a new quaternion (slerp) blending algorithm for blending transform tracks using a 'shortest path' rotation
- Tracks that support multiple blenders will now have a "Blending Algorithm ->" sub-menu on their context menu which is used for choosing the blender to use.

diff --git a/Engine/Source/Editor/Sequencer/Private/MVVM/ViewModels/TrackModel.cpp b/Engine/Source/Editor/Sequencer/Private/MVVM/ViewModels/TrackModel.cpp
--- a/Engine/Source/Editor/Sequencer/Private/MVVM/ViewModels/TrackModel.cpp
+++ b/Engine/Source/Editor/Sequencer/Private/MVVM/ViewModels/TrackModel.cpp
@@ -23,6 +23,7 @@
 #include "MovieSceneNameableTrack.h"
 #include "Tracks/MovieScene3DTransformTrack.h"
 #include "Tracks/MovieScenePrimitiveMaterialTrack.h"
+#include "EntitySystem/IMovieSceneBlenderSystemSupport.h"
 
 #include "ISequencer.h"
 #include "ISequencerSection.h"
@@ -644,6 +645,8 @@ void FTrackModel::BuildContextMenu(FMenuBuilder& MenuBuilder)
 		return;
 	}
 
+	TWeakPtr<ISequencer> WeakSequencer = GetEditor()->GetSequencer();
+
 	const int32 TrackRowIndex = GetRowIndex();
 
 	if (TrackEditor)
@@ -653,8 +656,6 @@ void FTrackModel::BuildContextMenu(FMenuBuilder& MenuBuilder)
 
 	if (Track && Track->GetSupportedBlendTypes().Num() > 0)
 	{
-		TWeakPtr<ISequencer> WeakSequencer = GetEditor()->GetSequencer();
-
 		MenuBuilder.AddSubMenu(
 			LOCTEXT("AddSection", "Add Section"),
 			FText(),
@@ -664,6 +665,25 @@ void FTrackModel::BuildContextMenu(FMenuBuilder& MenuBuilder)
 		);	
 	}
 
+	// Add menu items for selecting a blender
+	IMovieSceneBlenderSystemSupport* BlenderSystemSupport = Cast<IMovieSceneBlenderSystemSupport>(Track);
+	if (BlenderSystemSupport)
+	{
+		TArray<TSubclassOf<UMovieSceneBlenderSystem>> BlenderTypes;
+		BlenderSystemSupport->GetSupportedBlenderSystems(BlenderTypes);
+
+		if (BlenderTypes.Num() > 1)
+		{
+			MenuBuilder.AddSubMenu(
+				LOCTEXT("BlendingAlgorithmSubMenu", "Blending Algorithm"),
+				FText(),
+				FNewMenuDelegate::CreateLambda([=](FMenuBuilder& SubMenuBuilder){
+					FSequencerUtilities::PopulateMenu_BlenderSubMenu(SubMenuBuilder, Track, WeakSequencer);
+				})
+			);
+		}
+	}
+
 	// Find sections in the track to add batch properties for
 	TArray<TWeakObjectPtr<UObject>> TrackSections;
 
diff --git a/Engine/Source/Editor/Sequencer/Private/MVVM/ViewModels/TrackRowModel.cpp b/Engine/Source/Editor/Sequencer/Private/MVVM/ViewModels/TrackRowModel.cpp
--- a/Engine/Source/Editor/Sequencer/Private/MVVM/ViewModels/TrackRowModel.cpp
+++ b/Engine/Source/Editor/Sequencer/Private/MVVM/ViewModels/TrackRowModel.cpp
@@ -26,6 +26,7 @@
 #include "MovieSceneTrack.h"
 #include "Tracks/MovieScene3DTransformTrack.h"
 #include "Tracks/MovieScenePrimitiveMaterialTrack.h"
+#include "EntitySystem/IMovieSceneBlenderSystemSupport.h"
 
 #include "ScopedTransaction.h"
 #include "Styling/AppStyle.h"
@@ -297,6 +298,8 @@ void FTrackRowModel::BuildContextMenu(FMenuBuilder& MenuBuilder)
 		return;
 	}
 
+	TSharedPtr<ISequencer> WeakSequencer = GetEditor()->GetSequencer();
+
 	const int32 TrackRowIndex = GetRowIndex();
 
 	TSharedPtr<ISequencerTrackEditor> TrackEditorPtr = GetTrackEditor();
@@ -304,8 +307,6 @@ void FTrackRowModel::BuildContextMenu(FMenuBuilder& MenuBuilder)
 
 	if (Track && Track->GetSupportedBlendTypes().Num() > 0)
 	{
-		TSharedPtr<ISequencer> WeakSequencer = GetEditor()->GetSequencer();
-
 		MenuBuilder.AddSubMenu(
 			LOCTEXT("AddSection", "Add Section"),
 			FText(),
@@ -315,6 +316,25 @@ void FTrackRowModel::BuildContextMenu(FMenuBuilder& MenuBuilder)
 		);	
 	}
 
+	// Add menu items for selecting a blender
+	IMovieSceneBlenderSystemSupport* BlenderSystemSupport = Cast<IMovieSceneBlenderSystemSupport>(Track);
+	if (BlenderSystemSupport)
+	{
+		TArray<TSubclassOf<UMovieSceneBlenderSystem>> BlenderTypes;
+		BlenderSystemSupport->GetSupportedBlenderSystems(BlenderTypes);
+
+		if (BlenderTypes.Num() > 1)
+		{
+			MenuBuilder.AddSubMenu(
+				LOCTEXT("BlendingAlgorithmSubMenu", "Blending Algorithm"),
+				FText(),
+				FNewMenuDelegate::CreateLambda([=](FMenuBuilder& SubMenuBuilder){
+					FSequencerUtilities::PopulateMenu_BlenderSubMenu(SubMenuBuilder, Track, WeakSequencer);
+				})
+			);
+		}
+	}
+
 	// Find sections in the track to add batch properties for
 	TArray<TWeakObjectPtr<UObject>> TrackSections;
 	
diff --git a/Engine/Source/Editor/Sequencer/Private/SequencerUtilities.cpp b/Engine/Source/Editor/Sequencer/Private/SequencerUtilities.cpp
--- a/Engine/Source/Editor/Sequencer/Private/SequencerUtilities.cpp
+++ b/Engine/Source/Editor/Sequencer/Private/SequencerUtilities.cpp
@@ -25,6 +25,8 @@
 #include "Styling/AppStyle.h"
 #include "MovieSceneCopyableBinding.h"
 #include "MovieSceneCopyableTrack.h"
+#include "EntitySystem/MovieSceneBlenderSystem.h"
+#include "EntitySystem/IMovieSceneBlenderSystemSupport.h"
 #include "MovieSceneFolder.h"
 #include "MovieSceneSection.h"
 #include "MovieSceneSpawnRegister.h"
@@ -36,6 +38,7 @@
 #include "Tracks/MovieSceneSpawnTrack.h"
 #include "Compilation/MovieSceneCompiledDataManager.h"
 #include "Framework/MultiBox/MultiBoxBuilder.h"
+#include "Styling/SlateIconFinder.h"
 #include "ISequencerTrackEditor.h"
 #include "ISequencer.h"
 #include "Sequencer.h"
@@ -262,6 +265,50 @@ void FSequencerUtilities::PopulateMenu_CreateNewSection(FMenuBuilder& MenuBuilde
 	}
 }
 
+void FSequencerUtilities::PopulateMenu_BlenderSubMenu(FMenuBuilder& MenuBuilder, UMovieSceneTrack* Track, TWeakPtr<ISequencer> InSequencer)
+{
+	IMovieSceneBlenderSystemSupport* BlenderSystemSupport = Cast<IMovieSceneBlenderSystemSupport>(Track);
+	// Shouldn't have been called with a track that does not implement this interface
+	check(BlenderSystemSupport);
+
+	TArray<TSubclassOf<UMovieSceneBlenderSystem>> BlenderTypes;
+	BlenderSystemSupport->GetSupportedBlenderSystems(BlenderTypes);
+
+	// Ensure no nulls
+	BlenderTypes.Remove(TSubclassOf<UMovieSceneBlenderSystem>());
+
+	// Sort alphabetically
+	Algo::Sort(BlenderTypes,
+		[](TSubclassOf<UMovieSceneBlenderSystem> A, TSubclassOf<UMovieSceneBlenderSystem> B)
+		{
+			return A->GetDisplayNameText().CompareTo(B->GetDisplayNameText()) < 0;
+		}
+	);
+
+	for (TSubclassOf<UMovieSceneBlenderSystem> SystemClass : BlenderTypes)
+	{
+		MenuBuilder.AddMenuEntry(
+			SystemClass->GetDisplayNameText(),
+			SystemClass->GetToolTipText(),
+			FSlateIconFinder::FindIconForClass(SystemClass.Get()),
+			FUIAction(
+				FExecuteAction::CreateLambda([Track, BlenderSystemSupport, SystemClass]{
+					FScopedTransaction Transaction(FText::Format(LOCTEXT("ChangeBlenderType", "Change blender to '{0}'"), SystemClass.Get()->GetDisplayNameText()));
+
+					Track->Modify();
+					BlenderSystemSupport->SetBlenderSystem(SystemClass);
+				}),
+				FCanExecuteAction(),
+				FIsActionChecked::CreateLambda([BlenderSystemSupport, SystemClass]
+				{
+					return BlenderSystemSupport->GetBlenderSystem() == SystemClass;
+				})),
+			NAME_None,
+			EUserInterfaceActionType::RadioButton
+		);
+	}
+}
+
 void FSequencerUtilities::PopulateMenu_SetBlendType(FMenuBuilder& MenuBuilder, UMovieSceneSection* Section, TWeakPtr<ISequencer> InSequencer)
 {
 	PopulateMenu_SetBlendType(MenuBuilder, TArray<TWeakObjectPtr<UMovieSceneSection>>({ Section }), InSequencer);
diff --git a/Engine/Source/Editor/Sequencer/Public/SequencerUtilities.h b/Engine/Source/Editor/Sequencer/Public/SequencerUtilities.h
--- a/Engine/Source/Editor/Sequencer/Public/SequencerUtilities.h
+++ b/Engine/Source/Editor/Sequencer/Public/SequencerUtilities.h
@@ -128,6 +128,8 @@ struct SEQUENCER_API FSequencerUtilities
 
 	static void PopulateMenu_CreateNewSection(FMenuBuilder& MenuBuilder, int32 RowIndex, UMovieSceneTrack* Track, TWeakPtr<ISequencer> InSequencer);
 
+	static void PopulateMenu_BlenderSubMenu(FMenuBuilder& MenuBuilder, UMovieSceneTrack* Track, TWeakPtr<ISequencer> InSequencer);
+
 	static void PopulateMenu_SetBlendType(FMenuBuilder& MenuBuilder, UMovieSceneSection* Section, TWeakPtr<ISequencer> InSequencer);
 
 	static void PopulateMenu_SetBlendType(FMenuBuilder& MenuBuilder, const TArray<TWeakObjectPtr<UMovieSceneSection>>& InSections, TWeakPtr<ISequencer> InSequencer);
diff --git a/Engine/Source/Runtime/MovieScene/Private/MovieSceneSection.cpp b/Engine/Source/Runtime/MovieScene/Private/MovieSceneSection.cpp
--- a/Engine/Source/Runtime/MovieScene/Private/MovieSceneSection.cpp
+++ b/Engine/Source/Runtime/MovieScene/Private/MovieSceneSection.cpp
@@ -14,6 +14,8 @@
 #include "EntitySystem/BuiltInComponentTypes.h"
 #include "EntitySystem/MovieSceneEntitySystemTask.h"
 #include "EntitySystem/MovieSceneEntitySystemLinker.h"
+#include "EntitySystem/MovieSceneBlenderSystem.h"
+#include "EntitySystem/IMovieSceneBlenderSystemSupport.h"
 #include "Containers/ArrayView.h"
 #include "Channels/MovieSceneChannel.h"
 #include "UObject/SequencerObjectVersion.h"
@@ -323,8 +325,24 @@ void UMovieSceneSection::BuildDefaultComponents(UMovieSceneEntitySystemLinker* E
 	const bool bHasSectionPreRoll  = Params.EntityMetaData && EnumHasAnyFlags(Params.EntityMetaData->Flags, ESectionEvaluationFlags::PreRoll | ESectionEvaluationFlags::PostRoll);
 	const bool bHasSequencePreRoll = Params.Sequence.bPreRoll || Params.Sequence.bPostRoll;
 
+	TSubclassOf<UMovieSceneBlenderSystem> BlenderSystemClass = nullptr;
+
+	// Try and find a blender system to use
+	{
+		IMovieSceneBlenderSystemSupport* BlenderSystemSupport = Cast<IMovieSceneBlenderSystemSupport>(this);
+		if (!BlenderSystemSupport)
+		{
+			BlenderSystemSupport = GetImplementingOuter<IMovieSceneBlenderSystemSupport>();
+		}
+		if (BlenderSystemSupport)
+		{
+			BlenderSystemClass = BlenderSystemSupport->GetBlenderSystem();
+		}
+	}
+
 	OutImportedEntity->AddBuilder(
 		FEntityBuilder()
+		.AddConditional(Components->BlenderType,                BlenderSystemClass, BlenderSystemClass.Get() != nullptr)
 		.AddConditional(Components->Easing,                     FEasingComponentData{ decltype(FEasingComponentData::Section)(this) }, bHasEasing)
 		.AddConditional(Components->HierarchicalBias,           Params.Sequence.HierarchicalBias, Params.Sequence.HierarchicalBias != 0)
 		.AddConditional(Components->Interrogation.InputKey,     Params.InterrogationKey, Params.InterrogationKey.IsValid())
diff --git a/Engine/Source/Runtime/MovieScene/Public/EntitySystem/IMovieSceneBlenderSystemSupport.h b/Engine/Source/Runtime/MovieScene/Public/EntitySystem/IMovieSceneBlenderSystemSupport.h
new file mode 100644
--- /dev/null
+++ b/Engine/Source/Runtime/MovieScene/Public/EntitySystem/IMovieSceneBlenderSystemSupport.h
@@ -0,0 +1,33 @@
+// Copyright Epic Games, Inc. All Rights Reserved.
+
+#pragma once
+
+#include "UObject/Interface.h"
+#include "UObject/ObjectMacros.h"
+#include "Containers/Array.h"
+#include "Templates/SubclassOf.h"
+
+#include "IMovieSceneBlenderSystemSupport.generated.h"
+
+class UMovieSceneBlenderSystem;
+
+UINTERFACE()
+class MOVIESCENE_API UMovieSceneBlenderSystemSupport : public UInterface
+{
+public:
+	GENERATED_BODY()
+};
+
+
+/**
+ * Interface that can be added to UMovieSceneTracks to enable user-selection of blender systems
+ */
+class MOVIESCENE_API IMovieSceneBlenderSystemSupport
+{
+public:
+	GENERATED_BODY()
+
+	virtual TSubclassOf<UMovieSceneBlenderSystem> GetBlenderSystem() const = 0;
+	virtual void SetBlenderSystem(TSubclassOf<UMovieSceneBlenderSystem> BlenderSystemClass) = 0;
+	virtual void GetSupportedBlenderSystems(TArray<TSubclassOf<UMovieSceneBlenderSystem>>& OutSystemClasses) const = 0;
+};
\ No newline at end of file
diff --git a/Engine/Source/Runtime/MovieScene/Public/EntitySystem/MovieSceneEntitySystemTask.h b/Engine/Source/Runtime/MovieScene/Public/EntitySystem/MovieSceneEntitySystemTask.h
--- a/Engine/Source/Runtime/MovieScene/Public/EntitySystem/MovieSceneEntitySystemTask.h
+++ b/Engine/Source/Runtime/MovieScene/Public/EntitySystem/MovieSceneEntitySystemTask.h
@@ -556,6 +556,17 @@ struct TEntityTaskComponentsImpl<TIntegerSequence<int, Indices...>, T...>
 		return TEntityTaskComponents<T..., TWriteAccess<U> >(CommonParams, Accessors.template Get<Indices>()..., TWriteAccess<U>(ComponentType));
 	}
 
+	/**
+	 * Write all of the specified components and pass them through to the task as individual parameters
+	 *
+	 * @param ComponentTypes  The component types to visit
+	 */
+	template<typename... U>
+	TEntityTaskComponents<T..., TWriteAccess<U>...> WriteAllOf(TComponentTypeID<U>... ComponentTypes) const
+	{
+		return TEntityTaskComponents<T..., TWriteAccess<U>... >( CommonParams, Accessors.template Get<Indices>()..., TWriteAccess<U>(ComponentTypes)... );
+	}
+
 	/**
 	 * Write the type-erased value of a component. Passed to the task as a void*
 	 * @note Supplying an invalid ComponentType will be handled gracefully, but will result in no task being dispatched.
diff --git a/Engine/Source/Runtime/MovieSceneTracks/Private/EntitySystem/Interrogation/MovieSceneInterrogatedPropertyInstantiator.cpp b/Engine/Source/Runtime/MovieSceneTracks/Private/EntitySystem/Interrogation/MovieSceneInterrogatedPropertyInstantiator.cpp
--- a/Engine/Source/Runtime/MovieSceneTracks/Private/EntitySystem/Interrogation/MovieSceneInterrogatedPropertyInstantiator.cpp
+++ b/Engine/Source/Runtime/MovieSceneTracks/Private/EntitySystem/Interrogation/MovieSceneInterrogatedPropertyInstantiator.cpp
@@ -195,7 +195,8 @@ void UMovieSceneInterrogatedPropertyInstantiatorSystem::UpdateOutput(UE::MovieSc
 		Output->BlendChannel = INVALID_BLEND_CHANNEL;
 	}
 
-	Output->Blender = CastChecked<UMovieSceneBlenderSystem>(Linker->LinkSystem(BlenderClass));
+	UMovieSceneBlenderSystem* BlenderSystem = CastChecked<UMovieSceneBlenderSystem>(Linker->LinkSystem(BlenderClass));
+	Output->Blender = BlenderSystem;
 
 	const bool bWasAlreadyBlended = Output->BlendChannel != INVALID_BLEND_CHANNEL;
 	if (!bWasAlreadyBlended)
@@ -239,6 +240,7 @@ void UMovieSceneInterrogatedPropertyInstantiatorSystem::UpdateOutput(UE::MovieSc
 		.Add(BuiltInComponents->BlendChannelOutput, BlendChannel)
 		.AddTagConditional(BuiltInComponents->Tags.MigratedFromFastPath, Output->PropertyEntityID.IsValid())
 		.AddTag(BuiltInComponents->Tags.NeedsLink)
+		.AddTag(BlenderSystem->GetBlenderTypeTag())
 		.AddMutualComponents()
 		.CreateEntity(&Linker->EntityManager, NewMask);
 
@@ -277,6 +279,7 @@ void UMovieSceneInterrogatedPropertyInstantiatorSystem::UpdateOutput(UE::MovieSc
 	{
 		const FMovieSceneBlendChannelID BlendChannel(Output->Blender->GetBlenderSystemID(), Output->BlendChannel);
 		Linker->EntityManager.AddComponent(Input, BuiltInComponents->BlendChannelInput, BlendChannel);
+		Linker->EntityManager.AddComponent(Input, BlenderSystem->GetBlenderTypeTag());
 		Linker->EntityManager.RemoveComponents(Input, CleanFastPathMask);
 	}
 }
diff --git a/Engine/Source/Runtime/MovieSceneTracks/Private/Systems/MovieSceneQuaternionBlenderSystem.cpp b/Engine/Source/Runtime/MovieSceneTracks/Private/Systems/MovieSceneQuaternionBlenderSystem.cpp
new file mode 100644
--- /dev/null
+++ b/Engine/Source/Runtime/MovieSceneTracks/Private/Systems/MovieSceneQuaternionBlenderSystem.cpp
@@ -0,0 +1,314 @@
+// Copyright Epic Games, Inc. All Rights Reserved.
+
+#include "Systems/MovieSceneQuaternionBlenderSystem.h"
+#include "EntitySystem/BuiltInComponentTypes.h"
+#include "EntitySystem/EntityAllocationIterator.h"
+#include "EntitySystem/MovieSceneEntityManager.h"
+#include "EntitySystem/MovieSceneEntitySystemLinker.h"
+#include "EntitySystem/MovieSceneEntitySystemTask.h"
+#include "Systems/FloatChannelEvaluatorSystem.h"
+#include "Systems/DoubleChannelEvaluatorSystem.h"
+#include "Systems/MovieSceneBaseValueEvaluatorSystem.h"
+#include "Systems/MovieSceneQuaternionInterpolationRotationSystem.h"
+#include "Systems/WeightAndEasingEvaluatorSystem.h"
+#include "MovieSceneTracksComponentTypes.h"
+
+#include UE_INLINE_GENERATED_CPP_BY_NAME(MovieSceneQuaternionBlenderSystem)
+
+
+namespace UE::MovieScene
+{
+
+struct FAccumulateAbsoluteWeights
+{
+	FQuaternionBlenderAccumulationBuffers* Buffers;
+
+	FAccumulateAbsoluteWeights(FQuaternionBlenderAccumulationBuffers* InBuffers)
+		: Buffers(InBuffers)
+	{}
+
+	void PreTask()
+	{
+		Buffers->AbsoluteWeights.Reset();
+	}
+
+	void ForEachEntity(FMovieSceneBlendChannelID BlendChannel, double WeightResult)
+	{
+		if (!Buffers->AbsoluteWeights.IsValidIndex(BlendChannel.ChannelID))
+		{
+			Buffers->AbsoluteWeights.Insert(BlendChannel.ChannelID, WeightResult);
+		}
+		else
+		{
+			Buffers->AbsoluteWeights[BlendChannel.ChannelID] += WeightResult;
+		}
+	}
+};
+
+struct FCollectInitialValues
+{
+	FQuaternionBlenderAccumulationBuffers* Buffers;
+
+	FCollectInitialValues(FQuaternionBlenderAccumulationBuffers* InBuffers)
+		: Buffers(InBuffers)
+	{}
+
+	void PreTask()
+	{
+		FMemory::Memzero(Buffers->InitialValues.GetData(), sizeof(FQuatTransform) * Buffers->InitialValues.Num());
+	}
+
+	void ForEachEntity(FMovieSceneBlendChannelID BlendChannel, const FIntermediate3DTransform& InitialValue)
+	{
+		Buffers->InitialValues[BlendChannel.ChannelID] = FQuatTransform{
+			InitialValue.GetTranslation(),
+			InitialValue.GetRotation().Quaternion(),
+			InitialValue.GetScale()
+		};
+	}
+};
+
+
+/** Task for accumulating all weighted blend inputs into arrays based on BlendID. Will be run for Absolute, Additive and Relative blend modes*/
+struct FAbsoluteAccumulationTask
+{
+	FQuaternionBlenderAccumulationBuffers* Buffers;
+
+	FAbsoluteAccumulationTask(FQuaternionBlenderAccumulationBuffers* InBuffers)
+		: Buffers(InBuffers)
+	{}
+
+	void PreTask()
+	{
+		for (int32 Index = 0; Index < Buffers->Absolutes.Num(); ++Index)
+		{
+			Buffers->Absolutes[Index] = FQuatTransform{ 
+				FVector::ZeroVector,
+				FQuat::Identity,
+				FVector::ZeroVector
+			};
+		}
+	}
+
+	void ForEachEntity(
+		FMovieSceneBlendChannelID BlendChannel,
+		double LocationX, double LocationY, double LocationZ,
+		double RotationRoll, double RotationPitch, double RotationYaw,
+		double ScaleX, double ScaleY, double ScaleZ,
+		const double* WeightResult
+		)
+	{
+		double Weight = WeightResult ? *WeightResult : 1.0;
+		if (Weight == 0.0)
+		{
+			return;
+		}
+
+		const double TotalWeight = Buffers->AbsoluteWeights[BlendChannel.ChannelID]
+			? Buffers->AbsoluteWeights[BlendChannel.ChannelID]
+			: 1.0;
+
+		FQuatTransform& Transform = Buffers->Absolutes[BlendChannel.ChannelID];
+
+		FQuat ThisRotation = FRotator(RotationPitch, RotationYaw, RotationRoll).Quaternion();
+		FVector ThisTranslation(LocationX, LocationY, LocationZ);
+		FVector ThisScale(ScaleX, ScaleY, ScaleZ);
+
+		const double NormalizedWeight = Weight / TotalWeight;
+
+		Transform.Translation += ThisTranslation * NormalizedWeight;
+		Transform.Rotation    *= FQuat::Slerp(FQuat::Identity, ThisRotation, NormalizedWeight);
+		Transform.Scale       += ThisScale * NormalizedWeight;
+	}
+};
+
+struct FAdditiveAccumulationTask
+{
+	FQuaternionBlenderAccumulationBuffers* Buffers;
+
+	FAdditiveAccumulationTask(FQuaternionBlenderAccumulationBuffers* InBuffers)
+		: Buffers(InBuffers)
+	{}
+
+	void PreTask()
+	{
+		for (int32 Index = 0; Index < Buffers->Additives.Num(); ++Index)
+		{
+			Buffers->Additives[Index] = FQuatTransform{ 
+				FVector::ZeroVector,
+				FQuat::Identity,
+				FVector::ZeroVector
+			};
+		}
+	}
+
+	void ForEachEntity(
+		FMovieSceneBlendChannelID BlendChannel,
+		double LocationX, double LocationY, double LocationZ,
+		double RotationRoll, double RotationPitch, double RotationYaw,
+		double ScaleX, double ScaleY, double ScaleZ,
+		const double* WeightResult
+		)
+	{
+		FQuatTransform& Transform = Buffers->Additives[BlendChannel.ChannelID];
+
+		const double Weight = WeightResult ? *WeightResult : 1.0;
+		Transform.Translation += FVector(LocationX, LocationY, LocationZ) * Weight;
+		Transform.Rotation    *= FQuat::Slerp(FQuat::Identity, FRotator(RotationPitch, RotationYaw, RotationRoll).Quaternion(), Weight);
+		Transform.Scale       += FVector(ScaleX, ScaleY, ScaleZ) * Weight;
+	}
+};
+
+/** Task that combines all accumulated blends for any tracked property type that has blend inputs/outputs */
+struct FCombineTask
+{
+	FQuaternionBlenderAccumulationBuffers* Buffers;
+
+	FCombineTask(FQuaternionBlenderAccumulationBuffers* InBuffers)
+		: Buffers(InBuffers)
+	{}
+
+	void ForEachEntity(
+		FMovieSceneBlendChannelID BlendChannel,
+		double& LocationX, double& LocationY, double& LocationZ,
+		double& RotationRoll, double& RotationPitch, double& RotationYaw,
+		double& ScaleX, double& ScaleY, double& ScaleZ,
+		const FIntermediate3DTransform* InitialValue
+		)
+	{
+		const uint16 ChannelID = BlendChannel.ChannelID;
+		const double TotalWeight = Buffers->AbsoluteWeights.IsValidIndex(ChannelID) ? Buffers->AbsoluteWeights[ChannelID] : 1.0;
+
+		FVector  ResultTranslation = Buffers->Absolutes[ChannelID].Translation;
+		FQuat    ResultRotation    = Buffers->Absolutes[ChannelID].Rotation;
+		FVector  ResultScale       = Buffers->Absolutes[ChannelID].Scale;
+
+		if (TotalWeight >= 0.0 && TotalWeight < 1.0)
+		{
+			if (ensure(InitialValue))
+			{
+				ResultTranslation = FMath::Lerp(InitialValue->GetTranslation(),            ResultTranslation, TotalWeight);
+				ResultRotation    = FQuat::Slerp(InitialValue->GetRotation().Quaternion(), ResultRotation,    TotalWeight);
+				ResultScale       = FMath::Lerp(InitialValue->GetScale(),                  ResultScale,       TotalWeight);
+			}
+		}
+
+ 		ResultTranslation += Buffers->Additives[ChannelID].Translation;
+		ResultRotation    *= Buffers->Additives[ChannelID].Rotation;
+		ResultScale       += Buffers->Additives[ChannelID].Scale;
+
+		FRotator RotatorResult(ResultRotation);
+
+		LocationX = ResultTranslation.X;
+		LocationY = ResultTranslation.Y;
+		LocationZ = ResultTranslation.Z;
+		RotationRoll = RotatorResult.Roll;
+		RotationPitch = RotatorResult.Pitch;
+		RotationYaw = RotatorResult.Yaw;
+		ScaleX = ResultScale.X;
+		ScaleY = ResultScale.Y;
+		ScaleZ = ResultScale.Z;
+	}
+};
+
+
+} // namespace UE::MovieScene
+
+UMovieSceneQuaternionBlenderSystem::UMovieSceneQuaternionBlenderSystem(const FObjectInitializer& ObjInit)
+	: Super(ObjInit)
+{
+	if (HasAnyFlags(RF_ClassDefaultObject))
+	{
+		DefineImplicitPrerequisite(UFloatChannelEvaluatorSystem::StaticClass(), GetClass());
+		DefineImplicitPrerequisite(UDoubleChannelEvaluatorSystem::StaticClass(), GetClass());
+		DefineImplicitPrerequisite(UMovieSceneBaseValueEvaluatorSystem::StaticClass(), GetClass());
+		DefineImplicitPrerequisite(UMovieSceneQuaternionInterpolationRotationSystem::StaticClass(), GetClass());
+		DefineImplicitPrerequisite(UWeightAndEasingEvaluatorSystem::StaticClass(), GetClass());
+	}
+}
+
+void UMovieSceneQuaternionBlenderSystem::OnRun(FSystemTaskPrerequisites& InPrerequisites, FSystemSubsequentTasks& Subsequents)
+{
+	using namespace UE::MovieScene;
+
+	FBuiltInComponentTypes* BuiltInComponents = FBuiltInComponentTypes::Get();
+	FMovieSceneTracksComponentTypes* TracksComponents = FMovieSceneTracksComponentTypes::Get();
+
+	Buffers.Absolutes.SetNumZeroed(AllocatedBlendChannels.Num());
+	Buffers.Additives.SetNumZeroed(AllocatedBlendChannels.Num());
+
+	FSystemTaskPrerequisites Prereqs = InPrerequisites;
+
+	// Kick off additives first because they don't need normalized weighting
+	// Not handling additive from base yet because ideally we could do those in the AdditivesTask without caring about the base
+	FGraphEventRef AdditivesTask = FEntityTaskBuilder()
+	.Read(BuiltInComponents->BlendChannelInput)
+	.ReadAllOf(BuiltInComponents->DoubleResult[0], BuiltInComponents->DoubleResult[1], BuiltInComponents->DoubleResult[2]) // Translation components
+	.ReadAllOf(BuiltInComponents->DoubleResult[3], BuiltInComponents->DoubleResult[4], BuiltInComponents->DoubleResult[5]) // Rotation components
+	.ReadAllOf(BuiltInComponents->DoubleResult[6], BuiltInComponents->DoubleResult[7], BuiltInComponents->DoubleResult[8]) // Scale components
+	.ReadOptional(BuiltInComponents->WeightAndEasingResult)
+	.FilterAny({ BuiltInComponents->Tags.AdditiveBlend, BuiltInComponents->Tags.AdditiveFromBaseBlend })
+	.FilterAll({ GetBlenderTypeTag() })
+	.FilterNone({ BuiltInComponents->Tags.Ignored })
+	.template Dispatch_PerEntity<FAdditiveAccumulationTask>(&Linker->EntityManager, InPrerequisites, nullptr, &Buffers);
+
+	// Accumulate all weights so we can do a normalized lerp.
+	// Everything downstream depends on this task
+	FGraphEventRef CollectWeightsTask = FEntityTaskBuilder()
+	.Read(BuiltInComponents->BlendChannelInput)
+	.Read(BuiltInComponents->WeightAndEasingResult)
+	.FilterAll({ BuiltInComponents->Tags.AbsoluteBlend, GetBlenderTypeTag() })
+	.FilterNone({ BuiltInComponents->Tags.Ignored })
+	.template Dispatch_PerEntity<FAccumulateAbsoluteWeights>(&Linker->EntityManager, InPrerequisites, nullptr, &Buffers);
+
+	// Accumulate initial values
+	//FGraphEventRef CollectInitialValuesTask = FEntityTaskBuilder()
+	//.Read(BuiltInComponents->BlendChannelOutput)
+	//.Read(TracksComponents->ComponentTransform.InitialValue)
+	//.FilterAll({ GetBlenderTypeTag() })
+	//.template Dispatch_PerEntity<FCollectInitialValues>(&Linker->EntityManager, InPrerequisites, &Subsequents, &Buffers);
+
+	if (CollectWeightsTask)
+	{
+		InPrerequisites.AddRootTask(CollectWeightsTask);
+	}
+	//if (CollectInitialValuesTask)
+	//{
+	//	InPrerequisites.AddRootTask(CollectInitialValuesTask);
+	//}
+
+	// Next we accumulate normalized transforms
+	FGraphEventRef AbsolutesTask = FEntityTaskBuilder()
+	.Read(BuiltInComponents->BlendChannelInput)
+	.ReadAllOf(BuiltInComponents->DoubleResult[0], BuiltInComponents->DoubleResult[1], BuiltInComponents->DoubleResult[2]) // Translation components
+	.ReadAllOf(BuiltInComponents->DoubleResult[3], BuiltInComponents->DoubleResult[4], BuiltInComponents->DoubleResult[5]) // Rotation components
+	.ReadAllOf(BuiltInComponents->DoubleResult[6], BuiltInComponents->DoubleResult[7], BuiltInComponents->DoubleResult[8]) // Scale components
+	.ReadOptional(BuiltInComponents->WeightAndEasingResult)
+	.FilterAll({ BuiltInComponents->Tags.AbsoluteBlend, GetBlenderTypeTag() })
+	.FilterNone({ BuiltInComponents->Tags.Ignored })
+	.template Dispatch_PerEntity<FAbsoluteAccumulationTask>(&Linker->EntityManager, InPrerequisites, nullptr, &Buffers);
+
+	if (AbsolutesTask)
+	{
+		InPrerequisites.AddRootTask(AbsolutesTask);
+	}
+	if (AdditivesTask)
+	{
+		InPrerequisites.AddRootTask(AdditivesTask);
+	}
+
+	// Now blend the results
+	FEntityTaskBuilder()
+	.Read(BuiltInComponents->BlendChannelOutput)
+	.WriteAllOf(BuiltInComponents->DoubleResult[0], BuiltInComponents->DoubleResult[1], BuiltInComponents->DoubleResult[2]) // Translation components
+	.WriteAllOf(BuiltInComponents->DoubleResult[3], BuiltInComponents->DoubleResult[4], BuiltInComponents->DoubleResult[5]) // Rotation components
+	.WriteAllOf(BuiltInComponents->DoubleResult[6], BuiltInComponents->DoubleResult[7], BuiltInComponents->DoubleResult[8]) // Scale components
+	.ReadOptional(TracksComponents->ComponentTransform.InitialValue)
+	.FilterAll({ GetBlenderTypeTag() })
+	.template Dispatch_PerEntity<FCombineTask>(&Linker->EntityManager, InPrerequisites, &Subsequents, &Buffers);
+}
+
+FGraphEventRef UMovieSceneQuaternionBlenderSystem::DispatchDecomposeTask(const UE::MovieScene::FValueDecompositionParams& Params, UE::MovieScene::FAlignedDecomposedValue* Output)
+{
+	return nullptr;
+}
\ No newline at end of file
diff --git a/Engine/Source/Runtime/MovieSceneTracks/Private/Tracks/MovieScene3DTransformTrack.cpp b/Engine/Source/Runtime/MovieSceneTracks/Private/Tracks/MovieScene3DTransformTrack.cpp
--- a/Engine/Source/Runtime/MovieSceneTracks/Private/Tracks/MovieScene3DTransformTrack.cpp
+++ b/Engine/Source/Runtime/MovieSceneTracks/Private/Tracks/MovieScene3DTransformTrack.cpp
@@ -6,6 +6,8 @@
 #include "Compilation/MovieSceneSegmentCompiler.h"
 #include "Evaluation/MovieSceneEvaluationTrack.h"
 #include "Channels/MovieSceneChannelProxy.h"
+#include "Systems/MovieSceneQuaternionBlenderSystem.h"
+#include "Systems/MovieScenePiecewiseDoubleBlenderSystem.h"
 #include "MovieSceneCommonHelpers.h"
 #include "Algo/BinarySearch.h"
 
@@ -36,6 +38,30 @@ UMovieSceneSection* UMovieScene3DTransformTrack::CreateNewSection()
 	return NewObject<UMovieScene3DTransformSection>(this, NAME_None, RF_Transactional);
 }
 
+TSubclassOf<UMovieSceneBlenderSystem> UMovieScene3DTransformTrack::GetBlenderSystem() const
+{
+	if (BlenderSystemClass == nullptr)
+	{
+		return UMovieScenePiecewiseDoubleBlenderSystem::StaticClass();
+	}
+	return BlenderSystemClass;
+}
+
+void UMovieScene3DTransformTrack::SetBlenderSystem(TSubclassOf<UMovieSceneBlenderSystem> InBlenderSystemClass)
+{
+	if (InBlenderSystemClass == UMovieScenePiecewiseDoubleBlenderSystem::StaticClass())
+	{
+		InBlenderSystemClass = nullptr;
+	}
+
+	BlenderSystemClass = InBlenderSystemClass;
+}
+
+void UMovieScene3DTransformTrack::GetSupportedBlenderSystems(TArray<TSubclassOf<UMovieSceneBlenderSystem>>& OutSystemClasses) const
+{
+	OutSystemClasses.Add(UMovieSceneQuaternionBlenderSystem::StaticClass());
+	OutSystemClasses.Add(UMovieScenePiecewiseDoubleBlenderSystem::StaticClass());
+}
 
 #if WITH_EDITOR
 
diff --git a/Engine/Source/Runtime/MovieSceneTracks/Public/Systems/MovieScenePiecewiseDoubleBlenderSystem.h b/Engine/Source/Runtime/MovieSceneTracks/Public/Systems/MovieScenePiecewiseDoubleBlenderSystem.h
--- a/Engine/Source/Runtime/MovieSceneTracks/Public/Systems/MovieScenePiecewiseDoubleBlenderSystem.h
+++ b/Engine/Source/Runtime/MovieSceneTracks/Public/Systems/MovieScenePiecewiseDoubleBlenderSystem.h
@@ -84,7 +84,7 @@ struct FAccumulationBuffers
 } // namespace MovieScene
 } // namespace UE
 
-UCLASS()
+UCLASS(DisplayName="Weighted per-channel", meta=(Tooltip="Blends each channel of this object's transform as separate scalar components. Useful for blending to/from over-rotated objects (ie, 0 < rotation > 360."))
 class MOVIESCENETRACKS_API UMovieScenePiecewiseDoubleBlenderSystem : public UMovieSceneBlenderSystem, public IMovieSceneValueDecomposer
 {
 public:
diff --git a/Engine/Source/Runtime/MovieSceneTracks/Public/Systems/MovieSceneQuaternionBlenderSystem.h b/Engine/Source/Runtime/MovieSceneTracks/Public/Systems/MovieSceneQuaternionBlenderSystem.h
new file mode 100644
--- /dev/null
+++ b/Engine/Source/Runtime/MovieSceneTracks/Public/Systems/MovieSceneQuaternionBlenderSystem.h
@@ -0,0 +1,47 @@
+// Copyright Epic Games, Inc. All Rights Reserved.
+
+#pragma once
+
+#include "Containers/SortedMap.h"
+#include "EntitySystem/MovieSceneBlenderSystem.h"
+#include "EntitySystem/MovieSceneCachedEntityFilterResult.h"
+#include "EntitySystem/MovieSceneDecompositionQuery.h"
+#include "EntitySystem/MovieSceneEntitySystem.h"
+
+#include "MovieSceneQuaternionBlenderSystem.generated.h"
+
+namespace UE::MovieScene
+{
+	struct FQuatTransform
+	{
+		FVector Translation;
+		FQuat Rotation;
+		FVector Scale;
+	};
+
+	struct FQuaternionBlenderAccumulationBuffers
+	{
+		TSparseArray<double> AbsoluteWeights;
+		TArray<FQuatTransform> InitialValues;
+		TArray<FQuatTransform> Absolutes;
+		TArray<FQuatTransform> Additives;
+	};
+}
+
+UCLASS(DisplayName="Shortest Rotation (quaternion slerp)", meta=(Tooltip="Uses a quaternion spherical interpolation (slerp) to blend between transforms using the shortest rotation path. Does not support over-rotation."))
+class MOVIESCENETRACKS_API UMovieSceneQuaternionBlenderSystem : public UMovieSceneBlenderSystem, public IMovieSceneValueDecomposer
+{
+public:
+	GENERATED_BODY()
+
+	UMovieSceneQuaternionBlenderSystem(const FObjectInitializer& ObjInit);
+
+	virtual void OnRun(FSystemTaskPrerequisites& InPrerequisites, FSystemSubsequentTasks& Subsequents) override;
+
+	virtual FGraphEventRef DispatchDecomposeTask(const UE::MovieScene::FValueDecompositionParams& Params, UE::MovieScene::FAlignedDecomposedValue* Output) override;
+
+private:
+
+	UE::MovieScene::FQuaternionBlenderAccumulationBuffers Buffers;
+};
+
diff --git a/Engine/Source/Runtime/MovieSceneTracks/Public/Tracks/MovieScene3DTransformTrack.h b/Engine/Source/Runtime/MovieSceneTracks/Public/Tracks/MovieScene3DTransformTrack.h
--- a/Engine/Source/Runtime/MovieSceneTracks/Public/Tracks/MovieScene3DTransformTrack.h
+++ b/Engine/Source/Runtime/MovieSceneTracks/Public/Tracks/MovieScene3DTransformTrack.h
@@ -6,6 +6,7 @@
 #include "UObject/ObjectMacros.h"
 #include "Tracks/MovieScenePropertyTrack.h"
 #include "Sections/MovieScene3DTransformSection.h"
+#include "EntitySystem/IMovieSceneBlenderSystemSupport.h"
 #include "MovieScene3DTransformTrack.generated.h"
 
 enum class EMovieSceneTransformChannel : uint32;
@@ -55,6 +56,7 @@ struct FTrajectoryKey
 UCLASS(MinimalAPI)
 class UMovieScene3DTransformTrack
 	: public UMovieScenePropertyTrack
+	, public IMovieSceneBlenderSystemSupport
 {
 	GENERATED_UCLASS_BODY()
 
@@ -69,9 +71,19 @@ public:
 	virtual bool CanRename() const override { return true; }
 #endif
 
+	// IMovieSceneBlenderSystemSupport
+	TSubclassOf<UMovieSceneBlenderSystem> GetBlenderSystem() const override;
+	void SetBlenderSystem(TSubclassOf<UMovieSceneBlenderSystem> BlenderSystemClass) override;
+	void GetSupportedBlenderSystems(TArray<TSubclassOf<UMovieSceneBlenderSystem>>& OutSystemClasses) const override;
+
 #if WITH_EDITOR
 
 	MOVIESCENETRACKS_API TArray<FTrajectoryKey> GetTrajectoryData(FFrameNumber Time, int32 MaxNumDataPoints, TRange<FFrameNumber>) const;
 
 #endif
+
+private:
+
+	/** User-defined blender system to use for this track */
+	TSubclassOf<UMovieSceneBlenderSystem> BlenderSystemClass;
 };
