Sequencer: Transform track's blender system class is now a uproperty to support serialization

diff --git a/Engine/Source/Runtime/MovieSceneTracks/Public/Tracks/MovieScene3DTransformTrack.h b/Engine/Source/Runtime/MovieSceneTracks/Public/Tracks/MovieScene3DTransformTrack.h
--- a/Engine/Source/Runtime/MovieSceneTracks/Public/Tracks/MovieScene3DTransformTrack.h
+++ b/Engine/Source/Runtime/MovieSceneTracks/Public/Tracks/MovieScene3DTransformTrack.h
@@ -85,5 +85,6 @@ public:
 private:
 
 	/** User-defined blender system to use for this track */
+	UPROPERTY()
 	TSubclassOf<UMovieSceneBlenderSystem> BlenderSystemClass;
 };
