Sequencer: Added mutable accessors to some section channels

diff --git a/Engine/Source/Runtime/MovieSceneTracks/Public/Sections/MovieSceneColorSection.h b/Engine/Source/Runtime/MovieSceneTracks/Public/Sections/MovieSceneColorSection.h
--- a/Engine/Source/Runtime/MovieSceneTracks/Public/Sections/MovieSceneColorSection.h
+++ b/Engine/Source/Runtime/MovieSceneTracks/Public/Sections/MovieSceneColorSection.h
@@ -51,15 +51,19 @@ class UMovieSceneColorSection
 public:
 
 	/** Access the underlying generation function for the red component of this section */
+	FMovieSceneFloatChannel& GetRedChannel() { return RedCurve; }
 	const FMovieSceneFloatChannel& GetRedChannel() const { return RedCurve; }
 
 	/** Access the underlying generation function for the green component of this section */
+	FMovieSceneFloatChannel& GetGreenChannel() { return GreenCurve; }
 	const FMovieSceneFloatChannel& GetGreenChannel() const { return GreenCurve; }
 
 	/** Access the underlying generation function for the blue component of this section */
+	FMovieSceneFloatChannel& GetBlueChannel() { return BlueCurve; }
 	const FMovieSceneFloatChannel& GetBlueChannel() const { return BlueCurve; }
 
 	/** Access the underlying generation function for the alpha component of this section */
+	FMovieSceneFloatChannel& GetAlphaChannel() { return AlphaCurve; }
 	const FMovieSceneFloatChannel& GetAlphaChannel() const { return AlphaCurve; }
 
 protected:
diff --git a/Engine/Source/Runtime/MovieSceneTracks/Public/Sections/MovieSceneDoubleSection.h b/Engine/Source/Runtime/MovieSceneTracks/Public/Sections/MovieSceneDoubleSection.h
--- a/Engine/Source/Runtime/MovieSceneTracks/Public/Sections/MovieSceneDoubleSection.h
+++ b/Engine/Source/Runtime/MovieSceneTracks/Public/Sections/MovieSceneDoubleSection.h
@@ -29,6 +29,7 @@ public:
 	/**
 	 * Public access to this section's internal data function
 	 */
+	FMovieSceneDoubleChannel& GetChannel() { return DoubleCurve; }
 	const FMovieSceneDoubleChannel& GetChannel() const { return DoubleCurve; }
 
 protected:
diff --git a/Engine/Source/Runtime/MovieSceneTracks/Public/Sections/MovieSceneFloatSection.h b/Engine/Source/Runtime/MovieSceneTracks/Public/Sections/MovieSceneFloatSection.h
--- a/Engine/Source/Runtime/MovieSceneTracks/Public/Sections/MovieSceneFloatSection.h
+++ b/Engine/Source/Runtime/MovieSceneTracks/Public/Sections/MovieSceneFloatSection.h
@@ -31,6 +31,7 @@ public:
 	/**
 	 * Public access to this section's internal data function
 	 */
+	FMovieSceneFloatChannel& GetChannel() { return FloatCurve; }
 	const FMovieSceneFloatChannel& GetChannel() const { return FloatCurve; }
 
 protected:
