Sequencer: Crash Fix for null animation in animation section

diff --git a/Engine/Source/Runtime/MovieSceneTracks/Private/Sections/MovieSceneSkeletalAnimationSection.cpp b/Engine/Source/Runtime/MovieSceneTracks/Private/Sections/MovieSceneSkeletalAnimationSection.cpp
--- a/Engine/Source/Runtime/MovieSceneTracks/Private/Sections/MovieSceneSkeletalAnimationSection.cpp
+++ b/Engine/Source/Runtime/MovieSceneTracks/Private/Sections/MovieSceneSkeletalAnimationSection.cpp
@@ -56,39 +56,43 @@ double FMovieSceneSkeletalAnimationParams::MapTimeToAnimation(const UMovieSceneS
 double FMovieSceneSkeletalAnimationParams::MapTimeToAnimation(FFrameNumber InSectionStartTime, FFrameNumber InSectionEndTime, FFrameTime InPosition, FFrameRate InFrameRate) const
 {
 	// Get Animation Length and frame time
-	const FFrameTime AnimationLength = GetSequenceLength() * InFrameRate;
-	const int32 LengthInFrames = AnimationLength.FrameNumber.Value + (int)(AnimationLength.GetSubFrame() + 0.5f) + 1;
+	if (Animation)
+	{
+		const FFrameTime AnimationLength = GetSequenceLength() * InFrameRate;
+		const int32 LengthInFrames = AnimationLength.FrameNumber.Value + (int)(AnimationLength.GetSubFrame() + 0.5f) + 1;
 
-	// we only play end if we are not looping, and assuming we are looping if Length is greater than default length;
-	const bool bLooping = (InSectionEndTime.Value - InSectionStartTime.Value + StartFrameOffset + EndFrameOffset) > LengthInFrames;	
+		// we only play end if we are not looping, and assuming we are looping if Length is greater than default length;
+		const bool bLooping = (InSectionEndTime.Value - InSectionStartTime.Value + StartFrameOffset + EndFrameOffset) > LengthInFrames;
 
-	// Make sure InPosition FrameTime doesn't underflow InSectionStartTime or overflow InSectionEndTime
-	InPosition = FMath::Clamp(InPosition, FFrameTime(InSectionStartTime), FFrameTime(InSectionEndTime - 1));
-	
-	// Gather helper values
-	const float SectionPlayRate = PlayRate * Animation->RateScale;
-	const float AnimPlayRate = FMath::IsNearlyZero(SectionPlayRate) ? 1.0f : SectionPlayRate;
-	const double SeqLength = GetSequenceLength() - InFrameRate.AsSeconds(StartFrameOffset + EndFrameOffset);
+		// Make sure InPosition FrameTime doesn't underflow InSectionStartTime or overflow InSectionEndTime
+		InPosition = FMath::Clamp(InPosition, FFrameTime(InSectionStartTime), FFrameTime(InSectionEndTime - 1));
 
-	// The Time from the beginning of InSectionStartTime to InPosition in seconds
-	double SecondsFromSectionStart = FFrameTime::FromDecimal((InPosition - InSectionStartTime).AsDecimal() * AnimPlayRate) / InFrameRate;
+		// Gather helper values
+		const float SectionPlayRate = PlayRate * Animation->RateScale;
+		const float AnimPlayRate = FMath::IsNearlyZero(SectionPlayRate) ? 1.0f : SectionPlayRate;
+		const double SeqLength = GetSequenceLength() - InFrameRate.AsSeconds(StartFrameOffset + EndFrameOffset);
 
-	// Logic for reversed animation
-	if (bReverse)
-	{
-		// Duration of this section 
-		double SectionDuration = (((InSectionEndTime - InSectionStartTime) * AnimPlayRate) / InFrameRate);
-		SecondsFromSectionStart = SectionDuration - SecondsFromSectionStart;
-	}
-	// Make sure Seconds is in range
-	if (SeqLength > 0.0 && (bLooping || !FMath::IsNearlyEqual(SecondsFromSectionStart, SeqLength, 1e-4)))
-	{
-		SecondsFromSectionStart = FMath::Fmod(SecondsFromSectionStart, SeqLength);
-	}
+		// The Time from the beginning of InSectionStartTime to InPosition in seconds
+		double SecondsFromSectionStart = FFrameTime::FromDecimal((InPosition - InSectionStartTime).AsDecimal() * AnimPlayRate) / InFrameRate;
 
-	// Add the StartFrameOffset and FirstLoopStartFrameOffset to the current seconds in the section to get the right animation frame
-	SecondsFromSectionStart += InFrameRate.AsSeconds(StartFrameOffset) + InFrameRate.AsSeconds(FirstLoopStartFrameOffset);
-	return SecondsFromSectionStart;
+		// Logic for reversed animation
+		if (bReverse)
+		{
+			// Duration of this section 
+			double SectionDuration = (((InSectionEndTime - InSectionStartTime) * AnimPlayRate) / InFrameRate);
+			SecondsFromSectionStart = SectionDuration - SecondsFromSectionStart;
+		}
+		// Make sure Seconds is in range
+		if (SeqLength > 0.0 && (bLooping || !FMath::IsNearlyEqual(SecondsFromSectionStart, SeqLength, 1e-4)))
+		{
+			SecondsFromSectionStart = FMath::Fmod(SecondsFromSectionStart, SeqLength);
+		}
+
+		// Add the StartFrameOffset and FirstLoopStartFrameOffset to the current seconds in the section to get the right animation frame
+		SecondsFromSectionStart += InFrameRate.AsSeconds(StartFrameOffset) + InFrameRate.AsSeconds(FirstLoopStartFrameOffset);
+		return SecondsFromSectionStart;
+	}
+	return 0.0;
 }
 
 UMovieSceneSkeletalAnimationSection::UMovieSceneSkeletalAnimationSection( const FObjectInitializer& ObjectInitializer )
diff --git a/Engine/Source/Runtime/MovieSceneTracks/Private/Systems/MovieSceneSkeletalAnimationSystem.cpp b/Engine/Source/Runtime/MovieSceneTracks/Private/Systems/MovieSceneSkeletalAnimationSystem.cpp
--- a/Engine/Source/Runtime/MovieSceneTracks/Private/Systems/MovieSceneSkeletalAnimationSystem.cpp
+++ b/Engine/Source/Runtime/MovieSceneTracks/Private/Systems/MovieSceneSkeletalAnimationSystem.cpp
@@ -282,6 +282,10 @@ struct FGatherSkeletalAnimations
 			// Calculate the time at which to evaluate the animation
 			const UMovieSceneSkeletalAnimationSection* AnimSection = SkeletalAnimation.Section;
 			const FMovieSceneSkeletalAnimationParams& AnimParams = AnimSection->Params;
+			if (AnimParams.Animation == nullptr)
+			{
+				continue;
+			}
 			const float EvalTime = AnimParams.MapTimeToAnimation(AnimSection, Context.GetTime(), Context.GetFrameRate());
 			const float PreviousEvalTime = AnimParams.MapTimeToAnimation(AnimSection, Context.GetPreviousTime(), Context.GetFrameRate());
 
