Sequencer: Added a way to ignore hierarchical bias so that tracks can always be relevant

diff --git a/Engine/Source/Runtime/MovieScene/Private/EntitySystem/BuiltInComponentTypes.cpp b/Engine/Source/Runtime/MovieScene/Private/EntitySystem/BuiltInComponentTypes.cpp
--- a/Engine/Source/Runtime/MovieScene/Private/EntitySystem/BuiltInComponentTypes.cpp
+++ b/Engine/Source/Runtime/MovieScene/Private/EntitySystem/BuiltInComponentTypes.cpp
@@ -133,6 +133,8 @@ FBuiltInComponentTypes::FBuiltInComponentTypes()
 	ComponentRegistry->NewComponentType(&Interrogation.OutputKey, TEXT("Interrogation Output"));
 
 	Tags.RestoreState            = ComponentRegistry->NewTag(TEXT("Is Restore State Entity"));
+	Tags.IgnoreHierarchicalBias  = ComponentRegistry->NewTag(TEXT("Ignore Hierarchical Bias"));
+
 	Tags.AbsoluteBlend           = ComponentRegistry->NewTag(TEXT("Is Absolute Blend"));
 	Tags.RelativeBlend           = ComponentRegistry->NewTag(TEXT("Is Relative Blend"));
 	Tags.AdditiveBlend           = ComponentRegistry->NewTag(TEXT("Is Additive Blend"));
diff --git a/Engine/Source/Runtime/MovieScene/Public/EntitySystem/BuiltInComponentTypes.h b/Engine/Source/Runtime/MovieScene/Public/EntitySystem/BuiltInComponentTypes.h
--- a/Engine/Source/Runtime/MovieScene/Public/EntitySystem/BuiltInComponentTypes.h
+++ b/Engine/Source/Runtime/MovieScene/Public/EntitySystem/BuiltInComponentTypes.h
@@ -376,6 +376,9 @@ public:
 		// A tag specifying that an entity wants to restore state on completioon
 		FComponentTypeID RestoreState;
 
+		// A tag that specifies this entity should always contribute to the output, regardless of hbias
+		FComponentTypeID IgnoreHierarchicalBias;
+
 		FComponentTypeID AbsoluteBlend;
 		FComponentTypeID RelativeBlend;
 		FComponentTypeID AdditiveBlend;
diff --git a/Engine/Source/Runtime/MovieSceneTracks/Private/Systems/MovieSceneHierarchicalBiasSystem.cpp b/Engine/Source/Runtime/MovieSceneTracks/Private/Systems/MovieSceneHierarchicalBiasSystem.cpp
--- a/Engine/Source/Runtime/MovieSceneTracks/Private/Systems/MovieSceneHierarchicalBiasSystem.cpp
+++ b/Engine/Source/Runtime/MovieSceneTracks/Private/Systems/MovieSceneHierarchicalBiasSystem.cpp
@@ -44,10 +44,20 @@ struct FHierarchicalBiasTask
 		return MaxBiasByChannel.Num() != 0;
 	}
 
-	void ForEachAllocation(const FEntityAllocation* Allocation, TRead<FMovieSceneEntityID> EntityIDs, TRead<FRootInstanceHandle> RootInstanceHandles, TRead<FMovieSceneBlendChannelID> BlendChannels, TReadOptional<int16> OptHBiases)
+	void ForEachAllocation(FEntityAllocationIteratorItem Iterator, TRead<FMovieSceneEntityID> EntityIDs, TRead<FRootInstanceHandle> RootInstanceHandles, TRead<FMovieSceneBlendChannelID> BlendChannels, TReadOptional<int16> OptHBiases)
 	{
-		const int32 Num = Allocation->Num();
-		if (OptHBiases)
+		const int32 Num = Iterator.GetAllocation()->Num();
+
+		const bool bIgnoreBias = Iterator.GetAllocationType().Contains(FBuiltInComponentTypes::Get()->Tags.IgnoreHierarchicalBias);
+		if (bIgnoreBias)
+		{
+			for (int32 Index = 0; Index < Num; ++Index)
+			{
+				FBlendChannelSequenceKey Key{ RootInstanceHandles[Index], BlendChannels[Index] };
+				ActiveContributorsByChannel.Add(Key, EntityIDs[Index]);
+			}
+		}
+		else if (OptHBiases)
 		{
 			for (int32 Index = 0; Index < Num; ++Index)
 			{
