Sequencer: Fix for First Frame Loop offset making animation go out of range, which messes with montage playback

diff --git a/Engine/Source/Runtime/MovieSceneTracks/Private/Sections/MovieSceneSkeletalAnimationSection.cpp b/Engine/Source/Runtime/MovieSceneTracks/Private/Sections/MovieSceneSkeletalAnimationSection.cpp
--- a/Engine/Source/Runtime/MovieSceneTracks/Private/Sections/MovieSceneSkeletalAnimationSection.cpp
+++ b/Engine/Source/Runtime/MovieSceneTracks/Private/Sections/MovieSceneSkeletalAnimationSection.cpp
@@ -82,14 +82,17 @@ double FMovieSceneSkeletalAnimationParams::MapTimeToAnimation(FFrameNumber InSec
 			double SectionDuration = (((InSectionEndTime - InSectionStartTime) * AnimPlayRate) / InFrameRate);
 			SecondsFromSectionStart = SectionDuration - SecondsFromSectionStart;
 		}
+
+		SecondsFromSectionStart += InFrameRate.AsSeconds(FirstLoopStartFrameOffset);
+
 		// Make sure Seconds is in range
 		if (SeqLength > 0.0 && (bLooping || !FMath::IsNearlyEqual(SecondsFromSectionStart, SeqLength, 1e-4)))
 		{
 			SecondsFromSectionStart = FMath::Fmod(SecondsFromSectionStart, SeqLength);
 		}
 
-		// Add the StartFrameOffset and FirstLoopStartFrameOffset to the current seconds in the section to get the right animation frame
-		SecondsFromSectionStart += InFrameRate.AsSeconds(StartFrameOffset) + InFrameRate.AsSeconds(FirstLoopStartFrameOffset);
+		// Add the StartFrameOffset to the current seconds in the section to get the right animation frame
+		SecondsFromSectionStart += InFrameRate.AsSeconds(StartFrameOffset);
 		return SecondsFromSectionStart;
 	}
 	return 0.0;
