- Added UDataLayerManager which is local to each WorldPartition.

diff --git a/Engine/Source/Runtime/Engine/Public/WorldPartition/WorldPartitionRuntimeSpatialHash.h b/Engine/Source/Runtime/Engine/Public/WorldPartition/WorldPartitionRuntimeSpatialHash.h
--- a/Engine/Source/Runtime/Engine/Public/WorldPartition/WorldPartitionRuntimeSpatialHash.h
+++ b/Engine/Source/Runtime/Engine/Public/WorldPartition/WorldPartitionRuntimeSpatialHash.h
@@ -141,10 +141,10 @@ struct ENGINE_API FSpatialHashStreamingGrid
 	// Used by PIE/Game
 	int64 GetCellSize(int32 Level) const;
 	void GetCells(const FWorldPartitionStreamingQuerySource& QuerySource, TSet<const UWorldPartitionRuntimeCell*>& OutCells, bool bEnableZCulling, FWorldPartitionQueryCache* QueryCache = nullptr) const;
-	void GetCells(const TArray<FWorldPartitionStreamingSource>& Sources, const class UDataLayerSubsystem* DataLayerSubsystem, UWorldPartitionRuntimeHash::FStreamingSourceCells& OutActivateCells, UWorldPartitionRuntimeHash::FStreamingSourceCells& OutLoadCells, bool bEnableZCulling) const;
-	void GetNonSpatiallyLoadedCells(const UDataLayerSubsystem* DataLayerSubsystem, TSet<const UWorldPartitionRuntimeCell*>& OutActivateCells, TSet<const UWorldPartitionRuntimeCell*>& OutLoadCells) const;
-	void Draw2D(UCanvas* Canvas, UWorld* World, const TArray<FWorldPartitionStreamingSource>& Sources, const FBox& Region, const FBox2D& GridScreenBounds, TFunctionRef<FVector2D(const FVector2D&)> WorldToScreen) const;
-	void Draw3D(UWorld* World, const TArray<FWorldPartitionStreamingSource>& Sources, const FTransform& Transform) const;
+	void GetCells(const TArray<FWorldPartitionStreamingSource>& Sources, UWorldPartitionRuntimeHash::FStreamingSourceCells& OutActivateCells, UWorldPartitionRuntimeHash::FStreamingSourceCells& OutLoadCells, bool bEnableZCulling) const;
+	void GetNonSpatiallyLoadedCells(TSet<const UWorldPartitionRuntimeCell*>& OutActivateCells, TSet<const UWorldPartitionRuntimeCell*>& OutLoadCells) const;
+	void Draw2D(UCanvas* Canvas, const class UWorldPartitionRuntimeSpatialHash* Owner, const TArray<FWorldPartitionStreamingSource>& Sources, const FBox& Region, const FBox2D& GridScreenBounds, TFunctionRef<FVector2D(const FVector2D&)> WorldToScreen) const;
+	void Draw3D(const class UWorldPartitionRuntimeSpatialHash* Owner, const TArray<FWorldPartitionStreamingSource>& Sources, const FTransform& Transform) const;
 	void ForEachRuntimeCell(TFunctionRef<bool(const UWorldPartitionRuntimeCell*)> Func) const;
 	const FSquare2DGridHelper& GetGridHelper() const;
 
@@ -153,7 +153,7 @@ private:
 	void ForEachLayerCell(const FGridCellCoord& Coords, TFunctionRef<void(const FSpatialHashStreamingGridLayerCell*)> Func) const;
 	void DrawStreamingSource2D(UCanvas* Canvas, const FSphericalSector& Shape, TFunctionRef<FVector2D(const FVector2D&)> WorldToScreen, const FColor& Color) const;
 	void DrawStreamingSource3D(UWorld* World, const FSphericalSector& Shape, const FTransform& Transform, const FColor& Color) const;
-	void GetFilteredCellsForDebugDraw(const FSpatialHashStreamingGridLayerCell* LayerCell, const UDataLayerSubsystem* DataLayerSubsystem, TArray<const UWorldPartitionRuntimeCell*>& FilteredCells) const;
+	void GetFilteredCellsForDebugDraw(const FSpatialHashStreamingGridLayerCell* LayerCell, TArray<const UWorldPartitionRuntimeCell*>& FilteredCells) const;
 	EWorldPartitionRuntimeCellVisualizeMode GetStreamingCellVisualizeMode() const;
 	mutable FSquare2DGridHelper* GridHelper;
 
diff --git a/Engine/Source/Runtime/Engine/Public/WorldPartition/WorldPartitionSubsystem.h b/Engine/Source/Runtime/Engine/Public/WorldPartition/WorldPartitionSubsystem.h
--- a/Engine/Source/Runtime/Engine/Public/WorldPartition/WorldPartitionSubsystem.h
+++ b/Engine/Source/Runtime/Engine/Public/WorldPartition/WorldPartitionSubsystem.h
@@ -20,9 +20,6 @@ class ENGINE_API UWorldPartitionSubsystem : public UTickableWorldSubsystem
 {
 	GENERATED_BODY()
 
-protected:
-	virtual bool DoesSupportWorldType(const EWorldType::Type WorldType) const override;
-
 public:
 	UWorldPartitionSubsystem();
 
@@ -65,14 +62,19 @@ public:
 	DECLARE_DELEGATE_RetVal_OneParam(bool, FWorldPartitionStreamingSourceProviderFilter, const IWorldPartitionStreamingSourceProvider*);
 	FWorldPartitionStreamingSourceProviderFilter& OnIsStreamingSourceProviderFiltered() { return IsStreamingSourceProviderFiltered; }
 
-#if WITH_EDITOR
 	void ForEachWorldPartition(TFunctionRef<bool(UWorldPartition*)> Func);
 
+#if WITH_EDITOR
 	static bool IsRunningConvertWorldPartitionCommandlet();
 
 	TSet<FWorldPartitionActorDesc*> SelectedActorDescs;
 #endif
 
+protected:
+	//~ Begin USubsystem Interface.
+	virtual bool DoesSupportWorldType(const EWorldType::Type WorldType) const override;
+	//~ End USubsystem Interface.
+
 private:
 
 	void OnWorldPartitionInitialized(UWorldPartition* InWorldPartition);
diff --git a/Engine/Source/Runtime/MovieSceneTracks/Private/Systems/MovieSceneDataLayerSystem.cpp b/Engine/Source/Runtime/MovieSceneTracks/Private/Systems/MovieSceneDataLayerSystem.cpp
--- a/Engine/Source/Runtime/MovieSceneTracks/Private/Systems/MovieSceneDataLayerSystem.cpp
+++ b/Engine/Source/Runtime/MovieSceneTracks/Private/Systems/MovieSceneDataLayerSystem.cpp
@@ -13,7 +13,7 @@
 #include "MovieSceneTracksComponentTypes.h"
 #include "ProfilingDebugging/CsvProfiler.h"
 #include "Sections/MovieSceneDataLayerSection.h"
-#include "WorldPartition/DataLayer/DataLayerSubsystem.h"
+#include "WorldPartition/DataLayer/DataLayerManager.h"
 #include "WorldPartition/DataLayer/DataLayerInstance.h"
 #include "WorldPartition/WorldPartitionSubsystem.h"
 #include "WorldPartition/WorldPartitionRuntimeCell.h"
@@ -67,14 +67,14 @@ struct FPreAnimatedDataLayerStorage
 	void OnGroupDestroyed(FPreAnimatedStorageGroupHandle Group) override;
 
 	/** Make an entry for the specified data layer */
-	FPreAnimatedStateEntry MakeEntry(UDataLayerInstance* InDataLayer);
+	FPreAnimatedStateEntry MakeEntry(const UDataLayerInstance* InDataLayer);
 
 	/** Save the value of a data layer. Should only be used for runtime / PIE worlds */
-	void SavePreAnimatedState(UDataLayerInstance* DataLayer, UDataLayerSubsystem* SubSystem);
+	void SavePreAnimatedState(const UDataLayerInstance* DataLayer, UDataLayerManager* DataLayerManager);
 
 #if WITH_EDITOR
 	/** Save the value of a data layer. Should only be used for editor worlds */
-	void SavePreAnimatedStateInEditor(UDataLayerInstance* DataLayer);
+	void SavePreAnimatedStateInEditor(const UDataLayerInstance* DataLayer);
 #endif
 
 private:
@@ -106,7 +106,7 @@ struct FDesiredLayerStates
 {
 	bool IsEmpty() const;
 	void Reset();
-	EDataLayerUpdateFlags Apply(FPreAnimatedDataLayerStorage* PreAnimatedStorage, UDataLayerSubsystem* DataLayerSubsystem, UWorldPartitionSubsystem* WorldPartitionSubsystem);
+	EDataLayerUpdateFlags Apply(FPreAnimatedDataLayerStorage* PreAnimatedStorage, UDataLayerManager* DataLayerManager, UWorldPartitionSubsystem* WorldPartitionSubsystem);
 #if WITH_EDITOR
 	void ApplyInEditor(FPreAnimatedDataLayerStorage* PreAnimatedStorage, UDataLayerEditorSubsystem* EditorSubSystem);
 #endif
@@ -136,9 +136,9 @@ void FPreAnimatedDataLayerStorageTraits::RestorePreAnimatedValue(const TObjectKe
 		}
 		else
 #endif
-		if (UDataLayerSubsystem* SubSystem = World->GetSubsystem<UDataLayerSubsystem>())
+		if (UDataLayerManager* DataLayerManager = World->GetDataLayerManager())
 		{
-			SubSystem->SetDataLayerRuntimeState(DataLayer, PreviousState);
+			DataLayerManager->SetDataLayerInstanceRuntimeState(DataLayer, PreviousState);
 		}
 	}
 }
@@ -157,7 +157,7 @@ void FPreAnimatedDataLayerStorage::OnGroupDestroyed(FPreAnimatedStorageGroupHand
 	GroupHandle = FPreAnimatedStorageGroupHandle();
 }
 
-FPreAnimatedStateEntry FPreAnimatedDataLayerStorage::MakeEntry(UDataLayerInstance* InDataLayer)
+FPreAnimatedStateEntry FPreAnimatedDataLayerStorage::MakeEntry(const UDataLayerInstance* InDataLayer)
 {
 	if (!GroupHandle)
 	{
@@ -167,7 +167,7 @@ FPreAnimatedStateEntry FPreAnimatedDataLayerStorage::MakeEntry(UDataLayerInstanc
 	return FPreAnimatedStateEntry{ GroupHandle, FPreAnimatedStateCachedValueHandle{ StorageID, StorageIndex } };
 }
 
-void FPreAnimatedDataLayerStorage::SavePreAnimatedState(UDataLayerInstance* DataLayer, UDataLayerSubsystem* SubSystem)
+void FPreAnimatedDataLayerStorage::SavePreAnimatedState(const UDataLayerInstance* DataLayer, UDataLayerManager* DataLayerManager)
 {
 	FPreAnimatedStateEntry         Entry              = MakeEntry(DataLayer);
 	EPreAnimatedStorageRequirement StorageRequirement = this->ParentExtension->GetStorageRequirement(Entry);
@@ -176,14 +176,14 @@ void FPreAnimatedDataLayerStorage::SavePreAnimatedState(UDataLayerInstance* Data
 	{
 		// @todo: If a data layer is loading when Sequencer attempts to activate it,
 		// should it return to ::Loading when sequencer is done?
-		EDataLayerRuntimeState ExistingState = SubSystem->GetDataLayerRuntimeState(DataLayer);
+		EDataLayerRuntimeState ExistingState = DataLayerManager->GetDataLayerInstanceRuntimeState(DataLayer);
 
 		AssignPreAnimatedValue(Entry.ValueHandle.StorageIndex, StorageRequirement, CopyTemp(ExistingState));
 	}
 }
 
 #if WITH_EDITOR
-void FPreAnimatedDataLayerStorage::SavePreAnimatedStateInEditor(UDataLayerInstance* DataLayer)
+void FPreAnimatedDataLayerStorage::SavePreAnimatedStateInEditor(const UDataLayerInstance* DataLayer)
 {
 	FPreAnimatedStateEntry         Entry              = MakeEntry(DataLayer);
 	EPreAnimatedStorageRequirement StorageRequirement = this->ParentExtension->GetStorageRequirement(Entry);
@@ -278,11 +278,11 @@ void FDesiredLayerStates::Reset()
 	}
 }
 
-EDataLayerUpdateFlags FDesiredLayerStates::Apply(FPreAnimatedDataLayerStorage* PreAnimatedStorage, UDataLayerSubsystem* DataLayerSubsystem, UWorldPartitionSubsystem* WorldPartitionSubsystem)
+EDataLayerUpdateFlags FDesiredLayerStates::Apply(FPreAnimatedDataLayerStorage* PreAnimatedStorage, UDataLayerManager* DataLayerManager, UWorldPartitionSubsystem* WorldPartitionSubsystem)
 {
 	EDataLayerUpdateFlags Flags = EDataLayerUpdateFlags::None;
 
-	auto IsDataLayerReady = [WorldPartitionSubsystem](UDataLayerInstance* DataLayer, EDataLayerRuntimeState DesireState, bool bExactState)
+	auto IsDataLayerReady = [WorldPartitionSubsystem](const UDataLayerInstance* DataLayer, EDataLayerRuntimeState DesireState, bool bExactState)
 	{
 		EWorldPartitionRuntimeCellState QueryState;
 		switch (DesireState)
@@ -319,16 +319,16 @@ EDataLayerUpdateFlags FDesiredLayerStates::Apply(FPreAnimatedDataLayerStorage* P
 
 		if (TOptional<EDataLayerRuntimeState> DesiredState = StateValue.ComputeDesiredState())
 		{
-			UDataLayerInstance* DataLayer = DataLayerSubsystem->GetDataLayerInstance(It.Key());
+			const UDataLayerInstance* DataLayer = DataLayerManager->GetDataLayerInstance(It.Key());
 			if (DataLayer)
 			{
 				if (PreAnimatedStorage)
 				{
-					PreAnimatedStorage->SavePreAnimatedState(DataLayer, DataLayerSubsystem);
+					PreAnimatedStorage->SavePreAnimatedState(DataLayer, DataLayerManager);
 				}
 
 				const EDataLayerRuntimeState DesiredStateValue = DesiredState.GetValue();
-				DataLayerSubsystem->SetDataLayerRuntimeState(DataLayer, DesiredStateValue);
+				DataLayerManager->SetDataLayerInstanceRuntimeState(DataLayer, DesiredStateValue);
 
 				if (StateValue.ShouldFlushStreaming(DesiredStateValue) && !IsDataLayerReady(DataLayer, DesiredStateValue, true))
 				{
@@ -525,14 +525,14 @@ void UMovieSceneDataLayerSystem::OnRun(FSystemTaskPrerequisites& InPrerequisites
 		}
 		else
 		{
-			UDataLayerSubsystem* DataLayerSubsystem = World->GetSubsystem<UDataLayerSubsystem>();
-			if (!DataLayerSubsystem)
+			UDataLayerManager* DataLayerManager = World->GetDataLayerManager();
+			if (!DataLayerManager)
 			{
 				UE_LOG(LogMovieScene, Warning, TEXT("Unable to retrieve data layer subsystem - data layer tracks will not function correctly"));
 			}
 			else
 			{
-				EDataLayerUpdateFlags UpdateFlags = DesiredLayerStates->Apply(WeakPreAnimatedStorage.Pin().Get(), DataLayerSubsystem, WorldPartitionSubsystem);
+				EDataLayerUpdateFlags UpdateFlags = DesiredLayerStates->Apply(WeakPreAnimatedStorage.Pin().Get(), DataLayerManager, WorldPartitionSubsystem);
 
 				const double StartTime = FPlatformTime::Seconds();
 				const TCHAR* FlushTypeString = nullptr;
@@ -594,8 +594,8 @@ void UMovieSceneDataLayerSystem::UpdateDesiredStates()
 	auto GatherDataLayers = [this, BuiltInComponents](FEntityAllocationIteratorItem Item, const FMovieSceneDataLayerComponentData* ComponentData, const int16* OptHBiases)
 	{
 		UWorld* World = GetWorld();
-		UDataLayerSubsystem* DataLayerSubsystem = World ? World->GetSubsystem<UDataLayerSubsystem>() : nullptr;
-		if (!DataLayerSubsystem)
+		UDataLayerManager* DataLayerManager = World ? World->GetDataLayerManager() : nullptr;
+		if (!DataLayerManager)
 		{
 			return;
 		}
@@ -614,7 +614,7 @@ void UMovieSceneDataLayerSystem::UpdateDesiredStates()
 
 			for (const UDataLayerAsset* DataLayerAsset : Section->GetDataLayerAssets())
 			{
-				const UDataLayerInstance* DataLayerInstance = DataLayerSubsystem->GetDataLayerInstanceFromAsset(DataLayerAsset);
+				const UDataLayerInstance* DataLayerInstance = DataLayerManager->GetDataLayerInstanceFromAsset(DataLayerAsset);
 				if (DataLayerInstance)
 				{
 					this->DesiredLayerStates->ApplyNewState(DataLayerInstance->GetDataLayerFName(), OptHBiases ? OptHBiases[Index] : 0, DesiredState, bRequiresStreamingFlush);
@@ -634,10 +634,10 @@ void UMovieSceneDataLayerSystem::BeginTrackingEntities()
 {
 	using namespace UE::MovieScene;
 
-	UWorld*                     World              = GetWorld();
-	UDataLayerSubsystem*        DataLayerSubsystem = World ? World->GetSubsystem<UDataLayerSubsystem>() : nullptr;
+	UWorld* World = GetWorld();
+	UDataLayerManager* DataLayerManager = World ? World->GetDataLayerManager() : nullptr;
 
-	if (!DataLayerSubsystem)
+	if (!DataLayerManager)
 	{
 		return;
 	}
@@ -653,7 +653,7 @@ void UMovieSceneDataLayerSystem::BeginTrackingEntities()
 	// ---------------------------------------------------------------------------------
 	// Only gather entity meta-data during SavePreAnimatedState - the actual values will be cached
 	// inside FDataLayerState::Apply
-	auto GatherDataLayers = [this, EntityMetaData, DataLayerSubsystem, PreAnimatedStorage](
+	auto GatherDataLayers = [this, EntityMetaData, DataLayerManager, PreAnimatedStorage](
 		FEntityAllocationIteratorItem Item,
 		TRead<FMovieSceneEntityID> EntityIDs,
 		TRead<FRootInstanceHandle> RootInstanceHandles,
@@ -670,7 +670,7 @@ void UMovieSceneDataLayerSystem::BeginTrackingEntities()
 
 				for (const UDataLayerAsset* DataLayerAsset : Section->GetDataLayerAssets())
 				{
-					UDataLayerInstance* DataLayerInstance = DataLayerSubsystem->GetDataLayerInstanceFromAsset(DataLayerAsset);
+					const UDataLayerInstance* DataLayerInstance = DataLayerManager->GetDataLayerInstanceFromAsset(DataLayerAsset);
 					if (DataLayerInstance)
 					{
 						FPreAnimatedStateEntry Entry = PreAnimatedStorage->MakeEntry(DataLayerInstance);
