Sequencer: Check for valid ChannelID in Quaternion Blender System

diff --git a/Engine/Source/Runtime/MovieSceneTracks/Private/Systems/MovieSceneQuaternionBlenderSystem.cpp b/Engine/Source/Runtime/MovieSceneTracks/Private/Systems/MovieSceneQuaternionBlenderSystem.cpp
--- a/Engine/Source/Runtime/MovieSceneTracks/Private/Systems/MovieSceneQuaternionBlenderSystem.cpp
+++ b/Engine/Source/Runtime/MovieSceneTracks/Private/Systems/MovieSceneQuaternionBlenderSystem.cpp
@@ -104,7 +104,7 @@ struct FAbsoluteAccumulationTask
 			return;
 		}
 
-		const double TotalWeight = Buffers->AbsoluteWeights[BlendChannel.ChannelID]
+		const double TotalWeight = Buffers->AbsoluteWeights.IsValidIndex(BlendChannel.ChannelID)
 			? Buffers->AbsoluteWeights[BlendChannel.ChannelID]
 			: 1.0;
 
