Sequencer: Root motion performance and swap root bone drawing correctly with Skeletons to avoid double transforms.

diff --git a/Engine/Source/Editor/MovieSceneTools/Private/EditModes/SkeletalAnimationTrackEditMode.cpp b/Engine/Source/Editor/MovieSceneTools/Private/EditModes/SkeletalAnimationTrackEditMode.cpp
--- a/Engine/Source/Editor/MovieSceneTools/Private/EditModes/SkeletalAnimationTrackEditMode.cpp
+++ b/Engine/Source/Editor/MovieSceneTools/Private/EditModes/SkeletalAnimationTrackEditMode.cpp
@@ -29,6 +29,7 @@
 #include "MovieSceneSkeletalAnimationRootHitProxy.h"
 #include "Systems/MovieSceneComponentTransformSystem.h"
 #include "SkeletalDebugRendering.h"
+#include "AnimSequencerInstanceProxy.h"
 
 FName FSkeletalAnimationTrackEditMode::ModeName("EditMode.SkeletalAnimationTrackEditMode");
 
@@ -133,7 +134,7 @@ static bool IsVirtualBone(USkeleton* Skeleton, FName BoneName)
 	return Skeleton->GetVirtualBones().IndexOfByPredicate([&](const FVirtualBone& VirtualBone) { return VirtualBone.VirtualBoneName == BoneName; }) != INDEX_NONE;
 }
 
-static void DrawBonesFromCompactPose(const FCompactPose& Pose, USkeletalMeshComponent* MeshComponent, FPrimitiveDrawInterface* PDI,
+static void DrawBonesFromCompactPose(const FCompactPose& Pose, USkeletalMeshComponent* MeshComponent, bool bIncludeRootBone, FPrimitiveDrawInterface* PDI,
 	const FLinearColor& DrawColour, FCompactPoseBoneIndex RootBoneIndex, FVector& LocationOfRootBone)
 {
 	LocationOfRootBone = FVector(TNumericLimits<float>::Max(), TNumericLimits<float>::Max(), TNumericLimits<float>::Max());
@@ -170,9 +171,11 @@ static void DrawBonesFromCompactPose(const FCompactPose& Pose, USkeletalMeshComp
 
 			int32 ParentIndex = Pose.GetBoneContainer().GetParentBoneIndex(MeshBoneIndex.GetInt());
 
+
 			if (ParentIndex == INDEX_NONE)
 			{
-				WorldTransforms[MeshBoneIndex.GetInt()] = Pose[BoneIndex] * MeshComponent->GetComponentTransform();
+				FTransform RootBone = (bIncludeRootBone == true) ? Pose[BoneIndex] : FTransform::Identity;
+				WorldTransforms[MeshBoneIndex.GetInt()] = RootBone * MeshComponent->GetComponentTransform();
 			}
 			else
 			{
@@ -280,8 +283,19 @@ void FSkeletalAnimationTrackEditMode::Render(const FSceneView* View, FViewport*
 								//show root motions
 								if (SkelAnimTrack->bShowRootMotionTrail)
 								{
+									if (SkelAnimTrack->RootMotionParams.bCacheRootTransforms == false)
+									{
+										SkelAnimTrack->RootMotionParams.bCacheRootTransforms = true;
+										SkelAnimTrack->SetUpRootMotions(true);
+
+									}
 									RenderTrail(BoundObject, SkelAnimTrack, InterrogationLinker, SkelMeshComp, Sequencer, PDI);								
 								}
+								else
+								{
+									SkelAnimTrack->RootMotionParams.bCacheRootTransforms = false;
+									SkelAnimTrack->RootMotionParams.RootTransforms.SetNum(0);
+								}
 								//show skeletons
 
 								const FLinearColor Colors[5] = { FLinearColor(1.0f,0.0f,1.0f,1.0f), FLinearColor(0.0f,1.0f,0.0f,1.0f), FLinearColor(0.0f, 0.0f, 1.0f, 1.0f),
@@ -340,7 +354,8 @@ void FSkeletalAnimationTrackEditMode::Render(const FSceneView* View, FViewport*
 											FLinearColor BoneColor = Colors[SectionIndex % 5] * (1.f - Alpha) + SectionColor * Alpha;
 
 											FVector RootLocation;
-											DrawBonesFromCompactPose(OutPose, SkelMeshComp, PDI, BoneColor, PoseIndex, RootLocation);
+											DrawBonesFromCompactPose(OutPose, SkelMeshComp, (AnimSection->Params.SwapRootBone == ESwapRootBone::SwapRootBone_None), PDI,
+												BoneColor, PoseIndex, RootLocation);
 											
 											if (IsRootSelected(AnimSection))
 											{
diff --git a/Engine/Source/Runtime/MovieSceneTracks/Private/Sections/MovieSceneSkeletalAnimationSection.cpp b/Engine/Source/Runtime/MovieSceneTracks/Private/Sections/MovieSceneSkeletalAnimationSection.cpp
--- a/Engine/Source/Runtime/MovieSceneTracks/Private/Sections/MovieSceneSkeletalAnimationSection.cpp
+++ b/Engine/Source/Runtime/MovieSceneTracks/Private/Sections/MovieSceneSkeletalAnimationSection.cpp
@@ -478,26 +478,16 @@ void UMovieSceneSkeletalAnimationSection::GetRootMotion(FFrameTime CurrentTime,
 {
 	if (GetRootMotionParams())
 	{
-		UMovieSceneSkeletalAnimationTrack* Track = GetTypedOuter<UMovieSceneSkeletalAnimationTrack>();
-		if (GetRootMotionParams()->bRootMotionsDirty)
+		if (UMovieSceneSkeletalAnimationTrack* Track = GetTypedOuter<UMovieSceneSkeletalAnimationTrack>())
 		{
-			if (Track)
-			{
-				Track->SetUpRootMotions(true);
-				if (GetRootMotionParams()->RootTransforms.Num() == 0) //should never be true but just in case
-				{
-					return;
-				}
-			}
+			OutRootMotionParams.Transform = Track->GetRootMotion(CurrentTime);
+			OutRootMotionParams.ChildBoneIndex = TempRootBoneIndex.IsSet() ? TempRootBoneIndex.GetValue() : INDEX_NONE;
+			OutRootMotionParams.bBlendFirstChildOfRoot = Track->bBlendFirstChildOfRoot ;
+			OutRootMotionParams.PreviousTransform = PreviousTransform;
 		}
-		OutRootMotionParams.bBlendFirstChildOfRoot = Track ? Track->bBlendFirstChildOfRoot : false;
-		OutRootMotionParams.ChildBoneIndex = TempRootBoneIndex.IsSet() ? TempRootBoneIndex.GetValue() : INDEX_NONE;
-		OutRootMotionParams.Transform = GetRootMotionParams()->GetRootMotion(CurrentTime);
-		OutRootMotionParams.PreviousTransform = PreviousTransform;
 	}
 }
 
-
 bool UMovieSceneSkeletalAnimationSection::GetRootMotionVelocity(FFrameTime PreviousTime, FFrameTime CurrentTime, FFrameRate FrameRate, 
 	FTransform& OutVelocity, float& OutWeight) const
 {
diff --git a/Engine/Source/Runtime/MovieSceneTracks/Private/Tracks/MovieSceneSkeletalAnimationTrack.cpp b/Engine/Source/Runtime/MovieSceneTracks/Private/Tracks/MovieSceneSkeletalAnimationTrack.cpp
--- a/Engine/Source/Runtime/MovieSceneTracks/Private/Tracks/MovieSceneSkeletalAnimationTrack.cpp
+++ b/Engine/Source/Runtime/MovieSceneTracks/Private/Tracks/MovieSceneSkeletalAnimationTrack.cpp
@@ -34,7 +34,6 @@ UMovieSceneSkeletalAnimationTrack::UMovieSceneSkeletalAnimationTrack(const FObje
 	: Super(ObjectInitializer)
 	, bUseLegacySectionIndexBlend(false)
 	, bBlendFirstChildOfRoot(false)
-	, ChildBoneIndex(INDEX_NONE)
 {
 #if WITH_EDITORONLY_DATA
 	TrackTint = FColor(124, 15, 124, 65);
@@ -142,6 +141,7 @@ void UMovieSceneSkeletalAnimationTrack::PostEditUndo()
 
 void UMovieSceneSkeletalAnimationTrack::PostEditChangeProperty(FPropertyChangedEvent& PropertyChangedEvent)
 {
+	Super::PostEditChangeProperty(PropertyChangedEvent);
 	RootMotionParams.bRootMotionsDirty = true;
 }
 
@@ -430,6 +430,92 @@ static float GetBestBlendPointTimeAtStart(UAnimSequenceBase* FirstAnimSeq, UAnim
 	return SmallIndex * 1.0f / FrameRate;
 }
 
+TOptional<FTransform>  UMovieSceneSkeletalAnimationTrack::GetRootMotion(FFrameTime CurrentTime)
+{
+	TOptional<FTransform> Transform;
+
+	UMovieScene* MovieScene = GetTypedOuter<UMovieScene>();
+	if (!MovieScene)
+	{
+		return Transform;
+	}
+	FFrameRate TickResolution = MovieScene->GetTickResolution();
+	
+	if (RootMotionParams.bRootMotionsDirty) 
+	{
+		SetUpRootMotions(true);
+	}
+	if (RootMotionParams.bHaveRootMotion == false)
+	{
+		return Transform;
+	}
+	SortSections();
+	TArray< UMovieSceneSkeletalAnimationSection*> SectionsAtCurrentTime;
+	TArray<FTransform> CurrentTransforms;
+	TArray<float> CurrentWeights;
+	TArray<FTransform> CurrentAdditiveTransforms;
+	TArray<float> CurrentAdditiveWeights;
+	FTransform CurrentTransform = FTransform::Identity;
+	for (UMovieSceneSection* Section : AnimationSections)
+	{
+		UMovieSceneSkeletalAnimationSection* AnimSection = Cast<UMovieSceneSkeletalAnimationSection>(Section);
+		if (AnimSection)
+		{
+			if (Section && AnimSection->Params.Animation && Section->GetRange().Contains(CurrentTime.FrameNumber))
+			{
+				UAnimSequenceBase* ValidAnimSequence = AnimSection->Params.Animation;
+				FMemMark Mark(FMemStack::Get());
+				FCompactPose OutPose;
+				TArray<FBoneIndexType> RequiredBoneIndexArray;
+				const FCurveEvaluationOption CurveEvalOption;
+				RequiredBoneIndexArray.AddUninitialized(ValidAnimSequence->GetSkeleton()->GetReferenceSkeleton().GetNum());
+				for (int32 BoneIndex = 0; BoneIndex < RequiredBoneIndexArray.Num(); ++BoneIndex)
+				{
+					RequiredBoneIndexArray[BoneIndex] = BoneIndex;
+				}
+
+				FBoneContainer BoneContainer(RequiredBoneIndexArray, CurveEvalOption, *ValidAnimSequence->GetSkeleton());
+				OutPose.ResetToRefPose(BoneContainer);
+				FBlendedCurve OutCurve;
+				OutCurve.InitFrom(BoneContainer);
+				UE::Anim::FStackAttributeContainer TempAttributes;
+				FAnimationPoseData AnimationPoseData(OutPose, OutCurve, TempAttributes);
+				UMovieSceneSkeletalAnimationSection::FRootMotionTransformParam Param;
+				Param.FrameRate = TickResolution;
+				Param.CurrentTime = CurrentTime.FrameNumber;
+
+				if (AnimSection->GetRootMotionTransform(AnimationPoseData, Param))
+				{
+					if (!Param.bOutIsAdditive)
+					{
+						CurrentTransform = Param.OutTransform * AnimSection->PreviousTransform;
+						CurrentTransforms.Add(CurrentTransform);
+						CurrentWeights.Add(Param.OutWeight);
+					}
+					else
+					{
+						CurrentAdditiveTransforms.Add(Param.OutTransform);
+						CurrentAdditiveWeights.Add(Param.OutWeight);
+					}
+				}
+			}
+		}
+	}
+
+	BlendTheseTransformsByWeight(CurrentTransform, CurrentTransforms, CurrentWeights);
+	//now handle additive onto the current
+	if (CurrentAdditiveWeights.Num() > 0)
+	{
+		FTransform AdditiveTransform;
+		BlendTheseTransformsByWeight(AdditiveTransform, CurrentAdditiveTransforms, CurrentAdditiveWeights);
+		const ScalarRegister VBlendWeight(1.0f);
+		FTransform::BlendFromIdentityAndAccumulate(CurrentTransform, AdditiveTransform, VBlendWeight);
+	}
+	Transform = CurrentTransform;
+	return Transform;
+
+}
+
 void UMovieSceneSkeletalAnimationTrack::SetUpRootMotions(bool bForce)
 {
 	UMovieScene* MovieScene = GetTypedOuter<UMovieScene>();
@@ -441,21 +527,17 @@ void UMovieSceneSkeletalAnimationTrack::SetUpRootMotions(bool bForce)
 	if (bForce || RootMotionParams.bRootMotionsDirty)
 	{
 		RootMotionParams.bRootMotionsDirty = false;
-
+		RootMotionParams.bHaveRootMotion = false;
 		const FFrameRate MinDisplayRate(60, 1);
 
 		FFrameRate DisplayRate =  MovieScene->GetDisplayRate().AsDecimal() < MinDisplayRate.AsDecimal() ? MinDisplayRate : MovieScene->GetDisplayRate();
 		FFrameRate TickResolution = MovieScene->GetTickResolution();
 		FFrameTime FrameTick = FFrameTime(FMath::Max(1, TickResolution.AsFrameNumber(1.0).Value / DisplayRate.AsFrameNumber(1.0).Value));
-		if (FrameTick.FrameNumber.Value == 0)
-		{
-			RootMotionParams.RootTransforms.SetNum(0);
-			return;
-		}
-
-		if (AnimationSections.Num() == 0)
+		if (AnimationSections.Num() == 0 || FrameTick.FrameNumber.Value == 0)
 		{
+#if WITH_EDITORONLY_DATA
 			RootMotionParams.RootTransforms.SetNum(0);
+#endif
 			return;
 		}
 		
@@ -545,8 +627,21 @@ void UMovieSceneSkeletalAnimationTrack::SetUpRootMotions(bool bForce)
 
 		if (bAnySectionsHaveOffset == false)
 		{
-			//no root transforms so bail
+#if WITH_EDITORONLY_DATA
 			RootMotionParams.RootTransforms.SetNum(0);
+#endif			
+			return;
+		}
+
+		RootMotionParams.bHaveRootMotion = true;
+		RootMotionParams.StartFrame = AnimationSections[0]->GetInclusiveStartFrame();
+		RootMotionParams.EndFrame = AnimationSections[AnimationSections.Num() - 1]->GetExclusiveEndFrame() - 1;
+		RootMotionParams.FrameTick = FrameTick;
+		
+		
+#if WITH_EDITORONLY_DATA
+		if (RootMotionParams.bCacheRootTransforms == false)
+		{
 			return;
 		}
 		//set up pose from valid anim sequences.
@@ -568,12 +663,9 @@ void UMovieSceneSkeletalAnimationTrack::SetUpRootMotions(bool bForce)
 			FBlendedCurve OutCurve;
 			OutCurve.InitFrom(BoneContainer);
 			TArray< UMovieSceneSkeletalAnimationSection*> SectionsAtCurrentTime;
-			RootMotionParams.StartFrame = AnimationSections[0]->GetInclusiveStartFrame();
-			RootMotionParams.EndFrame = AnimationSections[AnimationSections.Num() - 1]->GetExclusiveEndFrame() - 1;
-			RootMotionParams.FrameTick = FrameTick;
-
 			int32 NumTotal = (RootMotionParams.EndFrame.FrameNumber.Value - RootMotionParams.StartFrame.FrameNumber.Value) / (RootMotionParams.FrameTick.FrameNumber.Value) + 1;
 			RootMotionParams.RootTransforms.SetNum(NumTotal);
+
 			TArray<FTransform> CurrentTransforms;
 			TArray<float> CurrentWeights;
 			TArray<FTransform> CurrentAdditiveTransforms;
@@ -635,10 +727,10 @@ void UMovieSceneSkeletalAnimationTrack::SetUpRootMotions(bool bForce)
 		else //no valid anim sequence just clear out
 		{
 			RootMotionParams.RootTransforms.SetNum(0);
-			return;
 		}
-
+#endif
 	}
+
 }
 
 static FTransform GetTransformForBoneRelativeToIndex(UAnimSequence* AnimSequence, USkeletalMeshComponent* MeshComponent, const FName& InBoneName,
@@ -1079,50 +1171,6 @@ void UMovieSceneSkeletalAnimationTrack::ToggleShowRootMotionTrail()
 	bShowRootMotionTrail = bShowRootMotionTrail ? false : true;
 }
 #endif
-TOptional<FTransform>  FMovieSceneSkeletalAnimRootMotionTrackParams::GetRootMotion(FFrameTime CurrentTime)  const
-{
-	if (RootTransforms.Num() > 0)
-	{
-		if (CurrentTime >= StartFrame && CurrentTime <= EndFrame)
-		{
-			float FIndex = (float)(CurrentTime.FrameNumber.Value - StartFrame.FrameNumber.Value) / (float)(FrameTick.FrameNumber.Value);
-			int Index = (int)(FIndex);
-			FIndex -= (float)(Index);
-			FTransform Transform = RootTransforms[Index];
-			//Blends don't work with rotation if blend factor is smallish or largeish(returns Identity instead)so we have these 0.001f and >.99f checks.
-			if (FIndex > 0.001f)
-			{
-				if (Index < RootTransforms.Num() - 1)
-				{
-					if (FIndex < 0.99f)
-					{
-						FTransform Next = RootTransforms[Index + 1];
-						Transform.Blend(Transform, Next, FIndex);
-					}
-					else
-					{
-						Transform = RootTransforms[Index + 1];
-					}
-				}
-				else
-				{
-					Transform = RootTransforms[RootTransforms.Num() - 1];
-				}
-			}
-			return Transform;
-		}
-		else if (CurrentTime > EndFrame)
-		{
-			return RootTransforms[RootTransforms.Num() - 1];
-		}
-		else
-		{
-			return RootTransforms[0];
-		}
-	}
-	return TOptional<FTransform>();
-
-}
 
 
 //MZ To Do need way to get passed the skelmeshcomp when we add or move a section.
diff --git a/Engine/Source/Runtime/MovieSceneTracks/Public/Tracks/MovieSceneSkeletalAnimationTrack.h b/Engine/Source/Runtime/MovieSceneTracks/Public/Tracks/MovieSceneSkeletalAnimationTrack.h
--- a/Engine/Source/Runtime/MovieSceneTracks/Public/Tracks/MovieSceneSkeletalAnimationTrack.h
+++ b/Engine/Source/Runtime/MovieSceneTracks/Public/Tracks/MovieSceneSkeletalAnimationTrack.h
@@ -22,12 +22,15 @@ struct FMovieSceneSkeletalAnimRootMotionTrackParams
 	FFrameTime StartFrame;
 	FFrameTime EndFrame;
 	bool bRootMotionsDirty;
-	TArray<FTransform> RootTransforms;
+	bool bHaveRootMotion;
 	FTransform RootMotionStartOffset; // root motion may not be in Mesh Space if we are putting values on a bone that is a child of a root with an offset
 	/** Get the Root Motion transform at the specified time.*/
-	MOVIESCENETRACKS_API TOptional<FTransform>  GetRootMotion(FFrameTime CurrentTime) const;
-	FMovieSceneSkeletalAnimRootMotionTrackParams() : bRootMotionsDirty(true) {}
+	FMovieSceneSkeletalAnimRootMotionTrackParams() : bRootMotionsDirty(true),bHaveRootMotion(false) {}
 
+#if WITH_EDITORONLY_DATA
+	bool bCacheRootTransforms = false;
+	TArray<FTransform> RootTransforms;
+#endif
 };
 
 /**
@@ -90,6 +93,7 @@ public:
 
 	MOVIESCENETRACKS_API void SetRootMotionsDirty();
 	MOVIESCENETRACKS_API void SetUpRootMotions(bool bForce);
+	MOVIESCENETRACKS_API TOptional<FTransform>  GetRootMotion(FFrameTime CurrentTime);
 	MOVIESCENETRACKS_API void MatchSectionByBoneTransform(bool bMatchWithPrevious, USkeletalMeshComponent* SkelMeshComp, UMovieSceneSkeletalAnimationSection* CurrentSection, FFrameTime CurrentFrame, FFrameRate FrameRate,
 		const FName& BoneName, FTransform& SecondSectionRootDiff, FVector& TranslationDiff, FQuat& RotationDiff); //add options for z and for rotation.
 #if WITH_EDITORONLY_DATA
@@ -116,9 +120,6 @@ public:
 	UPROPERTY(EditAnywhere, Category = "Root Motions")
 	bool bBlendFirstChildOfRoot;
 
-	/** Child bone to blend if bBlendFirstChildOfRoot is true*/
-	int32 ChildBoneIndex;
-
 #if WITH_EDITORONLY_DATA
 public:
 
