Sequencer: add missing blender type tags on material parameter entities

diff --git a/Engine/Source/Runtime/MovieScene/Private/EntitySystem/MovieSceneEntityManager.cpp b/Engine/Source/Runtime/MovieScene/Private/EntitySystem/MovieSceneEntityManager.cpp
--- a/Engine/Source/Runtime/MovieScene/Private/EntitySystem/MovieSceneEntityManager.cpp
+++ b/Engine/Source/Runtime/MovieScene/Private/EntitySystem/MovieSceneEntityManager.cpp
@@ -1120,7 +1120,6 @@ void FEntityManager::AddComponent(FMovieSceneEntityID EntityID, FComponentTypeID
 		if (ExistingMask.Contains(ComponentType))
 		{
 			// Entity already has this component type
-			// @todo: Should this be an error?
 			return;
 		}
 
diff --git a/Engine/Source/Runtime/MovieSceneTracks/Private/Systems/MovieSceneMaterialParameterSystem.cpp b/Engine/Source/Runtime/MovieSceneTracks/Private/Systems/MovieSceneMaterialParameterSystem.cpp
--- a/Engine/Source/Runtime/MovieSceneTracks/Private/Systems/MovieSceneMaterialParameterSystem.cpp
+++ b/Engine/Source/Runtime/MovieSceneTracks/Private/Systems/MovieSceneMaterialParameterSystem.cpp
@@ -129,7 +129,7 @@ struct FApplyVectorParameters
 
 struct FScalarMixin
 {
-	void CreateEntity(UMovieSceneEntitySystemLinker* Linker, UObject* BoundMaterial, FName ParameterName, TArrayView<const FMovieSceneEntityID> Inputs, FAnimatedMaterialParameterInfo* Output)
+	void CreateEntity(UMovieSceneEntitySystemLinker* Linker, UObject* BoundMaterial, FName ParameterName, FComponentTypeID BlenderTypeTag, TArrayView<const FMovieSceneEntityID> Inputs, FAnimatedMaterialParameterInfo* Output)
 	{
 		FBuiltInComponentTypes* BuiltInComponents = FBuiltInComponentTypes::Get();
 		FMovieSceneTracksComponentTypes* TracksComponents = FMovieSceneTracksComponentTypes::Get();
@@ -165,6 +165,7 @@ struct FScalarMixin
 		.Add(BuiltInComponents->DoubleResult[0], 0.0)
 		.AddConditional(TracksComponents->FloatParameter.InitialValue, InitialValue, bHasInitialValue)
 		.AddTag(TracksComponents->FloatParameter.PropertyTag)
+		.AddTag(BlenderTypeTag)
 		.AddTag(BuiltInComponents->Tags.NeedsLink)
 		.AddMutualComponents()
 		.CreateEntity(&Linker->EntityManager);
@@ -195,7 +196,7 @@ struct FScalarMixin
 
 struct FVectorMixin
 {
-	void CreateEntity(UMovieSceneEntitySystemLinker* Linker, UObject* BoundMaterial, FName ParameterName, TArrayView<const FMovieSceneEntityID> Inputs, FAnimatedMaterialParameterInfo* Output)
+	void CreateEntity(UMovieSceneEntitySystemLinker* Linker, UObject* BoundMaterial, FName ParameterName, FComponentTypeID BlenderTypeTag, TArrayView<const FMovieSceneEntityID> Inputs, FAnimatedMaterialParameterInfo* Output)
 	{
 		FBuiltInComponentTypes* BuiltInComponents = FBuiltInComponentTypes::Get();
 		FMovieSceneTracksComponentTypes* TracksComponents = FMovieSceneTracksComponentTypes::Get();
@@ -240,6 +241,7 @@ struct FVectorMixin
 		.Add(BuiltInComponents->DoubleResult[3], 0.0)
 		.AddConditional(TracksComponents->ColorParameter.InitialValue, InitialValue, bHasInitialValue)
 		.AddTag(TracksComponents->ColorParameter.PropertyTag)
+		.AddTag(BlenderTypeTag)
 		.AddTag(BuiltInComponents->Tags.NeedsLink)
 		.AddMutualComponents()
 		.CreateEntity(&Linker->EntityManager);
@@ -320,9 +322,12 @@ struct TOverlappingMaterialParameterHandler : Mixin
 				Output->BlendChannelID = System->DoubleBlenderSystem->AllocateBlendChannel();
 
 				// Needs blending
-				Mixin::CreateEntity(Linker, BoundMaterial, ParameterName, Inputs, Output);
+				FComponentTypeID BlenderTypeTag = System->DoubleBlenderSystem->GetBlenderTypeTag();
+				Mixin::CreateEntity(Linker, BoundMaterial, ParameterName, BlenderTypeTag, Inputs, Output);
 			}
 
+			const FComponentTypeID BlenderTypeTag = System->DoubleBlenderSystem->GetBlenderTypeTag();
+
 			for (FMovieSceneEntityID Input : Inputs)
 			{
 				if (!Linker->EntityManager.HasComponent(Input, BuiltInComponents->BlendChannelInput))
@@ -334,6 +339,9 @@ struct TOverlappingMaterialParameterHandler : Mixin
 					// If the bound material changed, we might have been re-assigned a different blend channel so make sure it's up to date
 					Linker->EntityManager.WriteComponentChecked(Input, BuiltInComponents->BlendChannelInput, Output->BlendChannelID);
 				}
+
+				// Ensure we have the blender type tag on the inputs.
+				Linker->EntityManager.AddComponent(Input, BlenderTypeTag);
 			}
 		}
 		else if (!Output->OutputEntityID && Inputs.Num() == 1)
