Sequencer: Unify material parameter pre-animated state between Material Parameter Collections and Material Parameters

  - Up until this point proper pre-animated state has never actually been implemented for material parameters. This change unifies the pre-animated state for both so they utilize the same storage mechanisms.
  - As part of this change I needed to be able to run the begin tracking and caching methods on slightly different filters (previously it was explicitly only tied to NeedsLink entities, but material parameters need to be run on NeedsLink or BoundMaterialChanged entities). To this end I changed the callsites so they always use the filter provided by the params, and the params always adds NeedsLink to the filter by default.

diff --git a/Engine/Source/Runtime/MovieScene/Public/Evaluation/PreAnimatedState/MovieScenePreAnimatedStateStorage.h b/Engine/Source/Runtime/MovieScene/Public/Evaluation/PreAnimatedState/MovieScenePreAnimatedStateStorage.h
--- a/Engine/Source/Runtime/MovieScene/Public/Evaluation/PreAnimatedState/MovieScenePreAnimatedStateStorage.h
+++ b/Engine/Source/Runtime/MovieScene/Public/Evaluation/PreAnimatedState/MovieScenePreAnimatedStateStorage.h
@@ -50,6 +50,11 @@ struct TPreAnimatedStateTaskParams
 {
 	FEntityComponentFilter AdditionalFilter;
 
+	TSavePreAnimatedStateParams()
+	{
+		AdditionalFilter.All({ FBuiltInComponentTypes::Get()->Tags.NeedsLink });
+	}
+
 	static constexpr bool ShouldSavePreAnimatedState(typename TCallTraits<InputTypes>::ParamType... Inputs)
 	{
 		return true;
@@ -552,7 +557,6 @@ public:
 		.Read(BuiltInComponents->RootInstanceHandle)
 		.ReadAllOf(InComponentTypes...)
 		.CombineFilter(InParams.AdditionalFilter)
-		.FilterAll({ BuiltInComponents->Tags.NeedsLink })
 		.Iterate_PerAllocation(&Linker->EntityManager, VisitAllocation);
 	}
 
@@ -651,7 +655,6 @@ public:
 		FEntityTaskBuilder()
 		.ReadAllOf(InComponentTypes...)
 		.CombineFilter(InParams.AdditionalFilter)
-		.FilterAll({ BuiltInComponents->Tags.NeedsLink })
 		.Iterate_PerAllocation(&Linker->EntityManager, VisitAllocation);
 	}
 
@@ -810,7 +813,6 @@ public:
 		.Read(BuiltInComponents->RootInstanceHandle)
 		.ReadAllOf(InComponentTypes...)
 		.CombineFilter(InParams.AdditionalFilter)
-		.FilterAll({ BuiltInComponents->Tags.NeedsLink })
 		.Iterate_PerAllocation(&Linker->EntityManager, VisitAllocation);
 	}
 
diff --git a/Engine/Source/Runtime/MovieSceneTracks/Private/Systems/MovieSceneMaterialParameterCollectionSystem.cpp b/Engine/Source/Runtime/MovieSceneTracks/Private/Systems/MovieSceneMaterialParameterCollectionSystem.cpp
--- a/Engine/Source/Runtime/MovieSceneTracks/Private/Systems/MovieSceneMaterialParameterCollectionSystem.cpp
+++ b/Engine/Source/Runtime/MovieSceneTracks/Private/Systems/MovieSceneMaterialParameterCollectionSystem.cpp
@@ -6,14 +6,13 @@
 #include "IMovieScenePlayer.h"
 #include "MovieSceneSequence.h"
 #include "Evaluation/MovieSceneEvaluationTemplateInstance.h"
-#include "Evaluation/PreAnimatedState/MovieScenePreAnimatedStorageID.inl"
 #include "EntitySystem/BuiltInComponentTypes.h"
 #include "EntitySystem/MovieSceneEntityMutations.h"
-#include "Evaluation/PreAnimatedState/MovieScenePreAnimatedStorageID.inl"
 
 #include "MovieSceneTracksComponentTypes.h"
 #include "Systems/FloatChannelEvaluatorSystem.h"
 #include "Systems/MovieScenePiecewiseDoubleBlenderSystem.h"
+#include "Systems/MovieScenePreAnimatedMaterialParameters.h"
 
 #include "Logging/MessageLog.h"
 #include "Misc/UObjectToken.h"
@@ -26,105 +25,6 @@
 
 #include UE_INLINE_GENERATED_CPP_BY_NAME(MovieSceneMaterialParameterCollectionSystem)
 
-namespace UE::MovieScene
-{
-
-
-struct FMaterialParameterCollectionKey
-{
-	TObjectKey<UMaterialParameterCollectionInstance> MPCI;
-	FName ParameterName;
-
-	FMaterialParameterCollectionKey(UObject* InBoundMaterial, const FName& InParameterName)
-		: MPCI(CastChecked<UMaterialParameterCollectionInstance>(InBoundMaterial))
-		, ParameterName(InParameterName)
-	{}
-
-	friend uint32 GetTypeHash(const FMaterialParameterCollectionKey& InKey)
-	{
-		return GetTypeHash(InKey.MPCI) ^ GetTypeHash(InKey.ParameterName);
-	}
-
-	friend bool operator==(const FMaterialParameterCollectionKey& A, const FMaterialParameterCollectionKey& B)
-	{
-		return A.MPCI == B.MPCI && A.ParameterName == B.ParameterName;
-	}
-};
-
-struct FMaterialParameterCollectionScalarTraits : FBoundObjectPreAnimatedStateTraits
-{
-	using KeyType = FMaterialParameterCollectionKey;
-	using StorageType = float;
-
-	static void ReplaceObject(FMaterialParameterCollectionKey& InOutKey, const FObjectKey& NewObject)
-	{
-		InOutKey.MPCI = Cast<UMaterialParameterCollectionInstance>(NewObject.ResolveObjectPtr());
-	}
-
-	static float CachePreAnimatedValue(UObject* InBoundMaterial, const FName& ParameterName)
-	{
-		UMaterialParameterCollectionInstance* MPCI = CastChecked<UMaterialParameterCollectionInstance>(InBoundMaterial);
-
-		float ParameterValue = 0.f;
-		MPCI->GetScalarParameterValue(ParameterName, ParameterValue);
-		return ParameterValue;
-	}
-
-	static void RestorePreAnimatedValue(const FMaterialParameterCollectionKey& InKey, float OldValue, const FRestoreStateParams& Params)
-	{
-		UMaterialParameterCollectionInstance* MPCI = InKey.MPCI.ResolveObjectPtr();
-		if (MPCI)
-		{
-			MPCI->SetScalarParameterValue(InKey.ParameterName, OldValue);
-		}
-	}
-};
-
-struct FMaterialParameterCollectionVectorTraits : FBoundObjectPreAnimatedStateTraits
-{
-	using KeyType = FMaterialParameterCollectionKey;
-	using StorageType = FLinearColor;
-
-	static void ReplaceObject(FMaterialParameterCollectionKey& InOutKey, const FObjectKey& NewObject)
-	{
-		InOutKey.MPCI = Cast<UMaterialParameterCollectionInstance>(NewObject.ResolveObjectPtr());
-	}
-
-	static FLinearColor CachePreAnimatedValue(UObject* InBoundMaterial, const FName& ParameterName)
-	{
-		UMaterialParameterCollectionInstance* MPCI = CastChecked<UMaterialParameterCollectionInstance>(InBoundMaterial);
-		FLinearColor ParameterValue = FLinearColor::White;
-		MPCI->GetVectorParameterValue(ParameterName, ParameterValue);
-		return ParameterValue;
-	}
-
-	static void RestorePreAnimatedValue(const FMaterialParameterCollectionKey& InKey, const FLinearColor& OldValue, const FRestoreStateParams& Params)
-	{
-		UMaterialParameterCollectionInstance* MPCI = InKey.MPCI.ResolveObjectPtr();
-		if (MPCI)
-		{
-			MPCI->SetVectorParameterValue(InKey.ParameterName, OldValue);
-		}
-	}
-};
-
-struct FPreAnimatedMPCScalarStorage
-	: public TPreAnimatedStateStorage<FMaterialParameterCollectionScalarTraits>
-{
-	static TAutoRegisterPreAnimatedStorageID<FPreAnimatedMPCScalarStorage> StorageID;
-};
-
-struct FPreAnimatedMPCVectorStorage
-	: public TPreAnimatedStateStorage<FMaterialParameterCollectionVectorTraits>
-{
-	static TAutoRegisterPreAnimatedStorageID<FPreAnimatedMPCVectorStorage> StorageID;
-};
-
-
-TAutoRegisterPreAnimatedStorageID<FPreAnimatedMPCScalarStorage> FPreAnimatedMPCScalarStorage::StorageID;
-TAutoRegisterPreAnimatedStorageID<FPreAnimatedMPCVectorStorage> FPreAnimatedMPCVectorStorage::StorageID;
-
-} // namespace UE::MovieScene
 
 UMovieSceneMaterialParameterCollectionSystem::UMovieSceneMaterialParameterCollectionSystem(const FObjectInitializer& ObjInit)
 	: Super(ObjInit)
@@ -149,8 +49,8 @@ void UMovieSceneMaterialParameterCollectionSystem::OnLink()
 {
 	using namespace UE::MovieScene;
 
-	ScalarParameterStorage = Linker->PreAnimatedState.GetOrCreateStorage<FPreAnimatedMPCScalarStorage>();
-	VectorParameterStorage = Linker->PreAnimatedState.GetOrCreateStorage<FPreAnimatedMPCVectorStorage>();
+	ScalarParameterStorage = Linker->PreAnimatedState.GetOrCreateStorage<FPreAnimatedScalarMaterialParameterStorage>();
+	VectorParameterStorage = Linker->PreAnimatedState.GetOrCreateStorage<FPreAnimatedVectorMaterialParameterStorage>();
 }
 
 void UMovieSceneMaterialParameterCollectionSystem::OnRun(FSystemTaskPrerequisites& InPrerequisites, FSystemSubsequentTasks& Subsequents)
diff --git a/Engine/Source/Runtime/MovieSceneTracks/Private/Systems/MovieSceneMaterialParameterSystem.cpp b/Engine/Source/Runtime/MovieSceneTracks/Private/Systems/MovieSceneMaterialParameterSystem.cpp
--- a/Engine/Source/Runtime/MovieSceneTracks/Private/Systems/MovieSceneMaterialParameterSystem.cpp
+++ b/Engine/Source/Runtime/MovieSceneTracks/Private/Systems/MovieSceneMaterialParameterSystem.cpp
@@ -8,6 +8,9 @@
 #include "Systems/MovieSceneHierarchicalBiasSystem.h"
 #include "Systems/MovieSceneMaterialSystem.h"
 #include "Systems/MovieSceneInitialValueSystem.h"
+#include "Systems/MovieScenePreAnimatedMaterialParameters.h"
+
+#include "Evaluation/PreAnimatedState/MovieScenePreAnimatedObjectStorage.h"
 
 #include "Materials/MaterialParameterCollectionInstance.h"
 #include "Materials/MaterialInstanceDynamic.h"
@@ -410,6 +413,9 @@ void UMovieSceneMaterialParameterSystem::OnLink()
 
 	ScalarParameterTracker.Initialize(this);
 	VectorParameterTracker.Initialize(this);
+
+	ScalarParameterStorage = Linker->PreAnimatedState.GetOrCreateStorage<FPreAnimatedScalarMaterialParameterStorage>();
+	VectorParameterStorage = Linker->PreAnimatedState.GetOrCreateStorage<FPreAnimatedVectorMaterialParameterStorage>();
 }
 
 void UMovieSceneMaterialParameterSystem::OnUnlink()
@@ -496,6 +502,19 @@ void UMovieSceneMaterialParameterSystem::OnInstantiation()
 		// Process all blended vector parameters
 		TOverlappingMaterialParameterHandler<FVectorMixin> VectorHandler(this);
 		VectorParameterTracker.ProcessInvalidatedOutputs(Linker, VectorHandler);
+
+		TSavePreAnimatedStateParams<UObject*, FName> Params;
+		Params.AdditionalFilter.Reset();
+		Params.AdditionalFilter.None({ TracksComponents->MPC });
+		Params.AdditionalFilter.Any({ BuiltInComponents->Tags.NeedsLink, TracksComponents->Tags.BoundMaterialChanged });
+
+		ScalarParameterStorage->BeginTrackingEntitiesTask(Linker, Params, TracksComponents->BoundMaterial, TracksComponents->ScalarParameterName);
+		ScalarParameterStorage->CachePreAnimatedValuesTask(Linker, Params, TracksComponents->BoundMaterial, TracksComponents->ScalarParameterName);
+
+		VectorParameterStorage->BeginTrackingEntitiesTask(Linker, Params, TracksComponents->BoundMaterial, TracksComponents->VectorParameterName);
+		VectorParameterStorage->BeginTrackingEntitiesTask(Linker, Params, TracksComponents->BoundMaterial, TracksComponents->ColorParameterName);
+		VectorParameterStorage->CachePreAnimatedValuesTask(Linker, Params, TracksComponents->BoundMaterial, TracksComponents->VectorParameterName);
+		VectorParameterStorage->CachePreAnimatedValuesTask(Linker, Params, TracksComponents->BoundMaterial, TracksComponents->ColorParameterName);
 	}
 }
 
diff --git a/Engine/Source/Runtime/MovieSceneTracks/Private/Systems/MovieScenePreAnimatedMaterialParameters.cpp b/Engine/Source/Runtime/MovieSceneTracks/Private/Systems/MovieScenePreAnimatedMaterialParameters.cpp
new file mode 100644
--- /dev/null
+++ b/Engine/Source/Runtime/MovieSceneTracks/Private/Systems/MovieScenePreAnimatedMaterialParameters.cpp
@@ -0,0 +1,97 @@
+// Copyright Epic Games, Inc. All Rights Reserved.
+
+#include "Systems/MovieScenePreAnimatedMaterialParameters.h"
+#include "Evaluation/PreAnimatedState/MovieScenePreAnimatedStorageID.inl"
+#include "Materials/MaterialParameterCollectionInstance.h"
+#include "Materials/MaterialInstanceDynamic.h"
+#include "Materials/MaterialInterface.h"
+
+namespace UE::MovieScene
+{
+
+TAutoRegisterPreAnimatedStorageID<FPreAnimatedScalarMaterialParameterStorage> FPreAnimatedScalarMaterialParameterStorage::StorageID;
+TAutoRegisterPreAnimatedStorageID<FPreAnimatedVectorMaterialParameterStorage> FPreAnimatedVectorMaterialParameterStorage::StorageID;
+
+FMaterialParameterKey::FMaterialParameterKey(UObject* InBoundMaterial, const FName& InParameterName)
+	: BoundMaterial(InBoundMaterial)
+	, ParameterName(InParameterName)
+{}
+
+uint32 GetTypeHash(const FMaterialParameterKey& InKey)
+{
+	return GetTypeHash(InKey.BoundMaterial) ^ GetTypeHash(InKey.ParameterName);
+}
+
+bool operator==(const FMaterialParameterKey& A, const FMaterialParameterKey& B)
+{
+	return A.BoundMaterial == B.BoundMaterial && A.ParameterName == B.ParameterName;
+}
+
+void FMaterialParameterCollectionScalarTraits::ReplaceObject(FMaterialParameterKey& InOutKey, const FObjectKey& NewObject)
+{
+	InOutKey.BoundMaterial = NewObject.ResolveObjectPtr();
+}
+
+float FMaterialParameterCollectionScalarTraits::CachePreAnimatedValue(UObject* InBoundMaterial, const FName& ParameterName)
+{
+	float ParameterValue = 0.f;
+	if (UMaterialParameterCollectionInstance* MPCI = Cast<UMaterialParameterCollectionInstance>(InBoundMaterial))
+	{
+		MPCI->GetScalarParameterValue(ParameterName, ParameterValue);
+	}
+	else if (UMaterialInterface* MaterialInterface = Cast<UMaterialInterface>(InBoundMaterial))
+	{
+		MaterialInterface->GetScalarParameterValue(FMaterialParameterInfo(ParameterName), ParameterValue);
+	}
+	return ParameterValue;
+}
+
+void FMaterialParameterCollectionScalarTraits::RestorePreAnimatedValue(const FMaterialParameterKey& InKey, float OldValue, const FRestoreStateParams& Params)
+{
+	UObject* BoundObject = InKey.BoundMaterial.ResolveObjectPtr();
+
+	if (UMaterialParameterCollectionInstance* MPCI = Cast<UMaterialParameterCollectionInstance>(BoundObject))
+	{
+		MPCI->SetScalarParameterValue(InKey.ParameterName, OldValue);
+	}
+	else if (UMaterialInstanceDynamic* MID = Cast<UMaterialInstanceDynamic>(BoundObject))
+	{
+		MID->SetScalarParameterValue(InKey.ParameterName, OldValue);
+	}
+}
+
+void FMaterialParameterCollectionVectorTraits::ReplaceObject(FMaterialParameterKey& InOutKey, const FObjectKey& NewObject)
+{
+	InOutKey.BoundMaterial = NewObject.ResolveObjectPtr();
+}
+
+FLinearColor FMaterialParameterCollectionVectorTraits::CachePreAnimatedValue(UObject* InBoundMaterial, const FName& ParameterName)
+{
+	FLinearColor ParameterValue = FLinearColor::White;
+	if (UMaterialParameterCollectionInstance* MPCI = Cast<UMaterialParameterCollectionInstance>(InBoundMaterial))
+	{
+		MPCI->GetVectorParameterValue(ParameterName, ParameterValue);
+	}
+	else if (UMaterialInterface* MaterialInterface = Cast<UMaterialInterface>(InBoundMaterial))
+	{
+		MaterialInterface->GetVectorParameterValue(FMaterialParameterInfo(ParameterName), ParameterValue);
+	}
+	return ParameterValue;
+}
+
+void FMaterialParameterCollectionVectorTraits::RestorePreAnimatedValue(const FMaterialParameterKey& InKey, const FLinearColor& OldValue, const FRestoreStateParams& Params)
+{
+	UObject* BoundObject = InKey.BoundMaterial.ResolveObjectPtr();
+
+	if (UMaterialParameterCollectionInstance* MPCI = Cast<UMaterialParameterCollectionInstance>(BoundObject))
+	{
+		MPCI->SetVectorParameterValue(InKey.ParameterName, OldValue);
+	}
+	else if (UMaterialInstanceDynamic* MID = Cast<UMaterialInstanceDynamic>(BoundObject))
+	{
+		MID->SetVectorParameterValue(InKey.ParameterName, OldValue);
+	}
+}
+
+
+} // namespace UE::MovieScene
\ No newline at end of file
diff --git a/Engine/Source/Runtime/MovieSceneTracks/Public/Systems/MovieSceneMaterialParameterCollectionSystem.h b/Engine/Source/Runtime/MovieSceneTracks/Public/Systems/MovieSceneMaterialParameterCollectionSystem.h
--- a/Engine/Source/Runtime/MovieSceneTracks/Public/Systems/MovieSceneMaterialParameterCollectionSystem.h
+++ b/Engine/Source/Runtime/MovieSceneTracks/Public/Systems/MovieSceneMaterialParameterCollectionSystem.h
@@ -12,8 +12,8 @@
 namespace UE::MovieScene
 {
 
-struct FPreAnimatedMPCScalarStorage;
-struct FPreAnimatedMPCVectorStorage;
+struct FPreAnimatedScalarMaterialParameterStorage;
+struct FPreAnimatedVectorMaterialParameterStorage;
 
 } // namespace UE::MovieScene
 
@@ -43,7 +43,7 @@ private:
 private:
 
 	/** Holds pre-animated values for scalar MPC values */
-	TSharedPtr<UE::MovieScene::FPreAnimatedMPCScalarStorage> ScalarParameterStorage;
+	TSharedPtr<UE::MovieScene::FPreAnimatedScalarMaterialParameterStorage> ScalarParameterStorage;
 	/** Holds pre-animated values for vector or color MPC values */
-	TSharedPtr<UE::MovieScene::FPreAnimatedMPCVectorStorage> VectorParameterStorage;
+	TSharedPtr<UE::MovieScene::FPreAnimatedVectorMaterialParameterStorage> VectorParameterStorage;
 };
diff --git a/Engine/Source/Runtime/MovieSceneTracks/Public/Systems/MovieSceneMaterialParameterSystem.h b/Engine/Source/Runtime/MovieSceneTracks/Public/Systems/MovieSceneMaterialParameterSystem.h
--- a/Engine/Source/Runtime/MovieSceneTracks/Public/Systems/MovieSceneMaterialParameterSystem.h
+++ b/Engine/Source/Runtime/MovieSceneTracks/Public/Systems/MovieSceneMaterialParameterSystem.h
@@ -14,6 +14,9 @@ class UMovieScenePiecewiseDoubleBlenderSystem;
 namespace UE::MovieScene
 {
 
+struct FPreAnimatedScalarMaterialParameterStorage;
+struct FPreAnimatedVectorMaterialParameterStorage;
+
 struct FAnimatedMaterialParameterInfo
 {
 	/** Weak linker ptr - only assigned if the output entity is ever allocated */
@@ -69,6 +72,10 @@ private:
 	UE::MovieScene::TOverlappingEntityTracker<UE::MovieScene::FAnimatedMaterialParameterInfo, UObject*, FName> ScalarParameterTracker;
 	UE::MovieScene::TOverlappingEntityTracker<UE::MovieScene::FAnimatedMaterialParameterInfo, UObject*, FName> VectorParameterTracker;
 
+	/** Holds pre-animated values for scalar values */
+	TSharedPtr<UE::MovieScene::FPreAnimatedScalarMaterialParameterStorage> ScalarParameterStorage;
+	/** Holds pre-animated values for vector or color values */
+	TSharedPtr<UE::MovieScene::FPreAnimatedVectorMaterialParameterStorage> VectorParameterStorage;
 public:
 
 	UPROPERTY()
diff --git a/Engine/Source/Runtime/MovieSceneTracks/Public/Systems/MovieScenePreAnimatedMaterialParameters.h b/Engine/Source/Runtime/MovieSceneTracks/Public/Systems/MovieScenePreAnimatedMaterialParameters.h
new file mode 100644
--- /dev/null
+++ b/Engine/Source/Runtime/MovieSceneTracks/Public/Systems/MovieScenePreAnimatedMaterialParameters.h
@@ -0,0 +1,60 @@
+// Copyright Epic Games, Inc. All Rights Reserved.
+
+#pragma once
+
+#include "CoreTypes.h"
+#include "UObject/ObjectKey.h"
+#include "Evaluation/PreAnimatedState/MovieScenePreAnimatedStateStorage.h"
+
+namespace UE::MovieScene
+{
+
+struct FMaterialParameterKey
+{
+	FObjectKey BoundMaterial;
+	FName ParameterName;
+
+	MOVIESCENETRACKS_API FMaterialParameterKey(UObject* InBoundMaterial, const FName& InParameterName);
+
+	friend MOVIESCENETRACKS_API uint32 GetTypeHash(const FMaterialParameterKey& InKey);
+
+	friend MOVIESCENETRACKS_API bool operator==(const FMaterialParameterKey& A, const FMaterialParameterKey& B);
+};
+
+struct FMaterialParameterCollectionScalarTraits : FBoundObjectPreAnimatedStateTraits
+{
+	using KeyType = FMaterialParameterKey;
+	using StorageType = float;
+
+	static MOVIESCENETRACKS_API void ReplaceObject(FMaterialParameterKey& InOutKey, const FObjectKey& NewObject);
+
+	static MOVIESCENETRACKS_API float CachePreAnimatedValue(UObject* InBoundMaterial, const FName& ParameterName);
+
+	static MOVIESCENETRACKS_API void RestorePreAnimatedValue(const FMaterialParameterKey& InKey, float OldValue, const FRestoreStateParams& Params);
+};
+
+struct FMaterialParameterCollectionVectorTraits : FBoundObjectPreAnimatedStateTraits
+{
+	using KeyType = FMaterialParameterKey;
+	using StorageType = FLinearColor;
+
+	static MOVIESCENETRACKS_API void ReplaceObject(FMaterialParameterKey& InOutKey, const FObjectKey& NewObject);
+
+	static MOVIESCENETRACKS_API FLinearColor CachePreAnimatedValue(UObject* InBoundMaterial, const FName& ParameterName);
+
+	static MOVIESCENETRACKS_API void RestorePreAnimatedValue(const FMaterialParameterKey& InKey, const FLinearColor& OldValue, const FRestoreStateParams& Params);
+};
+
+struct FPreAnimatedScalarMaterialParameterStorage
+	: public TPreAnimatedStateStorage<FMaterialParameterCollectionScalarTraits>
+{
+	static MOVIESCENETRACKS_API TAutoRegisterPreAnimatedStorageID<FPreAnimatedScalarMaterialParameterStorage> StorageID;
+};
+
+struct FPreAnimatedVectorMaterialParameterStorage
+	: public TPreAnimatedStateStorage<FMaterialParameterCollectionVectorTraits>
+{
+	static MOVIESCENETRACKS_API TAutoRegisterPreAnimatedStorageID<FPreAnimatedVectorMaterialParameterStorage> StorageID;
+};
+
+} // namespace UE::MovieScene
\ No newline at end of file
