Fixed merge error with renamed type

diff --git a/Engine/Source/Runtime/MovieScene/Public/Evaluation/PreAnimatedState/MovieScenePreAnimatedStateStorage.h b/Engine/Source/Runtime/MovieScene/Public/Evaluation/PreAnimatedState/MovieScenePreAnimatedStateStorage.h
--- a/Engine/Source/Runtime/MovieScene/Public/Evaluation/PreAnimatedState/MovieScenePreAnimatedStateStorage.h
+++ b/Engine/Source/Runtime/MovieScene/Public/Evaluation/PreAnimatedState/MovieScenePreAnimatedStateStorage.h
@@ -50,7 +50,7 @@ struct TPreAnimatedStateTaskParams
 {
 	FEntityComponentFilter AdditionalFilter;
 
-	TSavePreAnimatedStateParams()
+	TPreAnimatedStateTaskParams()
 	{
 		AdditionalFilter.All({ FBuiltInComponentTypes::Get()->Tags.NeedsLink });
 	}
diff --git a/Engine/Source/Runtime/MovieSceneTracks/Private/Systems/MovieSceneMaterialParameterSystem.cpp b/Engine/Source/Runtime/MovieSceneTracks/Private/Systems/MovieSceneMaterialParameterSystem.cpp
--- a/Engine/Source/Runtime/MovieSceneTracks/Private/Systems/MovieSceneMaterialParameterSystem.cpp
+++ b/Engine/Source/Runtime/MovieSceneTracks/Private/Systems/MovieSceneMaterialParameterSystem.cpp
@@ -503,7 +503,7 @@ void UMovieSceneMaterialParameterSystem::OnInstantiation()
 		TOverlappingMaterialParameterHandler<FVectorMixin> VectorHandler(this);
 		VectorParameterTracker.ProcessInvalidatedOutputs(Linker, VectorHandler);
 
-		TSavePreAnimatedStateParams<UObject*, FName> Params;
+		TPreAnimatedStateTaskParams<UObject*, FName> Params;
 		Params.AdditionalFilter.Reset();
 		Params.AdditionalFilter.None({ TracksComponents->MPC });
 		Params.AdditionalFilter.Any({ BuiltInComponents->Tags.NeedsLink, TracksComponents->Tags.BoundMaterialChanged });
