Add native accessors for common float/color light properties to improve perf when setting them after sequencer evaluation.

diff --git a/Engine/Source/Runtime/MovieSceneTracks/Private/MovieSceneTracksComponentTypes.cpp b/Engine/Source/Runtime/MovieSceneTracks/Private/MovieSceneTracksComponentTypes.cpp
--- a/Engine/Source/Runtime/MovieSceneTracks/Private/MovieSceneTracksComponentTypes.cpp
+++ b/Engine/Source/Runtime/MovieSceneTracks/Private/MovieSceneTracksComponentTypes.cpp
@@ -4,6 +4,7 @@
 #include "Components/ExponentialHeightFogComponent.h"
 #include "Components/LightComponent.h"
 #include "Components/PrimitiveComponent.h"
+#include "Components/SkyAtmosphereComponent.h"
 #include "Components/SkyLightComponent.h"
 #include "EntitySystem/BuiltInComponentTypes.h"
 #include "EntitySystem/MovieSceneEntityManager.h"
@@ -221,6 +222,18 @@ void SetLightComponentLightColor(UObject* Object, EColorPropertyType InColorType
 	LightComponent->SetLightColor(InColor.GetLinearColor(), bConvertBackToSRgb);
 }
 
+float GetLightComponentIntensity(const UObject* Object)
+{
+	const ULightComponent* LightComponent = CastChecked<const ULightComponent>(Object);
+	return LightComponent->Intensity;
+}
+
+void SetLightComponentIntensity(UObject* Object, float InIntensity)
+{
+	ULightComponent* LightComponent = CastChecked<ULightComponent>(Object);
+	LightComponent->SetIntensity(InIntensity);
+}
+
 FIntermediateColor GetSkyLightComponentLightColor(const UObject* Object, EColorPropertyType InColorType)
 {
 	ensure(InColorType == EColorPropertyType::Color);
@@ -241,6 +254,158 @@ void SetSkyLightComponentLightColor(UObject* Object, EColorPropertyType InColorT
 	SkyLightComponent->SetLightColor(InColor.GetLinearColor());
 }
 
+float GetSkyLightComponentIntensity(const UObject* Object)
+{
+	const USkyLightComponent* SkyLightComponent = CastChecked<const USkyLightComponent>(Object);
+	return SkyLightComponent->Intensity;
+}
+
+void SetSkyLightComponentIntensity(UObject* Object, float InIntensity)
+{
+	USkyLightComponent* SkyLightComponent = CastChecked<USkyLightComponent>(Object);
+	SkyLightComponent->SetIntensity(InIntensity);
+}
+
+float GetSkyAtmosphereComponentMieScatteringScale(const UObject* Object)
+{
+	const USkyAtmosphereComponent* SkyAtmosphereComponent = CastChecked<const USkyAtmosphereComponent>(Object);
+	return SkyAtmosphereComponent->MieScatteringScale;
+}
+
+void SetSkyAtmosphereComponentMieScatteringScale(UObject* Object, float InMieScatteringScale)
+{
+	USkyAtmosphereComponent* SkyAtmosphereComponent = CastChecked<USkyAtmosphereComponent>(Object);
+	SkyAtmosphereComponent->SetMieScatteringScale(InMieScatteringScale);
+}
+
+float GetSkyAtmosphereComponentOtherAbsorptionScale(const UObject* Object)
+{
+	const USkyAtmosphereComponent* SkyAtmosphereComponent = CastChecked<const USkyAtmosphereComponent>(Object);
+	return SkyAtmosphereComponent->OtherAbsorptionScale;
+}
+
+void SetSkyAtmosphereComponentOtherAbsorptionScale(UObject* Object, float InOtherAbsorptionScale)
+{
+	USkyAtmosphereComponent* SkyAtmosphereComponent = CastChecked<USkyAtmosphereComponent>(Object);
+	SkyAtmosphereComponent->SetOtherAbsorptionScale(InOtherAbsorptionScale);
+}
+
+float GetSkyAtmosphereComponentRayleighScatteringScale(const UObject* Object)
+{
+	const USkyAtmosphereComponent* SkyAtmosphereComponent = CastChecked<const USkyAtmosphereComponent>(Object);
+	return SkyAtmosphereComponent->RayleighScatteringScale;
+}
+
+void SetSkyAtmosphereComponentRayleighScatteringScale(UObject* Object, float InRayleighScatteringScale)
+{
+	USkyAtmosphereComponent* SkyAtmosphereComponent = CastChecked<USkyAtmosphereComponent>(Object);
+	SkyAtmosphereComponent->SetRayleighScatteringScale(InRayleighScatteringScale);
+}
+
+float GetSkyAtmosphereComponentRayleighExponentialDistribution(const UObject* Object)
+{
+	const USkyAtmosphereComponent* SkyAtmosphereComponent = CastChecked<const USkyAtmosphereComponent>(Object);
+	return SkyAtmosphereComponent->RayleighExponentialDistribution;
+}
+
+void SetSkyAtmosphereComponentRayleighExponentialDistribution(UObject* Object, float InRayleighExponentialDistribution)
+{
+	USkyAtmosphereComponent* SkyAtmosphereComponent = CastChecked<USkyAtmosphereComponent>(Object);
+	SkyAtmosphereComponent->SetRayleighExponentialDistribution(InRayleighExponentialDistribution);
+}
+
+float GetSkyAtmosphereComponentAerialPerspectiveViewDistanceScale(const UObject* Object)
+{
+	const USkyAtmosphereComponent* SkyAtmosphereComponent = CastChecked<const USkyAtmosphereComponent>(Object);
+	return SkyAtmosphereComponent->AerialPespectiveViewDistanceScale;
+}
+
+void SetSkyAtmosphereComponentAerialPerspectiveViewDistanceScale(UObject* Object, float InAerialPerspectiveViewDistanceScale)
+{
+	USkyAtmosphereComponent* SkyAtmosphereComponent = CastChecked<USkyAtmosphereComponent>(Object);
+	SkyAtmosphereComponent->SetAerialPespectiveViewDistanceScale(InAerialPerspectiveViewDistanceScale);
+}
+
+FIntermediateColor GetSkyAtmosphereComponentMieAbsorption(const UObject* Object, EColorPropertyType InColorType)
+{
+	// USkyAtmosphereComponent::MieAbsorption is a FLinearColor. FIntermediateColor is always blended and dealt with
+	// in linear space so it should be safe to reinterpret it as linear.
+	ensure(InColorType == EColorPropertyType::Linear);
+
+	const USkyAtmosphereComponent* SkyAtmosphereComponent = CastChecked<const USkyAtmosphereComponent>(Object);
+	return FIntermediateColor(SkyAtmosphereComponent->MieAbsorption);
+}
+
+void SetSkyAtmosphereComponentMieAbsorption(UObject* Object, EColorPropertyType InColorType, const FIntermediateColor& InColor)
+{
+	// USkyAtmosphereComponent::MieAbsorption is a FLinearColor. FIntermediateColor is always blended and dealt with
+	// in linear space so it should be safe to reinterpret it as linear.
+	ensure(InColorType == EColorPropertyType::Linear);
+
+	USkyAtmosphereComponent* SkyAtmosphereComponent = CastChecked<USkyAtmosphereComponent>(Object);
+	SkyAtmosphereComponent->SetMieAbsorption(InColor.GetLinearColor());
+}
+
+FIntermediateColor GetSkyAtmosphereComponentMieScattering(const UObject* Object, EColorPropertyType InColorType)
+{
+	// USkyAtmosphereComponent::MieScattering is a FLinearColor. FIntermediateColor is always blended and dealt with
+	// in linear space so it should be safe to reinterpret it as linear.
+	ensure(InColorType == EColorPropertyType::Linear);
+
+	const USkyAtmosphereComponent* SkyAtmosphereComponent = CastChecked<const USkyAtmosphereComponent>(Object);
+	return FIntermediateColor(SkyAtmosphereComponent->MieScattering);
+}
+
+void SetSkyAtmosphereComponentMieScattering(UObject* Object, EColorPropertyType InColorType, const FIntermediateColor& InColor)
+{
+	// USkyAtmosphereComponent::MieScattering is a FLinearColor. FIntermediateColor is always blended and dealt with
+	// in linear space so it should be safe to reinterpret it as linear.
+	ensure(InColorType == EColorPropertyType::Linear);
+
+	USkyAtmosphereComponent* SkyAtmosphereComponent = CastChecked<USkyAtmosphereComponent>(Object);
+	SkyAtmosphereComponent->SetMieScattering(InColor.GetLinearColor());
+}
+
+FIntermediateColor GetSkyAtmosphereComponentRayleighScattering(const UObject* Object, EColorPropertyType InColorType)
+{
+	// USkyAtmosphereComponent::RayleighScattering is a FLinearColor. FIntermediateColor is always blended and dealt with
+	// in linear space so it should be safe to reinterpret it as linear.
+	ensure(InColorType == EColorPropertyType::Linear);
+
+	const USkyAtmosphereComponent* SkyAtmosphereComponent = CastChecked<const USkyAtmosphereComponent>(Object);
+	return FIntermediateColor(SkyAtmosphereComponent->RayleighScattering);
+}
+
+void SetSkyAtmosphereComponentRayleighScattering(UObject* Object, EColorPropertyType InColorType, const FIntermediateColor& InColor)
+{
+	// USkyAtmosphereComponent::RayleighScattering is a FLinearColor. FIntermediateColor is always blended and dealt with
+	// in linear space so it should be safe to reinterpret it as linear.
+	ensure(InColorType == EColorPropertyType::Linear);
+
+	USkyAtmosphereComponent* SkyAtmosphereComponent = CastChecked<USkyAtmosphereComponent>(Object);
+	SkyAtmosphereComponent->SetRayleighScattering(InColor.GetLinearColor());
+}
+
+FIntermediateColor GetSkyAtmosphereComponentSkyLuminanceFactor(const UObject* Object, EColorPropertyType InColorType)
+{
+	// USkyAtmosphereComponent::SkyLuminanceFactor is a FLinearColor. FIntermediateColor is always blended and dealt with
+	// in linear space so it should be safe to reinterpret it as linear.
+	ensure(InColorType == EColorPropertyType::Linear);
+
+	const USkyAtmosphereComponent* SkyAtmosphereComponent = CastChecked<const USkyAtmosphereComponent>(Object);
+	return FIntermediateColor(SkyAtmosphereComponent->SkyLuminanceFactor);
+}
+
+void SetSkyAtmosphereComponentSkyLuminanceFactor(UObject* Object, EColorPropertyType InColorType, const FIntermediateColor& InColor)
+{
+	// USkyAtmosphereComponent::SkyLuminanceFactor is a FLinearColor. FIntermediateColor is always blended and dealt with
+	// in linear space so it should be safe to reinterpret it as linear.
+	ensure(InColorType == EColorPropertyType::Linear);
+
+	USkyAtmosphereComponent* SkyAtmosphereComponent = CastChecked<USkyAtmosphereComponent>(Object);
+	SkyAtmosphereComponent->SetSkyLuminanceFactor(InColor.GetLinearColor());
+}
+
 float GetMassScale(const UObject* Object)
 {
 	const UPrimitiveComponent* PrimitiveComponent = CastChecked<const UPrimitiveComponent>(Object);
@@ -253,6 +418,74 @@ void SetMassScale(UObject* Object, float InMassScale)
 	PrimitiveComponent->SetMassScale(NAME_None, InMassScale);
 }
 
+float GetExponentialHeightFogComponentFogDensity(const UObject* Object)
+{
+	const UExponentialHeightFogComponent* ExponentialHeightFogComponent = CastChecked<const UExponentialHeightFogComponent>(Object);
+	return ExponentialHeightFogComponent->FogDensity; 
+}
+
+void SetExponentialHeightFogComponentFogDensity(UObject* Object, float InFogDensity)
+{
+	UExponentialHeightFogComponent* ExponentialHeightFogComponent = CastChecked<UExponentialHeightFogComponent>(Object);
+	ExponentialHeightFogComponent->SetFogDensity(InFogDensity);
+}
+
+float GetExponentialHeightFogComponentFogHeightFalloff(const UObject* Object)
+{
+	const UExponentialHeightFogComponent* ExponentialHeightFogComponent = CastChecked<const UExponentialHeightFogComponent>(Object);
+	return ExponentialHeightFogComponent->FogHeightFalloff;
+}
+
+void SetExponentialHeightFogComponentFogHeightFalloff(UObject* Object, float InFogHeightFalloff)
+{
+	UExponentialHeightFogComponent* ExponentialHeightFogComponent = CastChecked<UExponentialHeightFogComponent>(Object);
+	ExponentialHeightFogComponent->SetFogHeightFalloff(InFogHeightFalloff);
+}
+
+float GetExponentialHeightFogComponentFogMaxOpacity(const UObject* Object)
+{
+	const UExponentialHeightFogComponent* ExponentialHeightFogComponent = CastChecked<const UExponentialHeightFogComponent>(Object);
+	return ExponentialHeightFogComponent->FogMaxOpacity;
+}
+
+void SetExponentialHeightFogComponentFogMaxOpacity(UObject* Object, float InFogMaxOpacity)
+{
+	UExponentialHeightFogComponent* ExponentialHeightFogComponent = CastChecked<UExponentialHeightFogComponent>(Object);
+	ExponentialHeightFogComponent->SetFogMaxOpacity(InFogMaxOpacity);
+}
+
+float GetExponentialHeightFogComponentDirectionalInscatteringExponent(const UObject* Object)
+{
+	const UExponentialHeightFogComponent* ExponentialHeightFogComponent = CastChecked<const UExponentialHeightFogComponent>(Object);
+	return ExponentialHeightFogComponent->DirectionalInscatteringExponent;
+}
+
+void SetExponentialHeightFogComponentDirectionalInscatteringExponent(UObject* Object, float InDirectionalInscatteringExponent)
+{
+	UExponentialHeightFogComponent* ExponentialHeightFogComponent = CastChecked<UExponentialHeightFogComponent>(Object);
+	ExponentialHeightFogComponent->SetDirectionalInscatteringExponent(InDirectionalInscatteringExponent);
+}
+
+FIntermediateColor GetExponentialHeightFogVolumetricFogAlbedo(const UObject* Object, EColorPropertyType InColorType)
+{
+	// UExponentialHeightFog::VolumetricFogAlbedo is an FColor. FIntermediateColor is always blended and dealt with
+	// in linear space, so reinterpret the color.
+	ensure(InColorType == EColorPropertyType::Color);
+
+	const UExponentialHeightFogComponent* ExponentialHeightFogComponent = CastChecked<const UExponentialHeightFogComponent>(Object);
+	return FIntermediateColor(ExponentialHeightFogComponent->VolumetricFogAlbedo);
+}
+
+void SetExponentialHeightFogVolumetricFogAlbedo(UObject* Object, EColorPropertyType InColorType, const FIntermediateColor& InColor)
+{
+	// UExponentialHeightFog::VolumetricFogAlbedo is an FColor. FIntermediateColor is always blended and dealt with
+	// in linear space, so reinterpret the color back to FColor when setting via the public interface.
+	ensure(InColorType == EColorPropertyType::Color);
+
+	UExponentialHeightFogComponent* ExponentialHeightFogComponent = CastChecked<UExponentialHeightFogComponent>(Object);
+	ExponentialHeightFogComponent->SetVolumetricFogAlbedo(InColor.GetColor());
+}
+
 float GetSecondFogDataFogDensity(const UObject* Object)
 {
 	const UExponentialHeightFogComponent* ExponentialHeightFogComponent = CastChecked<const UExponentialHeightFogComponent>(Object);
@@ -696,18 +929,74 @@ FMovieSceneTracksComponentTypes::FMovieSceneTracksComponentTypes()
 	.Commit(FColorParameterHandler());
 
 	// We have some custom accessors for well-known types.
+
+	// LightComponent
 	Accessors.Color.Add(
 			ULightComponent::StaticClass(), GET_MEMBER_NAME_CHECKED(ULightComponent, LightColor), 
 			GetLightComponentLightColor, SetLightComponentLightColor);
+	Accessors.Float.Add(
+			ULightComponent::StaticClass(), GET_MEMBER_NAME_CHECKED(ULightComponent, Intensity),
+			GetLightComponentIntensity, SetLightComponentIntensity);
+
+	// SkyLightComponent
 	Accessors.Color.Add(
 			USkyLightComponent::StaticClass(), GET_MEMBER_NAME_CHECKED(USkyLightComponent, LightColor), 
 			GetSkyLightComponentLightColor, SetSkyLightComponentLightColor);
-	
+	Accessors.Float.Add(
+			USkyLightComponent::StaticClass(), GET_MEMBER_NAME_CHECKED(USkyLightComponent, Intensity),
+			GetSkyLightComponentIntensity, SetSkyLightComponentIntensity);
+
+	// SkyAtmosphereComponent
+	Accessors.Float.Add(
+			USkyAtmosphereComponent::StaticClass(), GET_MEMBER_NAME_CHECKED(USkyAtmosphereComponent, MieScatteringScale),
+			GetSkyAtmosphereComponentMieScatteringScale, SetSkyAtmosphereComponentMieScatteringScale);
+	Accessors.Float.Add(
+			USkyAtmosphereComponent::StaticClass(), GET_MEMBER_NAME_CHECKED(USkyAtmosphereComponent, OtherAbsorptionScale),
+			GetSkyAtmosphereComponentOtherAbsorptionScale, SetSkyAtmosphereComponentOtherAbsorptionScale);
+	Accessors.Float.Add(
+			USkyAtmosphereComponent::StaticClass(), GET_MEMBER_NAME_CHECKED(USkyAtmosphereComponent, RayleighExponentialDistribution),
+			GetSkyAtmosphereComponentRayleighExponentialDistribution, SetSkyAtmosphereComponentRayleighExponentialDistribution);
+	Accessors.Float.Add(
+			USkyAtmosphereComponent::StaticClass(), GET_MEMBER_NAME_CHECKED(USkyAtmosphereComponent, RayleighScatteringScale),
+			GetSkyAtmosphereComponentRayleighScatteringScale, SetSkyAtmosphereComponentRayleighScatteringScale);
+	Accessors.Float.Add(
+			USkyAtmosphereComponent::StaticClass(), GET_MEMBER_NAME_CHECKED(USkyAtmosphereComponent, AerialPespectiveViewDistanceScale),
+			GetSkyAtmosphereComponentAerialPerspectiveViewDistanceScale, SetSkyAtmosphereComponentAerialPerspectiveViewDistanceScale);
+	Accessors.Color.Add(
+			USkyAtmosphereComponent::StaticClass(), GET_MEMBER_NAME_CHECKED(USkyAtmosphereComponent, MieAbsorption),
+			GetSkyAtmosphereComponentMieAbsorption, SetSkyAtmosphereComponentMieAbsorption);
+	Accessors.Color.Add(
+			USkyAtmosphereComponent::StaticClass(), GET_MEMBER_NAME_CHECKED(USkyAtmosphereComponent, MieScattering),
+			GetSkyAtmosphereComponentMieScattering, SetSkyAtmosphereComponentMieScattering);
+	Accessors.Color.Add(
+			USkyAtmosphereComponent::StaticClass(), GET_MEMBER_NAME_CHECKED(USkyAtmosphereComponent, RayleighScattering),
+			GetSkyAtmosphereComponentRayleighScattering, SetSkyAtmosphereComponentRayleighScattering);
+	Accessors.Color.Add(
+			USkyAtmosphereComponent::StaticClass(), GET_MEMBER_NAME_CHECKED(USkyAtmosphereComponent, SkyLuminanceFactor),
+			GetSkyAtmosphereComponentSkyLuminanceFactor, SetSkyAtmosphereComponentSkyLuminanceFactor);
+
+	// PrimitiveComponent
 	const FString MassScalePath = FString::Printf(TEXT("%s.%s"), GET_MEMBER_NAME_STRING_CHECKED(UPrimitiveComponent, BodyInstance), GET_MEMBER_NAME_STRING_CHECKED(FBodyInstance, MassScale));
 	Accessors.Float.Add(
 		UPrimitiveComponent::StaticClass(), *MassScalePath,
 		GetMassScale, SetMassScale);
 
+	// ExponentialHeightFogComponent
+	Accessors.Float.Add(
+			UExponentialHeightFogComponent::StaticClass(), GET_MEMBER_NAME_CHECKED(UExponentialHeightFogComponent, FogDensity),
+			GetExponentialHeightFogComponentFogDensity, SetExponentialHeightFogComponentFogDensity);
+	Accessors.Float.Add(
+			UExponentialHeightFogComponent::StaticClass(), GET_MEMBER_NAME_CHECKED(UExponentialHeightFogComponent, FogHeightFalloff),
+			GetExponentialHeightFogComponentFogHeightFalloff, SetExponentialHeightFogComponentFogHeightFalloff);
+	Accessors.Float.Add(
+			UExponentialHeightFogComponent::StaticClass(), GET_MEMBER_NAME_CHECKED(UExponentialHeightFogComponent, DirectionalInscatteringExponent),
+			GetExponentialHeightFogComponentDirectionalInscatteringExponent, SetExponentialHeightFogComponentDirectionalInscatteringExponent);
+	Accessors.Float.Add(
+			UExponentialHeightFogComponent::StaticClass(), GET_MEMBER_NAME_CHECKED(UExponentialHeightFogComponent, FogMaxOpacity),
+			GetExponentialHeightFogComponentFogMaxOpacity, SetExponentialHeightFogComponentFogMaxOpacity);
+	Accessors.Color.Add(
+			UExponentialHeightFogComponent::StaticClass(), GET_MEMBER_NAME_CHECKED(UExponentialHeightFogComponent, VolumetricFogAlbedo),
+			GetExponentialHeightFogVolumetricFogAlbedo, SetExponentialHeightFogVolumetricFogAlbedo);
 	const FString SecondFogDataFogDensityPath = FString::Printf(TEXT("%s.%s"), GET_MEMBER_NAME_STRING_CHECKED(UExponentialHeightFogComponent, SecondFogData), GET_MEMBER_NAME_STRING_CHECKED(FExponentialHeightFogData, FogDensity));
 	Accessors.Float.Add(
 			UExponentialHeightFogComponent::StaticClass(), *SecondFogDataFogDensityPath,
