Sequencer: clean-up old blender type tag before adding a new/different one, and prevent unnecessary upgrades of properties to the blend path

diff --git a/Engine/Source/Runtime/MovieSceneTracks/Private/Systems/MovieScenePropertyInstantiator.cpp b/Engine/Source/Runtime/MovieSceneTracks/Private/Systems/MovieScenePropertyInstantiator.cpp
--- a/Engine/Source/Runtime/MovieSceneTracks/Private/Systems/MovieScenePropertyInstantiator.cpp
+++ b/Engine/Source/Runtime/MovieSceneTracks/Private/Systems/MovieScenePropertyInstantiator.cpp
@@ -202,6 +202,15 @@ void UMovieScenePropertyInstantiatorSystem::DiscoverInvalidatedProperties(TBitAr
 				OutInvalidatedProperties[PropertyIndex] = true;
 
 				this->Contributors.Remove(PropertyIndex, EntityID);
+				
+				// The output entity of a property could be going away if that property is using the
+				// fast path, which means its sole contributor is that output entity. It can go away when
+				// that track or section gets reimported.
+				FObjectPropertyInfo& PropertyInfo = this->ResolvedProperties[PropertyIndex];
+				if (PropertyInfo.PropertyEntityID == EntityID)
+				{
+					PropertyInfo.PropertyEntityID = FMovieSceneEntityID::Invalid();
+				}
 			}
 
 			// Always remove the entity ID from the LUT
@@ -702,6 +711,9 @@ void UMovieScenePropertyInstantiatorSystem::InitializeBlendPath(const FPropertyP
 
 			Params.MakeOutputComponentType(Linker->EntityManager, Composites, OutputMutation.AddMask);
 
+			// Remove the old blender type tag before add the new one
+			OutputMutation.AddMask.Remove({ SetupResult.PreviousInfo.BlenderTypeTag });
+
 			FEntityBuilder()
 			.Add(BuiltInComponents->BlendChannelOutput, NewBlendChannel)
 			.AddTag(SetupResult.CurrentInfo.BlenderTypeTag)
