Sequencer: All UObject references within entity-component data are now weak by default

Bound objects are now always accompanied by a BoundObjectKey component that is used for validation during garbage collection, and will be used to remove stale objects.
UMovieSceneTrackInstances now reference their animating object by weak pointer
Introduced a new component type, FObjectComponent, that can be used to reference both strong or weak pointers depending on how it is initialized. This is now used by bound materials.

diff --git a/Engine/Source/Runtime/MovieScene/Private/EntitySystem/BuiltInComponentTypes.cpp b/Engine/Source/Runtime/MovieScene/Private/EntitySystem/BuiltInComponentTypes.cpp
--- a/Engine/Source/Runtime/MovieScene/Private/EntitySystem/BuiltInComponentTypes.cpp
+++ b/Engine/Source/Runtime/MovieScene/Private/EntitySystem/BuiltInComponentTypes.cpp
@@ -17,9 +17,51 @@ namespace UE
 namespace MovieScene
 {
 
+bool FObjectComponent::IsStrongReference() const
+{
+	return ObjectKey == FObjectKey();
+}
+
+UObject* FObjectComponent::GetObject() const
+{
+	if (IsStrongReference())
+	{
+		return ObjectPtr;
+	}
+	return ObjectKey.ResolveObjectPtr();
+}
+
+void AddReferencedObjectForComponent(FReferenceCollector& ReferenceCollector, FObjectComponent* ComponentData)
+{
+	if (ComponentData->IsStrongReference())
+	{
+		ReferenceCollector.AddReferencedObject(ComponentData->ObjectPtr);
+	}
+}
+void AddReferencedObjectForBoundObject(FReferenceCollector& ReferenceCollector, void* ComponentData, int32 Num)
+{
+	// Intentionally hidden from the reference graph
+}
+
 static bool GMovieSceneBuiltInComponentTypesDestroyed = false;
 static TUniquePtr<FBuiltInComponentTypes> GMovieSceneBuiltInComponentTypes;
 
+struct FBoundObjectKeyInitializer : IMutualComponentInitializer
+{
+	void Run(const FEntityRange& Range, const FEntityAllocationWriteContext& WriteContext) override
+	{
+		FBuiltInComponentTypes* BuiltInComponents = FBuiltInComponentTypes::Get();
+
+		TComponentReader<UObject*>   Objects    = Range.Allocation->ReadComponents(BuiltInComponents->BoundObject);
+		TComponentWriter<FObjectKey> ObjectKeys = Range.Allocation->WriteComponents(BuiltInComponents->BoundObjectKey, WriteContext);
+
+		for (int32 Index = 0; Index < Range.Num; ++Index)
+		{
+			ObjectKeys[Range.ComponentStartOffset + Index] = Objects[Range.ComponentStartOffset + Index];
+		}
+	}
+};
+
 FBuiltInComponentTypes::FBuiltInComponentTypes()
 {
 	using namespace UE::MovieScene;
@@ -34,7 +76,8 @@ FBuiltInComponentTypes::FBuiltInComponentTypes()
 	ComponentRegistry->NewComponentType(&EvalTime,              TEXT("Eval Time"));
 	ComponentRegistry->NewComponentType(&EvalSeconds,           TEXT("Eval Seconds"));
 
-	ComponentRegistry->NewComponentType(&BoundObject,           TEXT("Bound Object"));
+	ComponentRegistry->NewComponentType(&BoundObjectKey,        TEXT("Bound Object Key"));
+	ComponentRegistry->NewComponentType(&BoundObject,           TEXT("Bound Object"), FNewComponentTypeParams(&AddReferencedObjectForBoundObject, EComponentTypeFlags::None));
 
 	ComponentRegistry->NewComponentType(&PropertyBinding,         TEXT("Property Binding"), EComponentTypeFlags::CopyToOutput);
 	ComponentRegistry->NewComponentType(&GenericObjectBinding,    TEXT("Generic Object Binding ID"));
@@ -159,6 +202,14 @@ FBuiltInComponentTypes::FBuiltInComponentTypes()
 
 	FinishedMask.SetAll({ Tags.NeedsUnlink, Tags.Finished });
 
+	{
+		FMutuallyInclusiveComponentParams ObjectKeyParams;
+		ObjectKeyParams.CustomInitializer = MakeUnique<FBoundObjectKeyInitializer>();
+		ObjectKeyParams.Type = EMutuallyInclusiveComponentType::Mandatory;
+
+		ComponentRegistry->Factories.DefineMutuallyInclusiveComponents(BoundObject, { BoundObjectKey }, MoveTemp(ObjectKeyParams));
+	}
+
 	// New children always need link
 	ComponentRegistry->Factories.DefineChildComponent(Tags.NeedsLink);
 
diff --git a/Engine/Source/Runtime/MovieScene/Private/EntitySystem/MovieSceneEntitySystemLinker.cpp b/Engine/Source/Runtime/MovieScene/Private/EntitySystem/MovieSceneEntitySystemLinker.cpp
--- a/Engine/Source/Runtime/MovieScene/Private/EntitySystem/MovieSceneEntitySystemLinker.cpp
+++ b/Engine/Source/Runtime/MovieScene/Private/EntitySystem/MovieSceneEntitySystemLinker.cpp
@@ -293,8 +293,9 @@ void UMovieSceneEntitySystemLinker::TagInvalidBoundObjects()
 	// Tag any bound objects that are now invalid
 	TArray<FMovieSceneEntityID> ExpiredBoundObjects;
 
-	auto Iter = [&ExpiredBoundObjects](FMovieSceneEntityID EntityID, UObject* BoundObject)
+	auto Iter = [&ExpiredBoundObjects](FMovieSceneEntityID EntityID, const FObjectKey& BoundObjectKey)
 	{
+		UObject* BoundObject = BoundObjectKey.ResolveObjectPtr();
 		if (FBuiltInComponentTypes::IsBoundObjectGarbage(BoundObject))
 		{
 			ExpiredBoundObjects.Add(EntityID);
@@ -302,7 +303,7 @@ void UMovieSceneEntitySystemLinker::TagInvalidBoundObjects()
 	};
 	FEntityTaskBuilder()
 	.ReadEntityIDs()
-	.Read(BuiltInComponents->BoundObject)
+	.Read(BuiltInComponents->BoundObjectKey)
 	.Iterate_PerEntity(&EntityManager, Iter);
 
 	for (FMovieSceneEntityID Entity : ExpiredBoundObjects)
diff --git a/Engine/Source/Runtime/MovieScene/Private/EntitySystem/TrackInstance/MovieSceneTrackInstance.cpp b/Engine/Source/Runtime/MovieScene/Private/EntitySystem/TrackInstance/MovieSceneTrackInstance.cpp
--- a/Engine/Source/Runtime/MovieScene/Private/EntitySystem/TrackInstance/MovieSceneTrackInstance.cpp
+++ b/Engine/Source/Runtime/MovieScene/Private/EntitySystem/TrackInstance/MovieSceneTrackInstance.cpp
@@ -13,7 +13,7 @@ void UMovieSceneTrackInstance::Initialize(UObject* InAnimatedObject, UMovieScene
 {
 	// We make the difference between a root track instance, and a bound track instance that lost its binding,
 	// by remembering if this track instance was initialized with a valid bound object.
-	AnimatedObject = InAnimatedObject;
+	WeakAnimatedObject = InAnimatedObject;
 	bIsRootTrackInstance = (InAnimatedObject == nullptr);
 
 	PrivateLinker = InLinker;
@@ -30,7 +30,7 @@ void UMovieSceneTrackInstance::Destroy()
 {
 	using namespace UE::MovieScene;
 
-	if (bIsRootTrackInstance || !UE::MovieScene::FBuiltInComponentTypes::IsBoundObjectGarbage(AnimatedObject))
+	if (bIsRootTrackInstance || !UE::MovieScene::FBuiltInComponentTypes::IsBoundObjectGarbage(WeakAnimatedObject.Get()))
 	{
 		OnDestroyed();
 	}
@@ -152,6 +152,7 @@ void UMovieSceneTrackInstance::UpdateInputs(TArray<FMovieSceneTrackInstanceInput
 
 UWorld* UMovieSceneTrackInstance::GetWorld() const
 {
+	UObject* AnimatedObject = WeakAnimatedObject.Get();
 	return AnimatedObject ? AnimatedObject->GetWorld() : Super::GetWorld();
 }
 
diff --git a/Engine/Source/Runtime/MovieScene/Private/Evaluation/PreAnimatedState/MovieScenePreAnimatedObjectStorage.cpp b/Engine/Source/Runtime/MovieScene/Private/Evaluation/PreAnimatedState/MovieScenePreAnimatedObjectStorage.cpp
--- a/Engine/Source/Runtime/MovieScene/Private/Evaluation/PreAnimatedState/MovieScenePreAnimatedObjectStorage.cpp
+++ b/Engine/Source/Runtime/MovieScene/Private/Evaluation/PreAnimatedState/MovieScenePreAnimatedObjectStorage.cpp
@@ -19,6 +19,11 @@ FPreAnimatedStorageGroupHandle FBoundObjectPreAnimatedStateTraits::MakeGroupImpl
 	return ObjectGroupManager->MakeGroupForKey(BoundObject);
 }
 
+FPreAnimatedStorageGroupHandle FBoundObjectPreAnimatedStateTraits::MakeGroupImpl(const FObjectComponent& BoundObject)
+{
+	return ObjectGroupManager->MakeGroupForKey(BoundObject.GetObject());
+}
+
 } // namespace MovieScene
 } // namespace UE
 
diff --git a/Engine/Source/Runtime/MovieScene/Public/EntitySystem/BuiltInComponentTypes.h b/Engine/Source/Runtime/MovieScene/Public/EntitySystem/BuiltInComponentTypes.h
--- a/Engine/Source/Runtime/MovieScene/Public/EntitySystem/BuiltInComponentTypes.h
+++ b/Engine/Source/Runtime/MovieScene/Public/EntitySystem/BuiltInComponentTypes.h
@@ -216,6 +216,98 @@ struct FSourceObjectPathChannel
 	const FMovieSceneObjectPathChannel* Source;
 };
 
+/** A component that represents a UObject* either as a strong or weak reference */
+struct FObjectComponent
+{
+	FObjectComponent()
+		: ObjectPtr(nullptr)
+	{}
+
+	/** Construct a new null component */
+	static FObjectComponent Null()
+	{
+		return FObjectComponent{ nullptr, FObjectKey() };
+	}
+
+	/** Construct a new strongly-referenced component from an object ptr */
+	static FObjectComponent Strong(UObject* InObject)
+	{
+		return FObjectComponent{ InObject, FObjectKey() };
+	}
+
+	/** Construct a new weakly-referenced component from an object ptr */
+	static FObjectComponent Weak(UObject* InObject)
+	{
+		return FObjectComponent{ InObject, InObject };
+	}
+
+	/** Check whether this object component is valid */
+	explicit operator bool() const
+	{
+		return ObjectPtr != nullptr;
+	}
+
+	/** Compare this component with another object ptr */
+	bool operator==(UObject* InObject) const
+	{
+		return ObjectPtr == InObject;
+	}
+
+	/** Compare this component with another object ptr */
+	bool operator!=(UObject* InObject) const
+	{
+		return ObjectPtr != InObject;
+	}
+
+	/** Equality operator */
+	friend bool operator==(const FObjectComponent& A, const FObjectComponent& B)
+	{
+		if (A.ObjectKey != FObjectKey() || B.ObjectKey != FObjectKey())
+		{
+			return A.ObjectKey == B.ObjectKey;
+		}
+
+		return A.ObjectPtr == B.ObjectPtr;
+	}
+
+	/** Generate a type has from this component */
+	friend uint32 GetTypeHash(const FObjectComponent& In)
+	{
+		if (In.ObjectKey != FObjectKey())
+		{
+			return GetTypeHash(In.ObjectKey);
+		}
+
+		return GetTypeHash(In.ObjectPtr);
+	}
+
+	UObject* operator->() const
+	{
+		UObject* Object = GetObject();
+		checkSlow(Object);
+		return Object;
+	}
+
+	MOVIESCENE_API UObject* GetObject() const;
+
+	/** Conditionally add a reference for the specified component data based on whether it is strongly referenced or not */
+	friend void AddReferencedObjectForComponent(FReferenceCollector& ReferenceCollector, FObjectComponent* ComponentData);
+
+private:
+
+	FObjectComponent(UObject* InObject, const FObjectKey& InObjectKey)
+		: ObjectPtr(InObject)
+		, ObjectKey(InObjectKey)
+	{}
+
+	bool IsStrongReference() const;
+
+	/** Raw pointer to the object. This is explicitly assigned for both strong and weak ptrs, but is only added to the reference graph when FObjectKey() is default constructed.  */
+	UObject* ObjectPtr;
+
+	/** Default constructed for strong pointers. Assigned on construction for weak ptrs. */
+	FObjectKey ObjectKey;
+};
 
 /**
  * Pre-defined built in component types
@@ -238,7 +330,12 @@ public:
 
 	TComponentTypeID<FMovieSceneEntityID> ParentEntity;
 
+	/**
+	 * A bound object ptr component that defines the object being animated. This ptr is explicitly hidden from the reference graph and cleaned up
+	 * after a garbage collection pass if it becomes invalid by checking the BoundObjectKey component that must exist alongside it
+	 */
 	TComponentTypeID<UObject*>            BoundObject;
+	TComponentTypeID<FObjectKey>          BoundObjectKey;
 
 	TComponentTypeID<FInstanceHandle>     InstanceHandle;
 
diff --git a/Engine/Source/Runtime/MovieScene/Public/EntitySystem/MovieSceneComponentRegistry.h b/Engine/Source/Runtime/MovieScene/Public/EntitySystem/MovieSceneComponentRegistry.h
--- a/Engine/Source/Runtime/MovieScene/Public/EntitySystem/MovieSceneComponentRegistry.h
+++ b/Engine/Source/Runtime/MovieScene/Public/EntitySystem/MovieSceneComponentRegistry.h
@@ -20,6 +20,28 @@ namespace MovieScene
 
 template<typename> struct TPropertyComponents;
 
+struct FNewComponentTypeParams
+{
+	FNewComponentTypeParams()
+		: ReferenceCollectionCallback(nullptr)
+		, Flags(EComponentTypeFlags::None)
+	{}
+
+	/** Implicit construction from type flags to support legacy API */
+	FNewComponentTypeParams(EComponentTypeFlags InFlags)
+		: ReferenceCollectionCallback(nullptr)
+		, Flags(InFlags)
+	{}
+
+	explicit FNewComponentTypeParams(FComponentReferenceCollectionPtr RefCollectionPtr, EComponentTypeFlags InFlags)
+		: ReferenceCollectionCallback(RefCollectionPtr)
+		, Flags(InFlags)
+	{}
+
+	FComponentReferenceCollectionPtr ReferenceCollectionCallback;
+	EComponentTypeFlags Flags;
+};
+
 struct MOVIESCENE_API FComponentRegistry
 {
 public:
@@ -39,8 +61,8 @@ public:
 	 * Define a new tag type using the specified information. Tags have 0 memory overhead.
 	 * @note Transitory tag types must be unregistered when no longer required by calling DestroyComponentTypeSafe or Unsafe to prevent leaking component type IDs
 	 *
-	 * @param Flags          Flags relating to the new component type
 	 * @param DebugName      A developer friendly name that accompanies this component type for debugging purposes
+	 * @param Flags          Flags relating to the new component type
 	 * @return A new component type identifier for the tag
 	 */
 	FComponentTypeID NewTag(const TCHAR* const DebugName, EComponentTypeFlags Flags = EComponentTypeFlags::None);
@@ -51,16 +73,16 @@ public:
 	 * @note Transitory tag types must be unregistered when no longer required by calling DestroyComponentTypeSafe or Unsafe to prevent leaking component type IDs
 	 *
 	 * @param DebugName      A developer friendly name that accompanies this component type for debugging purposes
-	 * @param Flags          (Optional) Flags relating to the new component type
+	 * @param Params         (Optional) Parameters for the type including component flags
 	 * @return A new component type identifier for the tag
 	 */
 	template<typename T>
-	TComponentTypeID<T> NewComponentType(const TCHAR* const DebugName, EComponentTypeFlags Flags = EComponentTypeFlags::None);
+	TComponentTypeID<T> NewComponentType(const TCHAR* const DebugName, const FNewComponentTypeParams& Params = FNewComponentTypeParams());
 
 	template<typename T>
-	void NewComponentType(TComponentTypeID<T>* Ref, const TCHAR* const DebugName, EComponentTypeFlags Flags = EComponentTypeFlags::None)
+	void NewComponentType(TComponentTypeID<T>* Ref, const TCHAR* const DebugName, const FNewComponentTypeParams& Params = FNewComponentTypeParams())
 	{
-		*Ref = NewComponentType<T>(DebugName, Flags);
+		*Ref = NewComponentType<T>(DebugName, Params);
 	}
 
 	template<typename PropertyTraits>
diff --git a/Engine/Source/Runtime/MovieScene/Public/EntitySystem/MovieSceneComponentTypeInfo.h b/Engine/Source/Runtime/MovieScene/Public/EntitySystem/MovieSceneComponentTypeInfo.h
--- a/Engine/Source/Runtime/MovieScene/Public/EntitySystem/MovieSceneComponentTypeInfo.h
+++ b/Engine/Source/Runtime/MovieScene/Public/EntitySystem/MovieSceneComponentTypeInfo.h
@@ -18,6 +18,7 @@ namespace MovieScene
 struct FNotImplemented;
 struct FComponentTypeDebugInfo;
 
+using FComponentReferenceCollectionPtr = void (*)(FReferenceCollector&, void*, int32);
 
 /**
  * Stub for components that do not need reference collection. Overload void AddReferencedObjectForComponent(FReferenceCollector* ReferenceCollector, T* Component) for your own implementation
@@ -183,10 +184,19 @@ struct FComponentTypeInfo
 		ComplexComponentOps = MakeUnique<TComplexComponentOps<T>>();
 	}
 
+	/**
+	 * Define complex component with a specific reference collection callback
+	 */
+	template<typename T>
+	void MakeComplexComponentOps(FComponentReferenceCollectionPtr RefCollectionPtr)
+	{
+		ComplexComponentOps = MakeUnique<TComplexComponentOpsCustomRefCollection<T>>(RefCollectionPtr);
+	}
+
 private:
 
 	template<typename T>
-	struct TComplexComponentOps : IComplexComponentOps
+	struct TComplexComponentOpsBase : IComplexComponentOps
 	{
 		virtual void ConstructItems(void* Components, int32 Num) const
 		{
@@ -224,6 +234,10 @@ private:
 				++TypedSrc;
 			}
 		}
+	};
+	template<typename T>
+	struct TComplexComponentOps : TComplexComponentOpsBase<T>
+	{
 		virtual void AddReferencedObjects(FReferenceCollector& ReferenceCollector, void* ComponentStart, int32 Num)
 		{
 			T* ComponentData = static_cast<T*>(ComponentStart);
@@ -235,6 +249,19 @@ private:
 			}
 		}
 	};
+	template<typename T>
+	struct TComplexComponentOpsCustomRefCollection : TComplexComponentOpsBase<T>
+	{
+		FComponentReferenceCollectionPtr RefCollectionPtr;
+		TComplexComponentOpsCustomRefCollection(FComponentReferenceCollectionPtr InRefCollectionPtr)
+			: RefCollectionPtr(InRefCollectionPtr)
+		{}
+
+		virtual void AddReferencedObjects(FReferenceCollector& ReferenceCollector, void* ComponentStart, int32 Num)
+		{
+			(*RefCollectionPtr)(ReferenceCollector, ComponentStart, Num);
+		}
+	};
 };
 
 
diff --git a/Engine/Source/Runtime/MovieScene/Public/EntitySystem/MovieSceneEntityFactoryTemplates.h b/Engine/Source/Runtime/MovieScene/Public/EntitySystem/MovieSceneEntityFactoryTemplates.h
--- a/Engine/Source/Runtime/MovieScene/Public/EntitySystem/MovieSceneEntityFactoryTemplates.h
+++ b/Engine/Source/Runtime/MovieScene/Public/EntitySystem/MovieSceneEntityFactoryTemplates.h
@@ -207,7 +207,7 @@ inline void FEntityFactories::DefineChildComponent(TComponentTypeID<ParentCompon
 }
 
 template<typename T>
-TComponentTypeID<T> FComponentRegistry::NewComponentType(const TCHAR* const DebugName, EComponentTypeFlags Flags)
+TComponentTypeID<T> FComponentRegistry::NewComponentType(const TCHAR* const DebugName, const FNewComponentTypeParams& Params)
 {
 	static const uint32 ComponentTypeSize = sizeof(T);
 	static_assert(ComponentTypeSize < TNumericLimits<decltype(FComponentTypeInfo::Sizeof)>::Max(), "Type too large to be used as component data");
@@ -222,10 +222,10 @@ TComponentTypeID<T> FComponentRegistry::NewComponentType(const TCHAR* const Debu
 	NewTypeInfo.bIsZeroConstructType       = TIsZeroConstructType<T>::Value;
 	NewTypeInfo.bIsTriviallyDestructable   = TIsTriviallyDestructible<T>::Value;
 	NewTypeInfo.bIsTriviallyCopyAssignable = TIsTriviallyCopyAssignable<T>::Value;
-	NewTypeInfo.bIsPreserved               = EnumHasAnyFlags(Flags, EComponentTypeFlags::Preserved);
-	NewTypeInfo.bIsCopiedToOutput          = EnumHasAnyFlags(Flags, EComponentTypeFlags::CopyToOutput);
-	NewTypeInfo.bIsMigratedToOutput        = EnumHasAnyFlags(Flags, EComponentTypeFlags::MigrateToOutput);
-	NewTypeInfo.bHasReferencedObjects      = THasAddReferencedObjectForComponent<T>::Value;
+	NewTypeInfo.bIsPreserved               = EnumHasAnyFlags(Params.Flags, EComponentTypeFlags::Preserved);
+	NewTypeInfo.bIsCopiedToOutput          = EnumHasAnyFlags(Params.Flags, EComponentTypeFlags::CopyToOutput);
+	NewTypeInfo.bIsMigratedToOutput        = EnumHasAnyFlags(Params.Flags, EComponentTypeFlags::MigrateToOutput);
+	NewTypeInfo.bHasReferencedObjects      = Params.ReferenceCollectionCallback != nullptr || THasAddReferencedObjectForComponent<T>::Value;
 
 #if UE_MOVIESCENE_ENTITY_DEBUG
 	NewTypeInfo.DebugInfo                = MakeUnique<FComponentTypeDebugInfo>();
@@ -234,7 +234,11 @@ TComponentTypeID<T> FComponentRegistry::NewComponentType(const TCHAR* const Debu
 	NewTypeInfo.DebugInfo->Type          = TComponentDebugType<T>::Type;
 #endif
 
-	if (!NewTypeInfo.bIsZeroConstructType || !NewTypeInfo.bIsTriviallyDestructable || !NewTypeInfo.bIsTriviallyCopyAssignable || NewTypeInfo.bHasReferencedObjects)
+	if (Params.ReferenceCollectionCallback)
+	{
+		NewTypeInfo.MakeComplexComponentOps<T>(Params.ReferenceCollectionCallback);
+	}
+	else if (!NewTypeInfo.bIsZeroConstructType || !NewTypeInfo.bIsTriviallyDestructable || !NewTypeInfo.bIsTriviallyCopyAssignable || NewTypeInfo.bHasReferencedObjects)
 	{
 		NewTypeInfo.MakeComplexComponentOps<T>();
 	}
@@ -242,7 +246,7 @@ TComponentTypeID<T> FComponentRegistry::NewComponentType(const TCHAR* const Debu
 	FComponentTypeID    NewTypeID = NewComponentTypeInternal(MoveTemp(NewTypeInfo));
 	TComponentTypeID<T> TypedTypeID = NewTypeID.ReinterpretCast<T>();
 
-	if (EnumHasAnyFlags(Flags, EComponentTypeFlags::CopyToChildren))
+	if (EnumHasAnyFlags(Params.Flags, EComponentTypeFlags::CopyToChildren))
 	{
 		Factories.DefineChildComponent(TDuplicateChildEntityInitializer<T>(TypedTypeID));
 	}
diff --git a/Engine/Source/Runtime/MovieScene/Public/EntitySystem/TrackInstance/MovieSceneTrackInstance.h b/Engine/Source/Runtime/MovieScene/Public/EntitySystem/TrackInstance/MovieSceneTrackInstance.h
--- a/Engine/Source/Runtime/MovieScene/Public/EntitySystem/TrackInstance/MovieSceneTrackInstance.h
+++ b/Engine/Source/Runtime/MovieScene/Public/EntitySystem/TrackInstance/MovieSceneTrackInstance.h
@@ -142,7 +142,7 @@ public:
 	 */
 	UObject* GetAnimatedObject() const
 	{
-		return AnimatedObject;
+		return WeakAnimatedObject.Get();
 	}
 
 
@@ -191,7 +191,7 @@ private:
 private:
 
 	UPROPERTY()
-	TObjectPtr<UObject> AnimatedObject;
+	TWeakObjectPtr<UObject> WeakAnimatedObject;
 
 	UPROPERTY()
 	bool bIsRootTrackInstance;
diff --git a/Engine/Source/Runtime/MovieScene/Public/Evaluation/PreAnimatedState/MovieScenePreAnimatedObjectStorage.h b/Engine/Source/Runtime/MovieScene/Public/Evaluation/PreAnimatedState/MovieScenePreAnimatedObjectStorage.h
--- a/Engine/Source/Runtime/MovieScene/Public/Evaluation/PreAnimatedState/MovieScenePreAnimatedObjectStorage.h
+++ b/Engine/Source/Runtime/MovieScene/Public/Evaluation/PreAnimatedState/MovieScenePreAnimatedObjectStorage.h
@@ -38,7 +38,14 @@ struct MOVIESCENE_API FBoundObjectPreAnimatedStateTraits : FPreAnimatedStateTrai
 	{
 		return MakeGroupImpl(BoundObject);
 	}
+	template<typename... T>
+	FPreAnimatedStorageGroupHandle MakeGroup(const FObjectComponent& BoundObject, T&&... Unused)
+	{
+		return MakeGroupImpl(BoundObject);
+	}
+
 	FPreAnimatedStorageGroupHandle MakeGroupImpl(UObject* BoundObject);
+	FPreAnimatedStorageGroupHandle MakeGroupImpl(const FObjectComponent& BoundObject);
 
 	template<typename ...T>
 	void ReplaceObject(TTuple<FObjectKey, T...>& InOutKey, const FObjectKey& NewObject)
diff --git a/Engine/Source/Runtime/MovieSceneTracks/Private/Systems/MovieSceneMaterialParameterCollectionSystem.cpp b/Engine/Source/Runtime/MovieSceneTracks/Private/Systems/MovieSceneMaterialParameterCollectionSystem.cpp
--- a/Engine/Source/Runtime/MovieSceneTracks/Private/Systems/MovieSceneMaterialParameterCollectionSystem.cpp
+++ b/Engine/Source/Runtime/MovieSceneTracks/Private/Systems/MovieSceneMaterialParameterCollectionSystem.cpp
@@ -74,7 +74,7 @@ void UMovieSceneMaterialParameterCollectionSystem::OnRun(FSystemTaskPrerequisite
 		}
 		virtual void InitializeAllocation(FEntityAllocation* Allocation, const FComponentMask& AllocationType) const
 		{
-			TComponentWriter<UObject*> OutBoundMaterials = Allocation->WriteComponents(TracksComponents->BoundMaterial, FEntityAllocationWriteContext::NewAllocation());
+			TComponentWriter<FObjectComponent> OutBoundMaterials = Allocation->WriteComponents(TracksComponents->BoundMaterial, FEntityAllocationWriteContext::NewAllocation());
 			TComponentReader<UMaterialParameterCollection*> MPCs = Allocation->ReadComponents(TracksComponents->MPC);
 			TComponentReader<FInstanceHandle> InstanceHandles = Allocation->ReadComponents(BuiltInComponents->InstanceHandle);
 
@@ -85,7 +85,7 @@ void UMovieSceneMaterialParameterCollectionSystem::OnRun(FSystemTaskPrerequisite
 			const int32 Num = Allocation->Num();
 			for (int32 Index = 0; Index < Num; ++Index)
 			{
-				OutBoundMaterials[Index] = nullptr;
+				OutBoundMaterials[Index] = FObjectComponent::Null();
 
 				UMaterialParameterCollection* Collection = MPCs[Index];
 				IMovieScenePlayer* Player = InstanceRegistry->GetInstance(InstanceHandles[Index]).GetPlayer();
@@ -94,7 +94,7 @@ void UMovieSceneMaterialParameterCollectionSystem::OnRun(FSystemTaskPrerequisite
 				if (World && Collection)
 				{
 					UMaterialParameterCollectionInstance* Instance = World->GetParameterCollectionInstance(Collection);
-					OutBoundMaterials[Index] = Instance;
+					OutBoundMaterials[Index] = FObjectComponent::Weak(Instance);
 
 #if !UE_BUILD_SHIPPING && !UE_BUILD_TEST
 					if (ensureAlwaysMsgf(Instance != nullptr,
@@ -177,7 +177,7 @@ void UMovieSceneMaterialParameterCollectionSystem::SavePreAnimatedState(const FP
 	FBuiltInComponentTypes*          BuiltInComponents = FBuiltInComponentTypes::Get();
 	FMovieSceneTracksComponentTypes* TracksComponents  = FMovieSceneTracksComponentTypes::Get();
 
-	TPreAnimatedStateTaskParams<UObject*, FName> Params;
+	TPreAnimatedStateTaskParams<FObjectComponent, FName> Params;
 
 	Params.AdditionalFilter.All({ TracksComponents->MPC });
 	ScalarParameterStorage->BeginTrackingEntitiesTask(Linker, Params, TracksComponents->BoundMaterial, TracksComponents->ScalarParameterName);
diff --git a/Engine/Source/Runtime/MovieSceneTracks/Private/Systems/MovieSceneMaterialParameterSystem.cpp b/Engine/Source/Runtime/MovieSceneTracks/Private/Systems/MovieSceneMaterialParameterSystem.cpp
--- a/Engine/Source/Runtime/MovieSceneTracks/Private/Systems/MovieSceneMaterialParameterSystem.cpp
+++ b/Engine/Source/Runtime/MovieSceneTracks/Private/Systems/MovieSceneMaterialParameterSystem.cpp
@@ -73,14 +73,14 @@ FAnimatedMaterialParameterInfo::~FAnimatedMaterialParameterInfo()
 /** Apply scalar material parameters */
 struct FApplyScalarParameters
 {
-	static void ForEachEntity(UObject* BoundMaterial, FName ParameterName, double InScalarValue)
+	static void ForEachEntity(const FObjectComponent& BoundMaterial, FName ParameterName, double InScalarValue)
 	{
 		// WARNING: BoundMaterial may be nullptr here
-		if (UMaterialInstanceDynamic* MID = Cast<UMaterialInstanceDynamic>(BoundMaterial))
+		if (UMaterialInstanceDynamic* MID = Cast<UMaterialInstanceDynamic>(BoundMaterial.GetObject()))
 		{
 			MID->SetScalarParameterValue(ParameterName, (float)InScalarValue);
 		}
-		else if (UMaterialParameterCollectionInstance* MPCI = Cast<UMaterialParameterCollectionInstance>(BoundMaterial))
+		else if (UMaterialParameterCollectionInstance* MPCI = Cast<UMaterialParameterCollectionInstance>(BoundMaterial.GetObject()))
 		{
 			MPCI->SetScalarParameterValue(ParameterName, (float)InScalarValue);
 		}
@@ -91,7 +91,7 @@ struct FApplyScalarParameters
 struct FApplyVectorParameters
 {
 	static void ForEachAllocation(FEntityAllocationIteratorItem Item,
-		TRead<UObject*> BoundMaterials,
+		TRead<FObjectComponent> BoundMaterials,
 		TReadOneOrMoreOf<FName, FName> VectorOrColorParameterNames,
 		TReadOneOrMoreOf<double, double, double, double> VectorChannels)
 	{
@@ -105,7 +105,7 @@ struct FApplyVectorParameters
 		
 		for (int32 Index = 0; Index < Num; ++Index)
 		{
-			UObject* BoundMaterial = BoundMaterials[Index];
+			const FObjectComponent& BoundMaterial = BoundMaterials[Index];
 			if (!BoundMaterial)
 			{
 				continue;
@@ -118,11 +118,11 @@ struct FApplyVectorParameters
 				B ? (float)B[Index] : 1.f,
 				A ? (float)A[Index] : 1.f
 			);
-			if (UMaterialInstanceDynamic* MID = Cast<UMaterialInstanceDynamic>(BoundMaterial))
+			if (UMaterialInstanceDynamic* MID = Cast<UMaterialInstanceDynamic>(BoundMaterial.GetObject()))
 			{
 				MID->SetVectorParameterValue(ParameterNames[Index], Color);
 			}
-			else if (UMaterialParameterCollectionInstance* MPCI = Cast<UMaterialParameterCollectionInstance>(BoundMaterial))
+			else if (UMaterialParameterCollectionInstance* MPCI = Cast<UMaterialParameterCollectionInstance>(BoundMaterial.GetObject()))
 			{
 				MPCI->SetVectorParameterValue(ParameterNames[Index], Color);
 			}
@@ -132,7 +132,7 @@ struct FApplyVectorParameters
 
 struct FScalarMixin
 {
-	void CreateEntity(UMovieSceneEntitySystemLinker* Linker, UObject* BoundMaterial, FName ParameterName, FComponentTypeID BlenderTypeTag, TArrayView<const FMovieSceneEntityID> Inputs, FAnimatedMaterialParameterInfo* Output)
+	void CreateEntity(UMovieSceneEntitySystemLinker* Linker, const FObjectComponent& BoundMaterial, FName ParameterName, FComponentTypeID BlenderTypeTag, TArrayView<const FMovieSceneEntityID> Inputs, FAnimatedMaterialParameterInfo* Output)
 	{
 		FBuiltInComponentTypes* BuiltInComponents = FBuiltInComponentTypes::Get();
 		FMovieSceneTracksComponentTypes* TracksComponents = FMovieSceneTracksComponentTypes::Get();
@@ -152,11 +152,11 @@ struct FScalarMixin
 
 		if (!bHasInitialValue)
 		{
-			if (UMaterialInstanceDynamic* MID = Cast<UMaterialInstanceDynamic>(BoundMaterial))
+			if (UMaterialInstanceDynamic* MID = Cast<UMaterialInstanceDynamic>(BoundMaterial.GetObject()))
 			{
 				bHasInitialValue = MID->GetScalarParameterValue(ParameterName, InitialValue);
 			}
-			else if (UMaterialParameterCollectionInstance* MPCI = Cast<UMaterialParameterCollectionInstance>(BoundMaterial))
+			else if (UMaterialParameterCollectionInstance* MPCI = Cast<UMaterialParameterCollectionInstance>(BoundMaterial.GetObject()))
 			{
 				bHasInitialValue = MPCI->GetScalarParameterValue(ParameterName, InitialValue);
 			}
@@ -176,19 +176,19 @@ struct FScalarMixin
 		Linker->EntityManager.CopyComponents(Inputs[0], Output->OutputEntityID, Linker->EntityManager.GetComponents()->GetCopyAndMigrationMask());
 	}
 
-	void InitializeSoleInput(UMovieSceneEntitySystemLinker* Linker, UObject* BoundMaterial, FName ParameterName, FMovieSceneEntityID SoleContributor, FAnimatedMaterialParameterInfo* Output)
+	void InitializeSoleInput(UMovieSceneEntitySystemLinker* Linker, const FObjectComponent& BoundMaterial, FName ParameterName, FMovieSceneEntityID SoleContributor, FAnimatedMaterialParameterInfo* Output)
 	{
 		if (Linker->EntityManager.HasComponent(SoleContributor, FBuiltInComponentTypes::Get()->Tags.AlwaysCacheInitialValue))
 		{
 			FMovieSceneTracksComponentTypes* TracksComponents = FMovieSceneTracksComponentTypes::Get();
 
 			float InitialValue = 0.0;
-			if (UMaterialInstanceDynamic* MID = Cast<UMaterialInstanceDynamic>(BoundMaterial))
+			if (UMaterialInstanceDynamic* MID = Cast<UMaterialInstanceDynamic>(BoundMaterial.GetObject()))
 			{
 				MID->GetScalarParameterValue(ParameterName, InitialValue);
 				Linker->EntityManager.AddComponent(SoleContributor, TracksComponents->FloatParameter.InitialValue, InitialValue);
 			}
-			else if (UMaterialParameterCollectionInstance* MPCI = Cast<UMaterialParameterCollectionInstance>(BoundMaterial))
+			else if (UMaterialParameterCollectionInstance* MPCI = Cast<UMaterialParameterCollectionInstance>(BoundMaterial.GetObject()))
 			{
 				MPCI->GetScalarParameterValue(ParameterName, InitialValue);
 				Linker->EntityManager.AddComponent(SoleContributor, TracksComponents->FloatParameter.InitialValue, InitialValue);
@@ -199,7 +199,7 @@ struct FScalarMixin
 
 struct FVectorMixin
 {
-	void CreateEntity(UMovieSceneEntitySystemLinker* Linker, UObject* BoundMaterial, FName ParameterName, FComponentTypeID BlenderTypeTag, TArrayView<const FMovieSceneEntityID> Inputs, FAnimatedMaterialParameterInfo* Output)
+	void CreateEntity(UMovieSceneEntitySystemLinker* Linker, const FObjectComponent& BoundMaterial, FName ParameterName, FComponentTypeID BlenderTypeTag, TArrayView<const FMovieSceneEntityID> Inputs, FAnimatedMaterialParameterInfo* Output)
 	{
 		FBuiltInComponentTypes* BuiltInComponents = FBuiltInComponentTypes::Get();
 		FMovieSceneTracksComponentTypes* TracksComponents = FMovieSceneTracksComponentTypes::Get();
@@ -220,11 +220,11 @@ struct FVectorMixin
 		if (!bHasInitialValue)
 		{
 			FLinearColor ColorValue;
-			if (UMaterialInstanceDynamic* MID = Cast<UMaterialInstanceDynamic>(BoundMaterial))
+			if (UMaterialInstanceDynamic* MID = Cast<UMaterialInstanceDynamic>(BoundMaterial.GetObject()))
 			{
 				bHasInitialValue = MID->GetVectorParameterValue(ParameterName, ColorValue);
 			}
-			else if (UMaterialParameterCollectionInstance* MPCI = Cast<UMaterialParameterCollectionInstance>(BoundMaterial))
+			else if (UMaterialParameterCollectionInstance* MPCI = Cast<UMaterialParameterCollectionInstance>(BoundMaterial.GetObject()))
 			{
 				bHasInitialValue = MPCI->GetVectorParameterValue(ParameterName, ColorValue);
 			}
@@ -252,21 +252,21 @@ struct FVectorMixin
 		Linker->EntityManager.CopyComponents(Inputs[0], Output->OutputEntityID, Linker->EntityManager.GetComponents()->GetCopyAndMigrationMask());
 	}
 
-	void InitializeSoleInput(UMovieSceneEntitySystemLinker* Linker, UObject* BoundMaterial, FName ParameterName, FMovieSceneEntityID SoleContributor, FAnimatedMaterialParameterInfo* Output)
+	void InitializeSoleInput(UMovieSceneEntitySystemLinker* Linker, const FObjectComponent& BoundMaterial, FName ParameterName, FMovieSceneEntityID SoleContributor, FAnimatedMaterialParameterInfo* Output)
 	{
 		if (Linker->EntityManager.HasComponent(SoleContributor, FBuiltInComponentTypes::Get()->Tags.AlwaysCacheInitialValue))
 		{
 			FMovieSceneTracksComponentTypes* TracksComponents = FMovieSceneTracksComponentTypes::Get();
 
 			FLinearColor ColorValue = FLinearColor::White;
-			if (UMaterialInstanceDynamic* MID = Cast<UMaterialInstanceDynamic>(BoundMaterial))
+			if (UMaterialInstanceDynamic* MID = Cast<UMaterialInstanceDynamic>(BoundMaterial.GetObject()))
 			{
 				if (MID->GetVectorParameterValue(ParameterName, ColorValue))
 				{
 					Linker->EntityManager.AddComponent(SoleContributor, TracksComponents->ColorParameter.InitialValue, FIntermediateColor(ColorValue));
 				}
 			}
-			else if (UMaterialParameterCollectionInstance* MPCI = Cast<UMaterialParameterCollectionInstance>(BoundMaterial))
+			else if (UMaterialParameterCollectionInstance* MPCI = Cast<UMaterialParameterCollectionInstance>(BoundMaterial.GetObject()))
 			{
 				if (MPCI->GetVectorParameterValue(ParameterName, ColorValue))
 				{
@@ -289,12 +289,12 @@ struct TOverlappingMaterialParameterHandler : Mixin
 		, System(InSystem)
 	{}
 
-	void InitializeOutput(UObject* BoundMaterial, FName ParameterName, TArrayView<const FMovieSceneEntityID> Inputs, FAnimatedMaterialParameterInfo* Output, FEntityOutputAggregate Aggregate)
+	void InitializeOutput(const FObjectComponent& BoundMaterial, FName ParameterName, TArrayView<const FMovieSceneEntityID> Inputs, FAnimatedMaterialParameterInfo* Output, FEntityOutputAggregate Aggregate)
 	{
 		UpdateOutput(BoundMaterial, ParameterName, Inputs, Output, Aggregate);
 	}
 
-	void UpdateOutput(UObject* BoundMaterial, FName ParameterName, TArrayView<const FMovieSceneEntityID> Inputs, FAnimatedMaterialParameterInfo* Output, FEntityOutputAggregate Aggregate)
+	void UpdateOutput(const FObjectComponent& BoundMaterial, FName ParameterName, TArrayView<const FMovieSceneEntityID> Inputs, FAnimatedMaterialParameterInfo* Output, FEntityOutputAggregate Aggregate)
 	{
 		using namespace UE::MovieScene;
 
@@ -357,7 +357,7 @@ struct TOverlappingMaterialParameterHandler : Mixin
 		Output->NumContributors = NumContributors;
 	}
 
-	void DestroyOutput(UObject* BoundMaterial, FName ParameterName, FAnimatedMaterialParameterInfo* Output, FEntityOutputAggregate Aggregate)
+	void DestroyOutput(const FObjectComponent& BoundMaterial, FName ParameterName, FAnimatedMaterialParameterInfo* Output, FEntityOutputAggregate Aggregate)
 	{
 		if (Output->OutputEntityID)
 		{
@@ -466,7 +466,7 @@ void UMovieSceneMaterialParameterSystem::OnInstantiation()
 			this->VectorParameterTracker.VisitUnlinkedAllocation(Allocation);
 		};
 
-		auto HandleUpdatedAllocation = [this](const FEntityAllocation* Allocation, TRead<UObject*> InObjects, TReadOneOf<FName, FName, FName> ScalarVectorOrColorParameterNames)
+		auto HandleUpdatedAllocation = [this](const FEntityAllocation* Allocation, TRead<FObjectComponent> InObjects, TReadOneOf<FName, FName, FName> ScalarVectorOrColorParameterNames)
 		{
 			if (TComponentPtr<const FName> ScalarParameterNames = TComponentPtr<const FName>(ScalarVectorOrColorParameterNames.Get<0>()))
 			{
@@ -503,7 +503,7 @@ void UMovieSceneMaterialParameterSystem::OnInstantiation()
 		TOverlappingMaterialParameterHandler<FVectorMixin> VectorHandler(this);
 		VectorParameterTracker.ProcessInvalidatedOutputs(Linker, VectorHandler);
 
-		TPreAnimatedStateTaskParams<UObject*, FName> Params;
+		TPreAnimatedStateTaskParams<FObjectComponent, FName> Params;
 		Params.AdditionalFilter.Reset();
 		Params.AdditionalFilter.None({ TracksComponents->MPC });
 		Params.AdditionalFilter.Any({ BuiltInComponents->Tags.NeedsLink, TracksComponents->Tags.BoundMaterialChanged });
diff --git a/Engine/Source/Runtime/MovieSceneTracks/Private/Systems/MovieScenePreAnimatedMaterialParameters.cpp b/Engine/Source/Runtime/MovieSceneTracks/Private/Systems/MovieScenePreAnimatedMaterialParameters.cpp
--- a/Engine/Source/Runtime/MovieSceneTracks/Private/Systems/MovieScenePreAnimatedMaterialParameters.cpp
+++ b/Engine/Source/Runtime/MovieSceneTracks/Private/Systems/MovieScenePreAnimatedMaterialParameters.cpp
@@ -12,8 +12,8 @@ namespace UE::MovieScene
 TAutoRegisterPreAnimatedStorageID<FPreAnimatedScalarMaterialParameterStorage> FPreAnimatedScalarMaterialParameterStorage::StorageID;
 TAutoRegisterPreAnimatedStorageID<FPreAnimatedVectorMaterialParameterStorage> FPreAnimatedVectorMaterialParameterStorage::StorageID;
 
-FMaterialParameterKey::FMaterialParameterKey(UObject* InBoundMaterial, const FName& InParameterName)
-	: BoundMaterial(InBoundMaterial)
+FMaterialParameterKey::FMaterialParameterKey(const FObjectComponent& InBoundMaterial, const FName& InParameterName)
+	: BoundMaterial(InBoundMaterial.GetObject())
 	, ParameterName(InParameterName)
 {}
 
@@ -32,14 +32,16 @@ void FMaterialParameterCollectionScalarTraits::ReplaceObject(FMaterialParameterK
 	InOutKey.BoundMaterial = NewObject.ResolveObjectPtr();
 }
 
-float FMaterialParameterCollectionScalarTraits::CachePreAnimatedValue(UObject* InBoundMaterial, const FName& ParameterName)
+float FMaterialParameterCollectionScalarTraits::CachePreAnimatedValue(const FObjectComponent& InBoundMaterial, const FName& ParameterName)
 {
+	UObject* BoundMaterial = InBoundMaterial.GetObject();
+
 	float ParameterValue = 0.f;
-	if (UMaterialParameterCollectionInstance* MPCI = Cast<UMaterialParameterCollectionInstance>(InBoundMaterial))
+	if (UMaterialParameterCollectionInstance* MPCI = Cast<UMaterialParameterCollectionInstance>(BoundMaterial))
 	{
 		MPCI->GetScalarParameterValue(ParameterName, ParameterValue);
 	}
-	else if (UMaterialInterface* MaterialInterface = Cast<UMaterialInterface>(InBoundMaterial))
+	else if (UMaterialInterface* MaterialInterface = Cast<UMaterialInterface>(BoundMaterial))
 	{
 		MaterialInterface->GetScalarParameterValue(FMaterialParameterInfo(ParameterName), ParameterValue);
 	}
@@ -65,14 +67,16 @@ void FMaterialParameterCollectionVectorTraits::ReplaceObject(FMaterialParameterK
 	InOutKey.BoundMaterial = NewObject.ResolveObjectPtr();
 }
 
-FLinearColor FMaterialParameterCollectionVectorTraits::CachePreAnimatedValue(UObject* InBoundMaterial, const FName& ParameterName)
+FLinearColor FMaterialParameterCollectionVectorTraits::CachePreAnimatedValue(const FObjectComponent& InBoundMaterial, const FName& ParameterName)
 {
+	UObject* BoundMaterial = InBoundMaterial.GetObject();
+
 	FLinearColor ParameterValue = FLinearColor::White;
-	if (UMaterialParameterCollectionInstance* MPCI = Cast<UMaterialParameterCollectionInstance>(InBoundMaterial))
+	if (UMaterialParameterCollectionInstance* MPCI = Cast<UMaterialParameterCollectionInstance>(BoundMaterial))
 	{
 		MPCI->GetVectorParameterValue(ParameterName, ParameterValue);
 	}
-	else if (UMaterialInterface* MaterialInterface = Cast<UMaterialInterface>(InBoundMaterial))
+	else if (UMaterialInterface* MaterialInterface = Cast<UMaterialInterface>(BoundMaterial))
 	{
 		MaterialInterface->GetVectorParameterValue(FMaterialParameterInfo(ParameterName), ParameterValue);
 	}
diff --git a/Engine/Source/Runtime/MovieSceneTracks/Public/MovieSceneTracksComponentTypes.h b/Engine/Source/Runtime/MovieSceneTracks/Public/MovieSceneTracksComponentTypes.h
--- a/Engine/Source/Runtime/MovieSceneTracks/Public/MovieSceneTracksComponentTypes.h
+++ b/Engine/Source/Runtime/MovieSceneTracks/Public/MovieSceneTracksComponentTypes.h
@@ -457,7 +457,7 @@ struct MOVIESCENETRACKS_API FMovieSceneTracksComponentTypes
 	TComponentTypeID<FName> ColorParameterName;
 	TComponentTypeID<FName> TransformParameterName;
 
-	TComponentTypeID<UObject*> BoundMaterial;
+	TComponentTypeID<FObjectComponent> BoundMaterial;
 	TComponentTypeID<UMaterialParameterCollection*> MPC;
 
 	TComponentTypeID<FFadeComponentData> Fade;
diff --git a/Engine/Source/Runtime/MovieSceneTracks/Public/Systems/MovieSceneMaterialParameterSystem.h b/Engine/Source/Runtime/MovieSceneTracks/Public/Systems/MovieSceneMaterialParameterSystem.h
--- a/Engine/Source/Runtime/MovieSceneTracks/Public/Systems/MovieSceneMaterialParameterSystem.h
+++ b/Engine/Source/Runtime/MovieSceneTracks/Public/Systems/MovieSceneMaterialParameterSystem.h
@@ -69,8 +69,8 @@ private:
 private:
 
 	/** Overlapping trackers that track multiple entities animating the same bound object and name */
-	UE::MovieScene::TOverlappingEntityTracker<UE::MovieScene::FAnimatedMaterialParameterInfo, UObject*, FName> ScalarParameterTracker;
-	UE::MovieScene::TOverlappingEntityTracker<UE::MovieScene::FAnimatedMaterialParameterInfo, UObject*, FName> VectorParameterTracker;
+	UE::MovieScene::TOverlappingEntityTracker<UE::MovieScene::FAnimatedMaterialParameterInfo, UE::MovieScene::FObjectComponent, FName> ScalarParameterTracker;
+	UE::MovieScene::TOverlappingEntityTracker<UE::MovieScene::FAnimatedMaterialParameterInfo, UE::MovieScene::FObjectComponent, FName> VectorParameterTracker;
 
 	/** Holds pre-animated values for scalar values */
 	TSharedPtr<UE::MovieScene::FPreAnimatedScalarMaterialParameterStorage> ScalarParameterStorage;
diff --git a/Engine/Source/Runtime/MovieSceneTracks/Public/Systems/MovieSceneMaterialSystem.h b/Engine/Source/Runtime/MovieSceneTracks/Public/Systems/MovieSceneMaterialSystem.h
--- a/Engine/Source/Runtime/MovieSceneTracks/Public/Systems/MovieSceneMaterialSystem.h
+++ b/Engine/Source/Runtime/MovieSceneTracks/Public/Systems/MovieSceneMaterialSystem.h
@@ -155,12 +155,12 @@ struct TApplyMaterialSwitchers
 template<typename AccessorType, typename... RequiredComponents>
 struct TInitializeBoundMaterials
 {
-	static bool InitializeBoundMaterial(typename TCallTraits<RequiredComponents>::ParamType... Inputs, UObject*& OutDynamicMaterial)
+	static bool InitializeBoundMaterial(typename TCallTraits<RequiredComponents>::ParamType... Inputs, FObjectComponent& OutDynamicMaterial)
 	{
 		AccessorType Accessor(Inputs...);
 		if (!Accessor)
 		{
-			OutDynamicMaterial = nullptr;
+			OutDynamicMaterial = FObjectComponent::Null();
 			return false;
 		}
 
@@ -168,7 +168,7 @@ struct TInitializeBoundMaterials
 
 		if (!ExistingMaterial)
 		{
-			OutDynamicMaterial = nullptr;
+			OutDynamicMaterial = FObjectComponent::Null();
 			return true;
 		}
 
@@ -176,13 +176,13 @@ struct TInitializeBoundMaterials
 		{
 			if (OutDynamicMaterial != MID)
 			{
-				OutDynamicMaterial = MID;
+				OutDynamicMaterial = FObjectComponent::Weak(MID);
 				return true;
 			}
 			return false;
 		}
 
-		UMaterialInstanceDynamic* CurrentMID = Cast<UMaterialInstanceDynamic>(OutDynamicMaterial);
+		UMaterialInstanceDynamic* CurrentMID = Cast<UMaterialInstanceDynamic>(OutDynamicMaterial.GetObject());
 		if (CurrentMID && CurrentMID->Parent == ExistingMaterial)
 		{
 			Accessor.SetMaterial(CurrentMID);
@@ -190,12 +190,12 @@ struct TInitializeBoundMaterials
 		}
 		
 		UMaterialInstanceDynamic* NewMaterial = Accessor.CreateDynamicMaterial(ExistingMaterial);
-		OutDynamicMaterial = NewMaterial;
+		OutDynamicMaterial = FObjectComponent::Weak(NewMaterial);
 		Accessor.SetMaterial(NewMaterial);
 		return true;
 	}
 
-	static void ForEachEntity(typename TCallTraits<RequiredComponents>::ParamType... Inputs, UObject*& OutDynamicMaterial)
+	static void ForEachEntity(typename TCallTraits<RequiredComponents>::ParamType... Inputs, FObjectComponent& OutDynamicMaterial)
 	{
 		InitializeBoundMaterial(Inputs..., OutDynamicMaterial);
 	}
@@ -212,7 +212,7 @@ struct TReinitializeBoundMaterials
 		: Linker(InLinker)
 	{}
 
-	void ForEachAllocation(int32 Num, const FMovieSceneEntityID* EntityIDs, const RequiredComponents*... Inputs, UObject** Objects)
+	void ForEachAllocation(int32 Num, const FMovieSceneEntityID* EntityIDs, const RequiredComponents*... Inputs, FObjectComponent* Objects)
 	{
 		for (int32 Index = 0; Index < Num; ++Index)
 		{
@@ -220,7 +220,7 @@ struct TReinitializeBoundMaterials
 		}
 	}
 
-	void ForEachEntity(FMovieSceneEntityID EntityID, typename TCallTraits<RequiredComponents>::ParamType... Inputs, UObject*& OutDynamicMaterial)
+	void ForEachEntity(FMovieSceneEntityID EntityID, typename TCallTraits<RequiredComponents>::ParamType... Inputs, FObjectComponent& OutDynamicMaterial)
 	{
 		if (TInitializeBoundMaterials<AccessorType, RequiredComponents...>::InitializeBoundMaterial(Inputs..., OutDynamicMaterial))
 		{
@@ -256,16 +256,16 @@ struct TAddBoundMaterialMutationImpl<AccessorType, TIntegerSequence<int, Indices
 	}
 	virtual void InitializeAllocation(FEntityAllocation* Allocation, const FComponentMask& AllocationType) const
 	{
-		TComponentWriter<UObject*> BoundMaterials = Allocation->WriteComponents(TracksComponents->BoundMaterial, FEntityAllocationWriteContext::NewAllocation());
+		TComponentWriter<FObjectComponent> BoundMaterials = Allocation->WriteComponents(TracksComponents->BoundMaterial, FEntityAllocationWriteContext::NewAllocation());
 		InitializeAllocation(Allocation, BoundMaterials, Allocation->ReadComponents(ComponentTypes.template Get<Indices>())...);
 	}
 
-	void InitializeAllocation(FEntityAllocation* Allocation, UObject** OutBoundMaterials, const RequiredComponents*... InRequiredComponents) const
+	void InitializeAllocation(FEntityAllocation* Allocation, FObjectComponent* OutBoundMaterials, const RequiredComponents*... InRequiredComponents) const
 	{
 		const int32 Num = Allocation->Num();
 		for (int32 Index = 0; Index < Num; ++Index)
 		{
-			OutBoundMaterials[Index] = nullptr;
+			OutBoundMaterials[Index] = FObjectComponent::Null();
 		}
 	}
 private:
diff --git a/Engine/Source/Runtime/MovieSceneTracks/Public/Systems/MovieScenePreAnimatedMaterialParameters.h b/Engine/Source/Runtime/MovieSceneTracks/Public/Systems/MovieScenePreAnimatedMaterialParameters.h
--- a/Engine/Source/Runtime/MovieSceneTracks/Public/Systems/MovieScenePreAnimatedMaterialParameters.h
+++ b/Engine/Source/Runtime/MovieSceneTracks/Public/Systems/MovieScenePreAnimatedMaterialParameters.h
@@ -15,7 +15,7 @@ struct FMaterialParameterKey
 	FObjectKey BoundMaterial;
 	FName ParameterName;
 
-	MOVIESCENETRACKS_API FMaterialParameterKey(UObject* InBoundMaterial, const FName& InParameterName);
+	MOVIESCENETRACKS_API FMaterialParameterKey(const FObjectComponent& InBoundMaterial, const FName& InParameterName);
 
 	friend MOVIESCENETRACKS_API uint32 GetTypeHash(const FMaterialParameterKey& InKey);
 
@@ -29,7 +29,7 @@ struct FMaterialParameterCollectionScalarTraits : FBoundObjectPreAnimatedStateTr
 
 	static MOVIESCENETRACKS_API void ReplaceObject(FMaterialParameterKey& InOutKey, const FObjectKey& NewObject);
 
-	static MOVIESCENETRACKS_API float CachePreAnimatedValue(UObject* InBoundMaterial, const FName& ParameterName);
+	static MOVIESCENETRACKS_API float CachePreAnimatedValue(const FObjectComponent& InBoundMaterial, const FName& ParameterName);
 
 	static MOVIESCENETRACKS_API void RestorePreAnimatedValue(const FMaterialParameterKey& InKey, float OldValue, const FRestoreStateParams& Params);
 };
@@ -41,7 +41,7 @@ struct FMaterialParameterCollectionVectorTraits : FBoundObjectPreAnimatedStateTr
 
 	static MOVIESCENETRACKS_API void ReplaceObject(FMaterialParameterKey& InOutKey, const FObjectKey& NewObject);
 
-	static MOVIESCENETRACKS_API FLinearColor CachePreAnimatedValue(UObject* InBoundMaterial, const FName& ParameterName);
+	static MOVIESCENETRACKS_API FLinearColor CachePreAnimatedValue(const FObjectComponent& InBoundMaterial, const FName& ParameterName);
 
 	static MOVIESCENETRACKS_API void RestorePreAnimatedValue(const FMaterialParameterKey& InKey, const FLinearColor& OldValue, const FRestoreStateParams& Params);
 };
