Sequencer: Fixed failed assertion caused by considering NeedsUnlink entities when computing hierarchical easing weights

diff --git a/Engine/Source/Runtime/MovieSceneTracks/Private/Systems/WeightAndEasingEvaluatorSystem.cpp b/Engine/Source/Runtime/MovieSceneTracks/Private/Systems/WeightAndEasingEvaluatorSystem.cpp
--- a/Engine/Source/Runtime/MovieSceneTracks/Private/Systems/WeightAndEasingEvaluatorSystem.cpp
+++ b/Engine/Source/Runtime/MovieSceneTracks/Private/Systems/WeightAndEasingEvaluatorSystem.cpp
@@ -643,6 +643,7 @@ void UMovieSceneHierarchicalEasingInstantiatorSystem::OnRun(FSystemTaskPrerequis
 		.ReadEntityIDs()
 		.Read(BuiltInComponents->RootInstanceHandle)
 		.Read(BuiltInComponents->HierarchicalEasingProvider)
+		.FilterNone({ BuiltInComponents->Tags.NeedsUnlink })
 		.Iterate_PerEntity(&Linker->EntityManager, VisitEasingProvider);
 
 		// Step 2 : Iterate all sub instances and see if they exist within one of the previously collected easing providers
@@ -652,6 +653,7 @@ void UMovieSceneHierarchicalEasingInstantiatorSystem::OnRun(FSystemTaskPrerequis
 		.Read(BuiltInComponents->RootInstanceHandle)
 		.Read(BuiltInComponents->SequenceID)
 		.FilterAll({ BuiltInComponents->Tags.SubInstance })
+		.FilterNone({ BuiltInComponents->Tags.NeedsUnlink })
 		.Iterate_PerEntity(&Linker->EntityManager, VisitSubInstance);
 
 		// Step 3: Produce parent-first channel IDs for anything we collected, then pass on the buffer to the evaluator system
