Sequencer: Fixed material parameter pre-animated state being saved for blended Keep State tracks

  - The previous code ran 2 separate tasks: 1 to cache all the entities that contribute to a given material/parameter, and 1 to cache the actual pre-animated values of any material parameters.
  - For Keep State _blended_ material parameters the first of these tasks would not find anything (because it skips keep state entities), but it would find the output entity in the second task, and cache the value unbound to any particular track or sequence. When a restore state sequence was subsequently stopped it would restore this unbound state along with it.
  - This fix combines the BeginTracking and CacheValue tasks into the same task that explicitly excludes any BlendChannelOutputs, thus we only run it for unblended entities, or blend inputs.

diff --git a/Engine/Source/Runtime/MovieSceneTracks/Private/Systems/MovieSceneMaterialParameterSystem.cpp b/Engine/Source/Runtime/MovieSceneTracks/Private/Systems/MovieSceneMaterialParameterSystem.cpp
--- a/Engine/Source/Runtime/MovieSceneTracks/Private/Systems/MovieSceneMaterialParameterSystem.cpp
+++ b/Engine/Source/Runtime/MovieSceneTracks/Private/Systems/MovieSceneMaterialParameterSystem.cpp
@@ -503,18 +503,15 @@ void UMovieSceneMaterialParameterSystem::OnInstantiation()
 		TOverlappingMaterialParameterHandler<FVectorMixin> VectorHandler(this);
 		VectorParameterTracker.ProcessInvalidatedOutputs(Linker, VectorHandler);
 
+		// Gather inputs that contribute to the material parameter by excluding outputs (which will not have an instance handle)
 		TPreAnimatedStateTaskParams<FObjectComponent, FName> Params;
 		Params.AdditionalFilter.Reset();
-		Params.AdditionalFilter.None({ TracksComponents->MPC });
+		Params.AdditionalFilter.None({ TracksComponents->MPC, BuiltInComponents->BlendChannelOutput });
 		Params.AdditionalFilter.Any({ BuiltInComponents->Tags.NeedsLink, TracksComponents->Tags.BoundMaterialChanged });
-
-		ScalarParameterStorage->BeginTrackingEntitiesTask(Linker, Params, TracksComponents->BoundMaterial, TracksComponents->ScalarParameterName);
-		ScalarParameterStorage->CachePreAnimatedValuesTask(Linker, Params, TracksComponents->BoundMaterial, TracksComponents->ScalarParameterName);
-
-		VectorParameterStorage->BeginTrackingEntitiesTask(Linker, Params, TracksComponents->BoundMaterial, TracksComponents->VectorParameterName);
-		VectorParameterStorage->BeginTrackingEntitiesTask(Linker, Params, TracksComponents->BoundMaterial, TracksComponents->ColorParameterName);
-		VectorParameterStorage->CachePreAnimatedValuesTask(Linker, Params, TracksComponents->BoundMaterial, TracksComponents->VectorParameterName);
-		VectorParameterStorage->CachePreAnimatedValuesTask(Linker, Params, TracksComponents->BoundMaterial, TracksComponents->ColorParameterName);
+		
+		ScalarParameterStorage->BeginTrackingAndCachePreAnimatedValuesTask(Linker, Params, TracksComponents->BoundMaterial, TracksComponents->ScalarParameterName);
+		VectorParameterStorage->BeginTrackingAndCachePreAnimatedValuesTask(Linker, Params, TracksComponents->BoundMaterial, TracksComponents->VectorParameterName);
+		VectorParameterStorage->BeginTrackingAndCachePreAnimatedValuesTask(Linker, Params, TracksComponents->BoundMaterial, TracksComponents->ColorParameterName);
 	}
 }
 
