Sequencer: Fixed MPC tracks not restoring state correctly

There are 2 issues here that are tackled together:
  - Firstly, the MPC track was using the PreAnimatedStateInterface to save its MPC values. Unfortunately this meant that its SavePreAnimatedState function could get called before the system was actually run (though this was not deterministic). This meant that it would try and save values before the bound materials had even been added, and so it would fail.
    The fix is to remove the use of this interface and just cache the state inside its OnRun function when we know it's initialized its materials correctly.

  - Second, this system suffered the same problem as Material Parameter tracks whereby it could end up caching unbound state for keep state entites. Although there's no evidence of this being a problem in practice (I suspect Keep State MPC tracks are seldom used), I have fixed it using the same method as 24444161)

#localization none
#tests Tested the original repro, EngineTests and Ad-hoc FrontEnd testing

diff --git a/Engine/Source/Runtime/MovieSceneTracks/Private/Systems/MovieSceneMaterialParameterCollectionSystem.cpp b/Engine/Source/Runtime/MovieSceneTracks/Private/Systems/MovieSceneMaterialParameterCollectionSystem.cpp
--- a/Engine/Source/Runtime/MovieSceneTracks/Private/Systems/MovieSceneMaterialParameterCollectionSystem.cpp
+++ b/Engine/Source/Runtime/MovieSceneTracks/Private/Systems/MovieSceneMaterialParameterCollectionSystem.cpp
@@ -167,26 +167,13 @@ void UMovieSceneMaterialParameterCollectionSystem::OnRun(FSystemTaskPrerequisite
 	FAddMPCMutation BindMaterialsMutation(Linker);
 	Linker->EntityManager.MutateAll(Filter, BindMaterialsMutation);
 	BindMaterialsMutation.Cleanup(Linker);
-}
-
-
-void UMovieSceneMaterialParameterCollectionSystem::SavePreAnimatedState(const FPreAnimationParameters& InParameters)
-{
-	using namespace UE::MovieScene;
-
-	FBuiltInComponentTypes*          BuiltInComponents = FBuiltInComponentTypes::Get();
-	FMovieSceneTracksComponentTypes* TracksComponents  = FMovieSceneTracksComponentTypes::Get();
 
 	TPreAnimatedStateTaskParams<FObjectComponent, FName> Params;
 
+	Params.AdditionalFilter.None({ BuiltInComponents->BlendChannelOutput });
 	Params.AdditionalFilter.All({ TracksComponents->MPC });
-	ScalarParameterStorage->BeginTrackingEntitiesTask(Linker, Params, TracksComponents->BoundMaterial, TracksComponents->ScalarParameterName);
-	ScalarParameterStorage->CachePreAnimatedValuesTask(Linker, Params, TracksComponents->BoundMaterial, TracksComponents->ScalarParameterName);
-
-	VectorParameterStorage->BeginTrackingEntitiesTask(Linker, Params, TracksComponents->BoundMaterial, TracksComponents->VectorParameterName);
-	VectorParameterStorage->BeginTrackingEntitiesTask(Linker, Params, TracksComponents->BoundMaterial, TracksComponents->ColorParameterName);
-	VectorParameterStorage->CachePreAnimatedValuesTask(Linker, Params, TracksComponents->BoundMaterial, TracksComponents->VectorParameterName);
-	VectorParameterStorage->CachePreAnimatedValuesTask(Linker, Params, TracksComponents->BoundMaterial, TracksComponents->ColorParameterName);
+	ScalarParameterStorage->BeginTrackingAndCachePreAnimatedValuesTask(Linker, Params, TracksComponents->BoundMaterial, TracksComponents->ScalarParameterName);
+	VectorParameterStorage->BeginTrackingAndCachePreAnimatedValuesTask(Linker, Params, TracksComponents->BoundMaterial, TracksComponents->VectorParameterName);
+	VectorParameterStorage->BeginTrackingAndCachePreAnimatedValuesTask(Linker, Params, TracksComponents->BoundMaterial, TracksComponents->ColorParameterName);
 }
 
-
diff --git a/Engine/Source/Runtime/MovieSceneTracks/Public/Systems/MovieSceneMaterialParameterCollectionSystem.h b/Engine/Source/Runtime/MovieSceneTracks/Public/Systems/MovieSceneMaterialParameterCollectionSystem.h
--- a/Engine/Source/Runtime/MovieSceneTracks/Public/Systems/MovieSceneMaterialParameterCollectionSystem.h
+++ b/Engine/Source/Runtime/MovieSceneTracks/Public/Systems/MovieSceneMaterialParameterCollectionSystem.h
@@ -25,7 +25,6 @@ struct FPreAnimatedVectorMaterialParameterStorage;
 UCLASS(MinimalAPI)
 class UMovieSceneMaterialParameterCollectionSystem
 	: public UMovieSceneEntitySystem
-	, public IMovieScenePreAnimatedStateSystemInterface
 {
 public:
 
@@ -38,8 +37,6 @@ private:
 	virtual void OnLink() override;
 	virtual void OnRun(FSystemTaskPrerequisites& InPrerequisites, FSystemSubsequentTasks& Subsequents) override;
 
-	virtual void SavePreAnimatedState(const FPreAnimationParameters& InParameters) override;
-
 private:
 
 	/** Holds pre-animated values for scalar MPC values */
