Sequencer: Constraints: Need to use constraint copy name to find constraint and don't re-add the constraint. Note these changes were already done for control rig constraints.

diff --git a/Engine/Source/Runtime/MovieSceneTracks/Private/Systems/MovieSceneConstraintSystem.cpp b/Engine/Source/Runtime/MovieSceneTracks/Private/Systems/MovieSceneConstraintSystem.cpp
--- a/Engine/Source/Runtime/MovieSceneTracks/Private/Systems/MovieSceneConstraintSystem.cpp
+++ b/Engine/Source/Runtime/MovieSceneTracks/Private/Systems/MovieSceneConstraintSystem.cpp
@@ -35,7 +35,7 @@ static UTickableConstraint* CreateConstraintIfNeeded(FConstraintsManagerControll
 	if (!Constraint) //if constraint doesn't exist it probably got unspawned so recreate it and add it
 	{
 		Constraint = Controller.AddConstraintFromCopy(ConstraintAndActiveChannel->ConstraintCopyToSpawn);
-		ConstraintChannel.Section->ReplaceConstraint(ConstraintChannel.ConstraintName, Constraint);
+		ConstraintChannel.Section->ReplaceConstraint(ConstraintAndActiveChannel->ConstraintCopyToSpawn->GetFName(), Constraint);
 	}
 	else // it's possible that we have it but it's not in the manager, due to manager not being saved with it (due to spawning or undo/redo).
 	{
@@ -46,8 +46,7 @@ static UTickableConstraint* CreateConstraintIfNeeded(FConstraintsManagerControll
 			if (Controller.DoesExistInAnyWorld(Constraint) == false) //it's no where make a copy and add it
 			{
 				Constraint = Controller.AddConstraintFromCopy(ConstraintAndActiveChannel->ConstraintCopyToSpawn);
-				Controller.AddConstraint(Constraint);
-				ConstraintChannel.Section->ReplaceConstraint(ConstraintChannel.ConstraintName, Constraint);
+				ConstraintChannel.Section->ReplaceConstraint(ConstraintAndActiveChannel->ConstraintCopyToSpawn->GetFName(), Constraint);
 			}
 
 		}
