Fixed crash when attempting to restore a widget material on a widget that had been garbage collected

This change ensures that the widget material handle is not initialized if the widget is no longer valid.
The secondary change is to add Accessor validity checks throughout the material system where previously they were only used in InitializeBoundMaterial.

diff --git a/Engine/Source/Runtime/MovieSceneTracks/Public/Systems/MovieSceneMaterialSystem.h b/Engine/Source/Runtime/MovieSceneTracks/Public/Systems/MovieSceneMaterialSystem.h
--- a/Engine/Source/Runtime/MovieSceneTracks/Public/Systems/MovieSceneMaterialSystem.h
+++ b/Engine/Source/Runtime/MovieSceneTracks/Public/Systems/MovieSceneMaterialSystem.h
@@ -46,12 +46,17 @@ struct TPreAnimatedMaterialTraits : FBoundObjectPreAnimatedStateTraits
 
 	static UMaterialInterface* CachePreAnimatedValue(typename TCallTraits<RequiredComponents>::ParamType... InRequiredComponents)
 	{
-		return AccessorType{ InRequiredComponents... }.GetMaterial();
+		AccessorType Accessor{ InRequiredComponents... };
+		return Accessor ? Accessor.GetMaterial() : nullptr;
 	}
 
 	static void RestorePreAnimatedValue(const KeyType& InKeyType, UMaterialInterface* OldMaterial, const FRestoreStateParams& Params)
 	{
-		AccessorType{ InKeyType }.SetMaterial(OldMaterial);
+		AccessorType Accessor{ InKeyType };
+		if (Accessor)
+		{
+			Accessor.SetMaterial(OldMaterial);
+		}
 	}
 };
 
@@ -68,14 +73,18 @@ struct TPreAnimatedMaterialParameterTraits : FBoundObjectPreAnimatedStateTraits
 		AccessorType Accessor{ InRequiredComponents... };
 
 		FMovieScenePreAnimatedMaterialParameters Parameters;
-		Parameters.PreviousMaterial = Accessor.GetMaterial();
 
-		// If the current material we're overriding is already a material instance dynamic, copy it since we will be modifying the data. 
-		// The copied material will be used to restore the values in RestoreState.
-		UMaterialInstanceDynamic* MID = Cast<UMaterialInstanceDynamic>(Parameters.PreviousMaterial);
-		if (MID)
+		if (Accessor)
 		{
-			Parameters.PreviousParameterContainer = DuplicateObject<UMaterialInterface>(MID, MID->GetOuter());
+			Parameters.PreviousMaterial = Accessor.GetMaterial();
+
+			// If the current material we're overriding is already a material instance dynamic, copy it since we will be modifying the data. 
+			// The copied material will be used to restore the values in RestoreState.
+			UMaterialInstanceDynamic* MID = Cast<UMaterialInstanceDynamic>(Parameters.PreviousMaterial);
+			if (MID)
+			{
+				Parameters.PreviousParameterContainer = DuplicateObject<UMaterialInterface>(MID, MID->GetOuter());
+			}
 		}
 
 		return Parameters;
@@ -84,6 +93,11 @@ struct TPreAnimatedMaterialParameterTraits : FBoundObjectPreAnimatedStateTraits
 	static void RestorePreAnimatedValue(const KeyType& InKey, const FMovieScenePreAnimatedMaterialParameters& PreAnimatedValue, const FRestoreStateParams& Params)
 	{
 		AccessorType Accessor{ InKey };
+		if (!Accessor)
+		{
+			return;
+		}
+
 		if (PreAnimatedValue.PreviousParameterContainer != nullptr)
 		{
 			// If we cached parameter values in CachePreAnimatedValue that means the previous material was already a MID
@@ -137,6 +151,10 @@ struct TApplyMaterialSwitchers
 		UMaterialInterface* NewMaterial = Cast<UMaterialInterface>(ObjectResult);
 
 		AccessorType Accessor(Inputs...);
+		if (!Accessor)
+		{
+			return;
+		}
 
 		UMaterialInterface*        ExistingMaterial = Accessor.GetMaterial();
 		UMaterialInstanceDynamic*  ExistingMID      = Cast<UMaterialInstanceDynamic>(ExistingMaterial);
diff --git a/Engine/Source/Runtime/UMG/Private/Animation/MovieSceneWidgetMaterialSystem.cpp b/Engine/Source/Runtime/UMG/Private/Animation/MovieSceneWidgetMaterialSystem.cpp
--- a/Engine/Source/Runtime/UMG/Private/Animation/MovieSceneWidgetMaterialSystem.cpp
+++ b/Engine/Source/Runtime/UMG/Private/Animation/MovieSceneWidgetMaterialSystem.cpp
@@ -25,8 +25,14 @@ namespace UE::MovieScene
 
 FWidgetMaterialAccessor::FWidgetMaterialAccessor(const FWidgetMaterialKey& InKey)
 	: Widget(CastChecked<UWidget>(InKey.Object.ResolveObjectPtr()))
-	, WidgetMaterialHandle(InKey.WidgetMaterialHandle)
-{}
+{
+	// Only assign the material handle if the widget itself is still valid
+	// otherwise accessors constructed with the material handle will consider it valid
+	if (Widget)
+	{
+		WidgetMaterialHandle = InKey.WidgetMaterialHandle;
+	}
+}
 
 FWidgetMaterialAccessor::FWidgetMaterialAccessor(UObject* InObject, FWidgetMaterialHandle InWidgetMaterialHandle)
 	: Widget(Cast<UWidget>(InObject))
