Sequencer: Pre-animated materials are now stored as soft-object ptrs to prevent them from keeping the objects they animate alive

When animating a material parameter on an object that was already assigned a MID (that was probably outered to the object itself), the pre-animated state would previously maintain a strong reference to that MID, which in turn would keep alive the object, including its entire outer chain. We now store these references as soft object ptrs in order to allow us to correctly restore the material if it was previously assigned to a material asset.

This is controllable using the cvar "Sequencer.UseSoftObjectPtrsForPreAnimatedMaterial".

diff --git a/Engine/Source/Runtime/MovieSceneTracks/Private/Systems/MovieSceneMaterialSystem.cpp b/Engine/Source/Runtime/MovieSceneTracks/Private/Systems/MovieSceneMaterialSystem.cpp
new file mode 100644
--- /dev/null
+++ b/Engine/Source/Runtime/MovieSceneTracks/Private/Systems/MovieSceneMaterialSystem.cpp
@@ -0,0 +1,46 @@
+// Copyright Epic Games, Inc. All Rights Reserved.
+
+#include "Systems/MovieSceneMaterialSystem.h"
+
+namespace UE::MovieScene
+{
+
+bool GUseSoftObjectPtrsForPreAnimatedMaterial = true;
+FAutoConsoleVariableRef CVarUseSoftObjectPtrsForPreAnimatedMaterial(
+	TEXT("Sequencer.UseSoftObjectPtrsForPreAnimatedMaterial"),
+	GUseSoftObjectPtrsForPreAnimatedMaterial,
+	TEXT("Defines whether to use soft-object-ptrs for referencing pre-animated state (default), or strong TObjectPtrs.\n"),
+	ECVF_Default
+);
+
+} // namespace UE::MovieScene
+
+
+UMaterialInterface* FMovieScenePreAnimatedMaterialParameters::GetMaterial() const
+{
+	if (UE::MovieScene::GUseSoftObjectPtrsForPreAnimatedMaterial)
+	{
+		UMaterialInterface* Material = SoftPreviousMaterial.Get();
+		if (!Material)
+		{
+			Material = SoftPreviousMaterial.LoadSynchronous();
+		}
+		return Material;
+	}
+	else
+	{
+		return PreviousMaterial;
+	}
+}
+
+void FMovieScenePreAnimatedMaterialParameters::SetMaterial(UMaterialInterface* InMaterial)
+{
+	if (UE::MovieScene::GUseSoftObjectPtrsForPreAnimatedMaterial)
+	{
+		SoftPreviousMaterial = InMaterial;
+	}
+	else
+	{
+		PreviousMaterial = InMaterial;
+	}
+}
diff --git a/Engine/Source/Runtime/MovieSceneTracks/Public/Systems/MovieSceneMaterialSystem.h b/Engine/Source/Runtime/MovieSceneTracks/Public/Systems/MovieSceneMaterialSystem.h
--- a/Engine/Source/Runtime/MovieSceneTracks/Public/Systems/MovieSceneMaterialSystem.h
+++ b/Engine/Source/Runtime/MovieSceneTracks/Public/Systems/MovieSceneMaterialSystem.h
@@ -15,6 +15,7 @@
 #include "EntitySystem/MovieSceneCachedEntityFilterResult.h"
 #include "Evaluation/PreAnimatedState/MovieScenePreAnimatedObjectStorage.h"
 #include "MovieSceneTracksComponentTypes.h"
+#include "UObject/SoftObjectPtr.h"
 
 #include "MovieSceneMaterialSystem.generated.h"
 
@@ -26,11 +27,19 @@ struct FMovieScenePreAnimatedMaterialParameters
 {
 	GENERATED_BODY()
 
+	MOVIESCENETRACKS_API UMaterialInterface* GetMaterial() const;
+
+	MOVIESCENETRACKS_API void SetMaterial(UMaterialInterface* InMaterial);
+
+private:
+
+	/** Strong ptr to the previously assigned material interface (used when Sequencer.UseSoftObjectPtrsForPreAnimatedMaterial is false) */
 	UPROPERTY()
-	TObjectPtr<UMaterialInterface> PreviousMaterial = nullptr;
+	TObjectPtr<UMaterialInterface> PreviousMaterial;
 
+	/** Soft ptr to the previously assigned material interface (used when Sequencer.UseSoftObjectPtrsForPreAnimatedMaterial is true) */
 	UPROPERTY()
-	TObjectPtr<UMaterialInterface> PreviousParameterContainer = nullptr;
+	TSoftObjectPtr<UMaterialInterface> SoftPreviousMaterial;
 };
 
 namespace UE::MovieScene
@@ -38,30 +47,6 @@ namespace UE::MovieScene
 
 template<typename AccessorType, typename... RequiredComponents>
 struct TPreAnimatedMaterialTraits : FBoundObjectPreAnimatedStateTraits
-{
-	using KeyType     = typename AccessorType::KeyType;
-	using StorageType = UMaterialInterface*;
-
-	static_assert(THasAddReferencedObjectForComponent<StorageType>::Value, "StorageType is not correctly exposed to the reference graph!");
-
-	static UMaterialInterface* CachePreAnimatedValue(typename TCallTraits<RequiredComponents>::ParamType... InRequiredComponents)
-	{
-		AccessorType Accessor{ InRequiredComponents... };
-		return Accessor ? Accessor.GetMaterial() : nullptr;
-	}
-
-	static void RestorePreAnimatedValue(const KeyType& InKeyType, UMaterialInterface* OldMaterial, const FRestoreStateParams& Params)
-	{
-		AccessorType Accessor{ InKeyType };
-		if (Accessor)
-		{
-			Accessor.SetMaterial(OldMaterial);
-		}
-	}
-};
-
-template<typename AccessorType, typename... RequiredComponents>
-struct TPreAnimatedMaterialParameterTraits : FBoundObjectPreAnimatedStateTraits
 {
 	using KeyType     = typename AccessorType::KeyType;
 	using StorageType = FMovieScenePreAnimatedMaterialParameters;
@@ -76,15 +61,7 @@ struct TPreAnimatedMaterialParameterTraits : FBoundObjectPreAnimatedStateTraits
 
 		if (Accessor)
 		{
-			Parameters.PreviousMaterial = Accessor.GetMaterial();
-
-			// If the current material we're overriding is already a material instance dynamic, copy it since we will be modifying the data. 
-			// The copied material will be used to restore the values in RestoreState.
-			UMaterialInstanceDynamic* MID = Cast<UMaterialInstanceDynamic>(Parameters.PreviousMaterial);
-			if (MID)
-			{
-				Parameters.PreviousParameterContainer = DuplicateObject<UMaterialInterface>(MID, MID->GetOuter());
-			}
+			Parameters.SetMaterial(Accessor.GetMaterial());
 		}
 
 		return Parameters;
@@ -93,28 +70,19 @@ struct TPreAnimatedMaterialParameterTraits : FBoundObjectPreAnimatedStateTraits
 	static void RestorePreAnimatedValue(const KeyType& InKey, const FMovieScenePreAnimatedMaterialParameters& PreAnimatedValue, const FRestoreStateParams& Params)
 	{
 		AccessorType Accessor{ InKey };
-		if (!Accessor)
-		{
-			return;
-		}
-
-		if (PreAnimatedValue.PreviousParameterContainer != nullptr)
+		if (Accessor)
 		{
-			// If we cached parameter values in CachePreAnimatedValue that means the previous material was already a MID
-			// and we probably did not replace it with a new one when resolving bound materials. Therefore we
-			// just copy the parameters back over without changing the material
-			UMaterialInstanceDynamic* CurrentMID = Cast<UMaterialInstanceDynamic>(Accessor.GetMaterial());
-			if (CurrentMID)
+			if (UMaterialInterface* PreviousMaterial = PreAnimatedValue.GetMaterial())
 			{
-				CurrentMID->CopyMaterialUniformParameters(PreAnimatedValue.PreviousParameterContainer);
-				return;
+				Accessor.SetMaterial(PreAnimatedValue.GetMaterial());
 			}
 		}
-
-		Accessor.SetMaterial(PreAnimatedValue.PreviousMaterial);
 	}
 };
 
+template<typename AccessorType, typename... RequiredComponents>
+using TPreAnimatedMaterialParameterTraits = TPreAnimatedMaterialTraits<AccessorType, RequiredComponents...>;
+
 template<typename AccessorType, typename... RequiredComponents>
 class TMovieSceneMaterialSystem
 {
