SkeletalMeshComponent : Added a parameter bForceInitAnimScriptInstance = true to the SetAnimationMode, in order to fix Anim notifies trigger repeatedly/every frame when using an animation mode track set to Use Anim Blueprint in a Level Sequence.
Defaulted to true to keep current behavior for Blueprints, but the AnimSequencer,due to the use of property bindings, calls it with the parameter set to false.

diff --git a/Engine/Source/Runtime/Engine/Classes/Components/SkeletalMeshComponent.h b/Engine/Source/Runtime/Engine/Classes/Components/SkeletalMeshComponent.h
--- a/Engine/Source/Runtime/Engine/Classes/Components/SkeletalMeshComponent.h
+++ b/Engine/Source/Runtime/Engine/Classes/Components/SkeletalMeshComponent.h
@@ -562,7 +562,10 @@ protected:
 	/** Whether to use Animation Blueprint or play Single Animation Asset. */
 	UPROPERTY(EditAnywhere, BlueprintReadOnly, Category=Animation)
 	TEnumAsByte<EAnimationMode::Type>	AnimationMode;
-	
+public:
+	// helper function to get the member name and verify it exists, without making it public
+	static const FName& GetAnimationModePropertyNameChecked();
+
 private:
 	/** Teleport type to use on the next update */
 	ETeleportType PendingTeleportType;
@@ -1031,9 +1034,15 @@ public:
 	UFUNCTION(BlueprintCallable, Category = "Components|SkeletalMesh", meta = (Keywords = "Dynamics,Physics", UnsafeDuringActorConstruction = "true"))
 	void ResetAnimInstanceDynamics(ETeleportType InTeleportType = ETeleportType::ResetPhysics);
 
-	/** Below are the interface to control animation when animation mode, not blueprint mode **/
+	/** Below are the interface to control animation when animation mode, not blueprint mode */
+
+	/**
+	* Set the Animation Mode
+	* @param InAnimationMode : Requested AnimationMode
+	* @param bForceInitAnimScriptInstance : Init AnimScriptInstance if the AnimationMode is AnimationBlueprint even if the new animation mode is the same as current (this allows to use BP construction script to do this)
+	*/
 	UFUNCTION(BlueprintCallable, Category = "Components|Animation", meta = (Keywords = "Animation"))
-	void SetAnimationMode(EAnimationMode::Type InAnimationMode);
+	void SetAnimationMode(EAnimationMode::Type InAnimationMode, bool bForceInitAnimScriptInstance = true);
 	
 	UFUNCTION(BlueprintPure, Category = "Components|Animation", meta = (Keywords = "Animation"))
 	EAnimationMode::Type GetAnimationMode() const;
diff --git a/Engine/Source/Runtime/Engine/Private/Components/SkeletalMeshComponent.cpp b/Engine/Source/Runtime/Engine/Private/Components/SkeletalMeshComponent.cpp
--- a/Engine/Source/Runtime/Engine/Private/Components/SkeletalMeshComponent.cpp
+++ b/Engine/Source/Runtime/Engine/Private/Components/SkeletalMeshComponent.cpp
@@ -3484,9 +3484,9 @@ void USkeletalMeshComponent::GetResourceSizeEx(FResourceSizeEx& CumulativeResour
 	}
 }
 
-void USkeletalMeshComponent::SetAnimationMode(EAnimationMode::Type InAnimationMode)
+void USkeletalMeshComponent::SetAnimationMode(EAnimationMode::Type InAnimationMode, bool bForceInitAnimScriptInstance)
 {
-	bool bNeedChange = AnimationMode != InAnimationMode;
+	const bool bNeedChange = AnimationMode != InAnimationMode;
 	if (bNeedChange)
 	{
 		AnimationMode = InAnimationMode;
@@ -3496,7 +3496,7 @@ void USkeletalMeshComponent::SetAnimationMode(EAnimationMode::Type InAnimationMo
 	// when mode is swapped, make sure to reinitialize
 	// even if it was same mode, this was due to users who wants to use BP construction script to do this
 	// if you use it in the construction script, it gets serialized, but it never instantiate. 
-	if(GetSkeletalMeshAsset() != nullptr && (bNeedChange || AnimationMode == EAnimationMode::AnimationBlueprint))
+	if(GetSkeletalMeshAsset() != nullptr && (bNeedChange || (AnimationMode == EAnimationMode::AnimationBlueprint && bForceInitAnimScriptInstance)))
 	{
 		if (InitializeAnimScriptInstance(true))
 		{
@@ -4716,6 +4716,13 @@ bool USkeletalMeshComponent::FindAttributeChecked(const FName& BoneName, const F
 }
 
 
+const FName& USkeletalMeshComponent::GetAnimationModePropertyNameChecked()
+{
+	static FName Name = GET_MEMBER_NAME_CHECKED(USkeletalMeshComponent, AnimationMode);
+	return Name;
+}
+
+
 TArray<FTransform> USkeletalMeshComponent::GetBoneSpaceTransforms() 
 {
 	// We may be doing parallel evaluation on the current anim instance
diff --git a/Engine/Source/Runtime/MovieSceneTracks/Private/MovieSceneTracksComponentTypes.cpp b/Engine/Source/Runtime/MovieSceneTracks/Private/MovieSceneTracksComponentTypes.cpp
--- a/Engine/Source/Runtime/MovieSceneTracks/Private/MovieSceneTracksComponentTypes.cpp
+++ b/Engine/Source/Runtime/MovieSceneTracks/Private/MovieSceneTracksComponentTypes.cpp
@@ -4,6 +4,7 @@
 #include "Components/ExponentialHeightFogComponent.h"
 #include "Components/LightComponent.h"
 #include "Components/PrimitiveComponent.h"
+#include "Components/SkeletalMeshComponent.h"
 #include "Components/SkyAtmosphereComponent.h"
 #include "Components/SkyLightComponent.h"
 #include "EntitySystem/BuiltInComponentTypes.h"
@@ -522,6 +523,19 @@ void SetSecondFogDataFogHeightOffset(UObject* Object, float InFogHeightOffset)
 	ExponentialHeightFogComponent->SetSecondFogHeightOffset(InFogHeightOffset);
 }
 
+uint8 GetSkeletalMeshAnimationMode(const UObject* Object)
+{
+	const USkeletalMeshComponent* SkeletalMeshComponent = CastChecked<const USkeletalMeshComponent>(Object);
+	return SkeletalMeshComponent->GetAnimationMode();
+}
+
+void SetSkeletalMeshAnimationMode(UObject* Object, uint8 InAnimationMode)
+{
+	USkeletalMeshComponent* SkeletalMeshComponent = CastChecked<USkeletalMeshComponent>(Object);
+	constexpr bool bForceInitAnimScriptInstance = false; // Avoid reinits each frame if an anim node track is added with AnimBlueprint mode
+	SkeletalMeshComponent->SetAnimationMode((EAnimationMode::Type)InAnimationMode, bForceInitAnimScriptInstance);
+}
+
 void FIntermediate3DTransform::ApplyTo(USceneComponent* SceneComponent) const
 {
 	double DeltaTime = FApp::GetDeltaTime();
@@ -1010,6 +1024,13 @@ FMovieSceneTracksComponentTypes::FMovieSceneTracksComponentTypes()
 			UExponentialHeightFogComponent::StaticClass(), *SecondFogDataFogHeightOffsetPath,
 			GetSecondFogDataFogHeightOffset, SetSecondFogDataFogHeightOffset);
 
+	Accessors.Byte.Add(
+		USkeletalMeshComponent::StaticClass(), USkeletalMeshComponent::GetAnimationModePropertyNameChecked(),
+		GetSkeletalMeshAnimationMode, SetSkeletalMeshAnimationMode);
+	Accessors.Enum.Add(
+		USkeletalMeshComponent::StaticClass(), USkeletalMeshComponent::GetAnimationModePropertyNameChecked(),
+		GetSkeletalMeshAnimationMode, SetSkeletalMeshAnimationMode);
+
 	// --------------------------------------------------------------------------------------------
 	// Set up vector properties
 	BuiltInComponents->PropertyRegistry.DefineCompositeProperty(FloatVector, TEXT("Apply FloatVector Properties"))
