Sequencer: Do not cache pre-animated transforms for attachments that use KeepRelative for both attach and detach rules

This ensures that user-modifications to an object's transform do not get restored when an object was previously animated with an attach track, but is no longer actively animated (and its transform never changed)

diff --git a/Engine/Source/Runtime/MovieSceneTracks/Private/PreAnimatedState/MovieScenePreAnimatedComponentTransformStorage.cpp b/Engine/Source/Runtime/MovieSceneTracks/Private/PreAnimatedState/MovieScenePreAnimatedComponentTransformStorage.cpp
--- a/Engine/Source/Runtime/MovieSceneTracks/Private/PreAnimatedState/MovieScenePreAnimatedComponentTransformStorage.cpp
+++ b/Engine/Source/Runtime/MovieSceneTracks/Private/PreAnimatedState/MovieScenePreAnimatedComponentTransformStorage.cpp
@@ -70,13 +70,14 @@ FPreAnimatedComponentTransformStorage::FPreAnimatedComponentTransformStorage()
 	: TPreAnimatedPropertyStorage<FComponentTransformPreAnimatedTraits>(FBuiltInComponentTypes::Get()->PropertyRegistry.GetDefinition(FMovieSceneTracksComponentTypes::Get()->ComponentTransform.CompositeID))
 {}
 
-void FPreAnimatedComponentTransformStorage::CachePreAnimatedTransforms(const FCachePreAnimatedValueParams& Params, TArrayView<UObject* const> BoundObjects)
+void FPreAnimatedComponentTransformStorage::CachePreAnimatedTransforms(const FCachePreAnimatedValueParams& Params, TArrayView<UObject* const> BoundObjects, TOptional<TFunctionRef<bool(int32)>> Predicate)
 {
 	static FMovieScenePropertyBinding PropertyBinding("Transform", TEXT("Transform"));
 
-	for (UObject* BoundObject : BoundObjects)
+	for (int32 Index = 0; Index < BoundObjects.Num(); ++Index)
 	{
-		if (!BoundObject)
+		UObject* BoundObject = BoundObjects[Index];
+		if (!BoundObject || (Predicate && !(*Predicate)(Index)))
 		{
 			continue;
 		}
diff --git a/Engine/Source/Runtime/MovieSceneTracks/Private/PreAnimatedState/MovieScenePreAnimatedComponentTransformStorage.h b/Engine/Source/Runtime/MovieSceneTracks/Private/PreAnimatedState/MovieScenePreAnimatedComponentTransformStorage.h
--- a/Engine/Source/Runtime/MovieSceneTracks/Private/PreAnimatedState/MovieScenePreAnimatedComponentTransformStorage.h
+++ b/Engine/Source/Runtime/MovieSceneTracks/Private/PreAnimatedState/MovieScenePreAnimatedComponentTransformStorage.h
@@ -30,7 +30,7 @@ struct MOVIESCENETRACKS_API FPreAnimatedComponentTransformStorage
 
 	static TAutoRegisterPreAnimatedStorageID<FPreAnimatedComponentTransformStorage> StorageID;
 
-	void CachePreAnimatedTransforms(const FCachePreAnimatedValueParams& Params, TArrayView<UObject* const> BoundObjects);
+	void CachePreAnimatedTransforms(const FCachePreAnimatedValueParams& Params, TArrayView<UObject* const> BoundObjects, TOptional<TFunctionRef<bool(int32)>> Predicate = TOptional<TFunctionRef<bool(int32)>>());
 };
 
 
diff --git a/Engine/Source/Runtime/MovieSceneTracks/Private/Systems/MovieSceneComponentAttachmentSystem.cpp b/Engine/Source/Runtime/MovieSceneTracks/Private/Systems/MovieSceneComponentAttachmentSystem.cpp
--- a/Engine/Source/Runtime/MovieSceneTracks/Private/Systems/MovieSceneComponentAttachmentSystem.cpp
+++ b/Engine/Source/Runtime/MovieSceneTracks/Private/Systems/MovieSceneComponentAttachmentSystem.cpp
@@ -267,12 +267,25 @@ void UMovieSceneComponentAttachmentSystem::SavePreAnimatedState(const FPreAnimat
 	.ReadEntityIDs()
 	.Read(BuiltInComponents->RootInstanceHandle)
 	.Read(BuiltInComponents->BoundObject)
+	.Read(TrackComponents->AttachComponent)
 	.FilterAll(FilterMask)
 	.Iterate_PerAllocation(&Linker->EntityManager,
-		[EntityMetaData, ComponentTransformStorage, ComponentAttachmentStorage](FEntityAllocationIteratorItem Item, TRead<FMovieSceneEntityID> EntityIDs, TRead<FRootInstanceHandle> InstanceHandles, TRead<UObject*> BoundObjects)
+		[EntityMetaData, ComponentTransformStorage, ComponentAttachmentStorage](FEntityAllocationIteratorItem Item, TRead<FMovieSceneEntityID> EntityIDs, TRead<FRootInstanceHandle> InstanceHandles, TRead<UObject*> BoundObjects, TRead<FAttachmentComponent> AttachComponents)
 		{
 			TArrayView<UObject* const> BoundObjectArray = BoundObjects.AsArray(Item.GetAllocation()->Num());
 
+			auto ShouldCacheTransform = [AttachComponents](int32 Index)
+			{
+				FAttachmentComponent Component = AttachComponents[Index];
+				// Only cache persistent preanimated transforms if we're not using a keep relative attachment
+				return Component.AttachParams.AttachmentLocationRule != EAttachmentRule::KeepRelative
+					|| Component.AttachParams.AttachmentRotationRule != EAttachmentRule::KeepRelative
+					|| Component.AttachParams.AttachmentScaleRule    != EAttachmentRule::KeepRelative
+					|| Component.DetachParams.DetachmentLocationRule != EDetachmentRule::KeepRelative
+					|| Component.DetachParams.DetachmentRotationRule != EDetachmentRule::KeepRelative
+					|| Component.DetachParams.DetachmentScaleRule    != EDetachmentRule::KeepRelative;
+			};
+
 			// Order is important here - always cache the transforms first so that they are restored last
 
 			// If we're capturing global state for an attached parent, we forcibly persist the transform so that it
@@ -280,7 +293,7 @@ void UMovieSceneComponentAttachmentSystem::SavePreAnimatedState(const FPreAnimat
 			// state regardless of the detach rules used during normal playback.
 			FCachePreAnimatedValueParams ForcePersistParams;
 			ForcePersistParams.bForcePersist = true;
-			ComponentTransformStorage->CachePreAnimatedTransforms(ForcePersistParams, BoundObjectArray);
+			ComponentTransformStorage->CachePreAnimatedTransforms(ForcePersistParams, BoundObjectArray, TFunctionRef<bool(int32)>(ShouldCacheTransform));
 
 			FPreAnimatedTrackerParams AttachmentParams(Item);
 			AttachmentParams.bWantsRestoreState = false;
