Sequencer: Constraints:  Make sure constraints aren't masked out since they don't contribute to current masking behavior.
#trivial

diff --git a/Engine/Source/Runtime/MovieSceneTracks/Private/Sections/MovieScene3DTransformSection.cpp b/Engine/Source/Runtime/MovieSceneTracks/Private/Sections/MovieScene3DTransformSection.cpp
--- a/Engine/Source/Runtime/MovieSceneTracks/Private/Sections/MovieScene3DTransformSection.cpp
+++ b/Engine/Source/Runtime/MovieSceneTracks/Private/Sections/MovieScene3DTransformSection.cpp
@@ -740,7 +740,12 @@ FMovieSceneTransformMask UMovieScene3DTransformSection::GetMaskByName(const FNam
 	{
 		return EMovieSceneTransformChannel::ScaleZ;
 	}
-
+	//Constraints aren't masked since they need to be deleted etc via their own API's
+	FString ConstraintString = NSLOCTEXT("MovieSceneTransformSection", "Constraint", "Constraint").ToString();
+	if (InName.ToString().Contains(ConstraintString))
+	{
+		return EMovieSceneTransformChannel::None;
+	}
 	return EMovieSceneTransformChannel::All;
 }
 
