Sequencer: Swap Root Bone on Track

diff --git a/Engine/Source/Editor/MovieSceneTools/Private/TrackEditors/SkeletalAnimationTrackEditor.cpp b/Engine/Source/Editor/MovieSceneTools/Private/TrackEditors/SkeletalAnimationTrackEditor.cpp
--- a/Engine/Source/Editor/MovieSceneTools/Private/TrackEditors/SkeletalAnimationTrackEditor.cpp
+++ b/Engine/Source/Editor/MovieSceneTools/Private/TrackEditors/SkeletalAnimationTrackEditor.cpp
@@ -80,6 +80,7 @@
 #include "LevelSequenceAnimSequenceLink.h"
 #include "AnimSequenceLevelSequenceLink.h"
 #include "UObject/SavePackage.h"
+#include "AnimSequencerInstanceProxy.h"
 
 int32 FSkeletalAnimationTrackEditor::NumberActive = 0;
 
@@ -1874,34 +1875,75 @@ void FSkeletalAnimationTrackEditor::BuildTrackContextMenu( FMenuBuilder& MenuBui
 			NSLOCTEXT("Sequencer", "BlendFirstChildOfRootTooltip", "If True, do not blend and match the root bones but instead the first child bone of the root. Toggle this on when the matched sequences in the track have no motion on the root."),
 			FSlateIcon(),
 			FUIAction(
-				FExecuteAction::CreateLambda([=]()->void {
+				FExecuteAction::CreateLambda([SequencerPtr,SkeletalAnimationTrack]()->void {
 					SkeletalAnimationTrack->bBlendFirstChildOfRoot = SkeletalAnimationTrack->bBlendFirstChildOfRoot ? false : true;
-					SkeletalAnimationTrack->SetRootMotionsDirty();
-					SequencerPtr->NotifyMovieSceneDataChanged(EMovieSceneDataChangeType::TrackValueChanged);
+		SkeletalAnimationTrack->SetRootMotionsDirty();
+		SequencerPtr->NotifyMovieSceneDataChanged(EMovieSceneDataChangeType::TrackValueChanged);
 
 					}),
-				FCanExecuteAction::CreateLambda([=]()->bool { return SequencerPtr && SkeletalAnimationTrack != nullptr; }),
-						FIsActionChecked::CreateLambda([=]()->bool { return SkeletalAnimationTrack != nullptr && SkeletalAnimationTrack->bBlendFirstChildOfRoot; })),
+				FCanExecuteAction::CreateLambda([SequencerPtr, SkeletalAnimationTrack]()->bool { return SequencerPtr && SkeletalAnimationTrack != nullptr; }),
+						FIsActionChecked::CreateLambda([SequencerPtr, SkeletalAnimationTrack]()->bool { return SkeletalAnimationTrack != nullptr && SkeletalAnimationTrack->bBlendFirstChildOfRoot; })),
 			NAME_None,
-			EUserInterfaceActionType::ToggleButton
-			);
+						EUserInterfaceActionType::ToggleButton
+						);
 
 		MenuBuilder.AddMenuEntry(
 			NSLOCTEXT("Sequencer", "ShowRootMotionTrails", "Show Root Motion Trail"),
 			NSLOCTEXT("Sequencer", "ShowRootMotionTrailsTooltip", "Show the Root Motion Trail for all Animation Clips."),
 			FSlateIcon(),
 			FUIAction(
-				FExecuteAction::CreateLambda([=]()->void {
+				FExecuteAction::CreateLambda([SequencerPtr, SkeletalAnimationTrack]()->void {
 					SkeletalAnimationTrack->ToggleShowRootMotionTrail();
-					SequencerPtr->NotifyMovieSceneDataChanged(EMovieSceneDataChangeType::TrackValueChanged);
+		SequencerPtr->NotifyMovieSceneDataChanged(EMovieSceneDataChangeType::TrackValueChanged);
 
 					}),
-				FCanExecuteAction::CreateLambda([=]()->bool { return SequencerPtr && SkeletalAnimationTrack != nullptr; }),
-						FIsActionChecked::CreateLambda([=]()->bool { return SkeletalAnimationTrack != nullptr && SkeletalAnimationTrack->bShowRootMotionTrail; })),
+				FCanExecuteAction::CreateLambda([SequencerPtr, SkeletalAnimationTrack]()->bool { return SequencerPtr && SkeletalAnimationTrack != nullptr; }),
+						FIsActionChecked::CreateLambda([SequencerPtr, SkeletalAnimationTrack]()->bool { return SkeletalAnimationTrack != nullptr && SkeletalAnimationTrack->bShowRootMotionTrail; })),
 			NAME_None,
 						EUserInterfaceActionType::ToggleButton
-			);
+						);
+
+
+		MenuBuilder.AddMenuEntry(
+			NSLOCTEXT("Sequencer", "SwapRootBoneNone", "Swap Root Bone None"),
+			NSLOCTEXT("Sequencer", "SwapRootBoneNoneTooltip", "Do not swap root bone for all sections."),
+			FSlateIcon(),
+			FUIAction(
+				FExecuteAction::CreateLambda([SkeletalAnimationTrack]()->void {
+					SkeletalAnimationTrack->SetSwapRootBone(ESwapRootBone::SwapRootBone_None);
+					}),
+				FCanExecuteAction::CreateLambda([SkeletalAnimationTrack]()->bool { return  SkeletalAnimationTrack != nullptr; }),
+				FIsActionChecked::CreateLambda([SkeletalAnimationTrack]()->bool { return SkeletalAnimationTrack != nullptr && SkeletalAnimationTrack->SwapRootBone == ESwapRootBone::SwapRootBone_None; })),
+				NAME_None,
+				EUserInterfaceActionType::RadioButton
+		);
 
+		MenuBuilder.AddMenuEntry(
+			NSLOCTEXT("Sequencer", "SwapRootBoneActor", "Swap Root Bone Actor"),
+			NSLOCTEXT("Sequencer", "SwapRootBoneActorTooltip", "Swap root bone on root actor component for all sections."),
+			FSlateIcon(),
+			FUIAction(
+				FExecuteAction::CreateLambda([SkeletalAnimationTrack]()->void {
+					SkeletalAnimationTrack->SetSwapRootBone(ESwapRootBone::SwapRootBone_Actor);
+					}),
+				FCanExecuteAction::CreateLambda([SkeletalAnimationTrack]()->bool { return  SkeletalAnimationTrack != nullptr; }),
+				FIsActionChecked::CreateLambda([SkeletalAnimationTrack]()->bool { return SkeletalAnimationTrack != nullptr && SkeletalAnimationTrack->SwapRootBone == ESwapRootBone::SwapRootBone_Actor; })),
+				NAME_None,
+				EUserInterfaceActionType::RadioButton
+			);
+		MenuBuilder.AddMenuEntry(
+			NSLOCTEXT("Sequencer", "SwapRootBoneComponent", "Swap Root Bone Component"),
+			NSLOCTEXT("Sequencer", "SwapRootBoneComponentTooltip", "Swap root bone on current component for all sections."),
+			FSlateIcon(),
+			FUIAction(
+				FExecuteAction::CreateLambda([SkeletalAnimationTrack]()->void {
+					SkeletalAnimationTrack->SetSwapRootBone(ESwapRootBone::SwapRootBone_Component);
+					}),
+				FCanExecuteAction::CreateLambda([SkeletalAnimationTrack]()->bool { return  SkeletalAnimationTrack != nullptr; }),
+				FIsActionChecked::CreateLambda([SkeletalAnimationTrack]()->bool { return SkeletalAnimationTrack != nullptr && SkeletalAnimationTrack->SwapRootBone == ESwapRootBone::SwapRootBone_Component; })),
+				NAME_None,
+				EUserInterfaceActionType::RadioButton
+			);
 	}
 	MenuBuilder.EndSection();
 	MenuBuilder.AddSeparator();
diff --git a/Engine/Source/Runtime/MovieSceneTracks/Private/Systems/MovieSceneSkeletalAnimationSystem.cpp b/Engine/Source/Runtime/MovieSceneTracks/Private/Systems/MovieSceneSkeletalAnimationSystem.cpp
--- a/Engine/Source/Runtime/MovieSceneTracks/Private/Systems/MovieSceneSkeletalAnimationSystem.cpp
+++ b/Engine/Source/Runtime/MovieSceneTracks/Private/Systems/MovieSceneSkeletalAnimationSystem.cpp
@@ -36,6 +36,9 @@
 #include "Systems/WeightAndEasingEvaluatorSystem.h"
 
 #include "Components/SkeletalMeshComponent.h"
+#include "Animation/AnimationPoseData.h"
+#include "BonePose.h"
+#include "Animation/BuiltInAttributeTypes.h"
 
 #include UE_INLINE_GENERATED_CPP_BY_NAME(MovieSceneSkeletalAnimationSystem)
 
@@ -236,6 +239,75 @@ TAutoRegisterPreAnimatedStorageID<FPreAnimatedSkeletalAnimationAnimInstanceStora
 
 /** ------------------------------------------------------------------------- */
 
+/** FBoneTransformFinalizeData used for setting up root motion*/
+
+FBoneTransformFinalizeData::FBoneTransformFinalizeData() : SkeletalMeshComponent(nullptr), SwapRootBone(ESwapRootBone::SwapRootBone_None)
+{
+}
+
+void FBoneTransformFinalizeData::Register(USkeletalMeshComponent* InSkeleletalMeshCompononent, ESwapRootBone InSwapRootBone, FTransform& InCurrentTransform,TOptional<FTransform> InInitialTransform)
+{
+	if (OnBoneTransformsFinalizedHandle.IsValid() == false)
+	{
+		SkeletalMeshComponent = InSkeleletalMeshCompononent;
+		SwapRootBone = InSwapRootBone;
+		CurrentTransform = InCurrentTransform;
+		InitialTransform = InInitialTransform;
+		OnBoneTransformsFinalizedHandle = SkeletalMeshComponent->RegisterOnBoneTransformsFinalizedDelegate(
+			FOnBoneTransformsFinalizedMultiCast::FDelegate::CreateLambda([this]()
+				{
+					BoneTransformFinalized();
+				}));
+	}
+}
+
+void FBoneTransformFinalizeData::Unregister()
+{
+	if (SkeletalMeshComponent)
+	{
+		SkeletalMeshComponent->UnregisterOnBoneTransformsFinalizedDelegate(OnBoneTransformsFinalizedHandle);
+	}
+	OnBoneTransformsFinalizedHandle.Reset();
+}
+
+void FBoneTransformFinalizeData::BoneTransformFinalized()
+{
+	if (SkeletalMeshComponent && SwapRootBone != ESwapRootBone::SwapRootBone_None)
+	{
+		FTransform RelativeTransform = CurrentTransform;
+		
+		if (InitialTransform.IsSet())
+		{
+			RelativeTransform = RelativeTransform * InitialTransform.GetValue();
+		}
+		const bool bIsGameWorld = (SkeletalMeshComponent->GetWorld()->IsGameWorld());
+
+		if (SwapRootBone == ESwapRootBone::SwapRootBone_Component)
+		{
+			SkeletalMeshComponent->SetRelativeLocationAndRotation(RelativeTransform.GetLocation(), RelativeTransform.GetRotation().Rotator());
+		}
+		else if (SwapRootBone == ESwapRootBone::SwapRootBone_Actor)
+		{
+			AActor* Actor = SkeletalMeshComponent->GetOwner();
+			if (Actor && Actor->GetRootComponent())
+			{
+				Actor->GetRootComponent()->SetRelativeLocationAndRotation(RelativeTransform.GetLocation(), RelativeTransform.GetRotation().Rotator());
+			}
+		}
+	}
+}
+
+void FSkeletalAnimationSystemData::ResetSkeletalAnimations()
+{
+	//clear out the delegates
+	for (TTuple<UObject*, FBoundObjectActiveSkeletalAnimations>& Pair : SkeletalAnimations)
+	{
+		Pair.Value.BoneTransformFinalizeData.Unregister();
+	}
+	SkeletalAnimations.Reset();
+}
+
+/** ------------------------------------------------------------------------- */
 /** Task for gathering active skeletal animations */
 struct FGatherSkeletalAnimations
 {
@@ -399,12 +471,13 @@ private:
 		// Evaluate each group for the given bound object.
 		for (const TTuple<IMovieScenePlayer*, FBoundObjectActiveSkeletalAnimations>& Pair : SkeletalAnimationsPerPlayer)
 		{
-			EvaluateSkeletalAnimationsForPlayer(*Pair.Key, SkeletalMeshComponent, ExistingAnimInstance, SequencerInstance, Pair.Value);
+			EvaluateSkeletalAnimationsForPlayer(*Pair.Key,InObject, SkeletalMeshComponent, ExistingAnimInstance, SequencerInstance, Pair.Value);
 		}
 	}
 
 	void EvaluateSkeletalAnimationsForPlayer(
 			IMovieScenePlayer& InPlayer, 
+			UObject* InObject,
 			USkeletalMeshComponent* InSkeletalMeshComponent, 
 			UAnimInstance* InExistingAnimInstance, 
 			ISequencerAnimationSupport* InSequencerInstance, 
@@ -455,7 +528,7 @@ private:
 			UMovieSceneMotionVectorSimulationSystem* MotionVectorSim = InPlayer.GetEvaluationTemplate().GetEntitySystemLinker()->FindSystem<UMovieSceneMotionVectorSimulationSystem>();
 			if (MotionVectorSim && MotionVectorSim->IsSimulationEnabled())
 			{
-				ApplyAnimations(InPlayer, InSkeletalMeshComponent, InSkeletalAnimations.SimulatedAnimations, bPreviewPlayback, bFireNotifies, bResetDynamics);
+				ApplyAnimations(InPlayer, InObject, InSkeletalMeshComponent, InSkeletalAnimations.SimulatedAnimations, bPreviewPlayback, bFireNotifies, bResetDynamics);
 				InSkeletalMeshComponent->TickAnimation(0.f, false);
 				InSkeletalMeshComponent->ForceMotionVector();
 		
@@ -463,7 +536,7 @@ private:
 			}
 		}
 
-		ApplyAnimations(InPlayer, InSkeletalMeshComponent, InSkeletalAnimations.Animations, bPreviewPlayback, bFireNotifies, bResetDynamics);
+		ApplyAnimations(InPlayer, InObject, InSkeletalMeshComponent, InSkeletalAnimations.Animations, bPreviewPlayback, bFireNotifies, bResetDynamics);
 
 		// If the skeletal component has already ticked this frame because tick prerequisites weren't set up yet or a new binding was created, forcibly tick this component to update.
 		// This resolves first frame issues where the skeletal component ticks first, then the sequencer binding is resolved which sets up tick prerequisites
@@ -577,6 +650,7 @@ private:
 
 	void ApplyAnimations(
 		IMovieScenePlayer& Player,
+		UObject* InObject,
 		USkeletalMeshComponent* SkeletalMeshComponent,
 		TArrayView<const FActiveSkeletalAnimation> SkeletalAnimations,
 		bool bPreviewPlayback,
@@ -589,6 +663,18 @@ private:
 		{
 			const UMovieSceneSkeletalAnimationSection* AnimSection = SkeletalAnimation.AnimSection;
 			const FMovieSceneSkeletalAnimationParams& AnimParams = AnimSection->Params;
+			UMovieSceneSkeletalAnimationSection::FRootMotionParams RootMotionParams;
+			AnimSection->GetRootMotion(SkeletalAnimation.Context.GetTime().RoundToFrame(), RootMotionParams);
+			//set up root motion/bone transform delegates
+			if (AnimSection->Params.SwapRootBone != ESwapRootBone::SwapRootBone_None)
+			{
+				FTransform Transform = RootMotionParams.Transform.IsSet() ? RootMotionParams.Transform.GetValue() : FTransform::Identity;
+				TOptional<FTransform> InitialTransform = GetCurrentTransform(AnimSection->Params.SwapRootBone, SkeletalMeshComponent);
+				if (FBoundObjectActiveSkeletalAnimations* BoundObjectAnimations = SystemData.SkeletalAnimations.Find(InObject)) //InObject may be an AActor
+				{
+					BoundObjectAnimations->BoneTransformFinalizeData.Register(SkeletalMeshComponent, AnimSection->Params.SwapRootBone,Transform, InitialTransform);
+				}
+			}
 
 			// Don't fire notifies if looping around.
 			bool bLooped = false;
@@ -744,6 +830,7 @@ private:
 
 			int32 InstanceId = (SectionData) ? SectionData->MontageInstanceId : INDEX_NONE;
 
+
 			const float AssetPlayRate = FMath::IsNearlyZero(AnimParams.Animation->RateScale) ? 1.0f : AnimParams.Animation->RateScale;
 			TWeakObjectPtr<UAnimMontage> WeakMontage = FAnimMontageInstance::SetSequencerMontagePosition(
 					AnimParams.SlotName, 
@@ -759,6 +846,7 @@ private:
 			UAnimMontage* Montage = WeakMontage.Get();
 			if (Montage)
 			{
+
 				FMontagePlayerPerSectionData& DataContainer = SystemData.MontageData.FindOrAdd(Params.Section);
 				DataContainer.Montage = WeakMontage;
 				DataContainer.MontageInstanceId = InstanceId;
@@ -835,7 +923,7 @@ private:
 			FMontagePlayerPerSectionData* SectionData = SystemData.MontageData.Find(Params.Section);
 
 			int32 InstanceId = SectionData ? SectionData->MontageInstanceId : INDEX_NONE;
-
+		
 			const float AssetPlayRate = FMath::IsNearlyZero(AnimParams.Animation->RateScale) ? 1.0f : AnimParams.Animation->RateScale;
 			TWeakObjectPtr<UAnimMontage> WeakMontage = FAnimMontageInstance::PreviewSequencerMontagePosition(
 					AnimParams.SlotName,
@@ -853,6 +941,7 @@ private:
 			UAnimMontage* Montage = WeakMontage.Get();
 			if (Montage)
 			{
+
 				FMontagePlayerPerSectionData& DataContainer = SystemData.MontageData.FindOrAdd(Params.Section);
 				DataContainer.Montage = WeakMontage;
 				DataContainer.MontageInstanceId = InstanceId;
@@ -916,7 +1005,7 @@ void UMovieSceneSkeletalAnimationSystem::OnRun(FSystemTaskPrerequisites& InPrere
 	FMovieSceneTracksComponentTypes* TrackComponents = FMovieSceneTracksComponentTypes::Get();
 
 	// Start fresh every frame, gathering all active skeletal animations.
-	SystemData.SkeletalAnimations.Reset();
+	SystemData.ResetSkeletalAnimations();
 
 	FGraphEventRef GatherTask = FEntityTaskBuilder()
 	.ReadEntityIDs()
@@ -946,6 +1035,7 @@ void UMovieSceneSkeletalAnimationSystem::OnRun(FSystemTaskPrerequisites& InPrere
 
 void UMovieSceneSkeletalAnimationSystem::CleanSystemData()
 {
+	SystemData.ResetSkeletalAnimations();
 	// Clean-up old montage data.
 	for (auto It = SystemData.MontageData.CreateIterator(); It; ++It)
 	{
diff --git a/Engine/Source/Runtime/MovieSceneTracks/Private/Tracks/MovieSceneSkeletalAnimationTrack.cpp b/Engine/Source/Runtime/MovieSceneTracks/Private/Tracks/MovieSceneSkeletalAnimationTrack.cpp
--- a/Engine/Source/Runtime/MovieSceneTracks/Private/Tracks/MovieSceneSkeletalAnimationTrack.cpp
+++ b/Engine/Source/Runtime/MovieSceneTracks/Private/Tracks/MovieSceneSkeletalAnimationTrack.cpp
@@ -16,6 +16,7 @@
 #include "SkeletalDebugRendering.h"
 #include "Rendering/SkeletalMeshRenderData.h"
 #include "BoneContainer.h"
+#include "AnimSequencerInstanceProxy.h"
 
 #include UE_INLINE_GENERATED_CPP_BY_NAME(MovieSceneSkeletalAnimationTrack)
 
@@ -44,6 +45,7 @@ UMovieSceneSkeletalAnimationTrack::UMovieSceneSkeletalAnimationTrack(const FObje
 	SupportedBlendTypes.Add(EMovieSceneBlendType::Absolute);
 
 	EvalOptions.bEvaluateNearestSection_DEPRECATED = EvalOptions.bCanEvaluateNearestSection = true;
+	SwapRootBone = ESwapRootBone::SwapRootBone_None;
 }
 
 
@@ -95,6 +97,8 @@ TArray<UMovieSceneSection*> UMovieSceneSkeletalAnimationTrack::GetAnimSectionsAt
 void UMovieSceneSkeletalAnimationTrack::PostLoad()
 {
 	// UMovieSceneTrack::PostLoad removes null sections. However, RemoveAtSection requires SetupRootMotions, which accesses AnimationSections, so remove null sections here before anything else 
+	TOptional<ESwapRootBone> SectionsSwapRootBone;
+	bool bAllSectionsSameSwapRootBone = false;
 	for (int32 SectionIndex = 0; SectionIndex < AnimationSections.Num(); )
 	{
 		UMovieSceneSection* Section = AnimationSections[SectionIndex];
@@ -116,8 +120,28 @@ void UMovieSceneSkeletalAnimationTrack::PostLoad()
 		else
 		{
 			++SectionIndex;
+			if (UMovieSceneSkeletalAnimationSection* AnimSection = Cast<UMovieSceneSkeletalAnimationSection>(Section))
+			{
+				if (SectionsSwapRootBone.IsSet())
+				{
+					if (SectionsSwapRootBone.GetValue() != AnimSection->Params.SwapRootBone)
+					{
+						bAllSectionsSameSwapRootBone = false;
+					}
+				}
+				else
+				{
+					SectionsSwapRootBone = AnimSection->Params.SwapRootBone;
+					bAllSectionsSameSwapRootBone = true;
+				}
+			}
 		}
 	}
+	//if we have all sections with the same swap root bone, set that, probably from a previous version
+	if (bAllSectionsSameSwapRootBone && SectionsSwapRootBone.IsSet())
+	{
+		SwapRootBone = SectionsSwapRootBone.GetValue();
+	}
 
 	Super::PostLoad();
 
@@ -126,6 +150,26 @@ void UMovieSceneSkeletalAnimationTrack::PostLoad()
 		bUseLegacySectionIndexBlend = true;
 	}
 }
+
+void UMovieSceneSkeletalAnimationTrack::SetSwapRootBone(ESwapRootBone InValue)
+{
+	SwapRootBone = InValue;
+	for (UMovieSceneSection* Section : AnimationSections)
+	{
+		UMovieSceneSkeletalAnimationSection* AnimSection = Cast<UMovieSceneSkeletalAnimationSection>(Section);
+		if (AnimSection)
+		{
+			AnimSection->Params.SwapRootBone = SwapRootBone;
+		}
+	}
+	RootMotionParams.bRootMotionsDirty = true;
+}
+
+ESwapRootBone UMovieSceneSkeletalAnimationTrack::GetSwapRootBone() const
+{
+	return SwapRootBone;
+}
+
 #if WITH_EDITOR
 void UMovieSceneSkeletalAnimationTrack::PostEditImport()
 {
@@ -186,6 +230,7 @@ void UMovieSceneSkeletalAnimationTrack::AddSection(UMovieSceneSection& Section)
 	UMovieSceneSkeletalAnimationSection* AnimSection = Cast< UMovieSceneSkeletalAnimationSection>(&Section);
 	if (AnimSection)
 	{
+		AnimSection->Params.SwapRootBone = SwapRootBone;
 		SetUpRootMotions(true);
 	}
 }
@@ -623,7 +668,8 @@ void UMovieSceneSkeletalAnimationTrack::SetUpRootMotions(bool bForce)
 			}
 		}
 
-		if (bAnySectionsHaveOffset == false)
+		//if we are swapping root bone turn on root motion matching anyway
+		if (bAnySectionsHaveOffset == false && SwapRootBone == ESwapRootBone::SwapRootBone_None)
 		{
 #if WITH_EDITORONLY_DATA
 			RootMotionParams.RootTransforms.SetNum(0);
diff --git a/Engine/Source/Runtime/MovieSceneTracks/Public/Systems/MovieSceneSkeletalAnimationSystem.h b/Engine/Source/Runtime/MovieSceneTracks/Public/Systems/MovieSceneSkeletalAnimationSystem.h
--- a/Engine/Source/Runtime/MovieSceneTracks/Public/Systems/MovieSceneSkeletalAnimationSystem.h
+++ b/Engine/Source/Runtime/MovieSceneTracks/Public/Systems/MovieSceneSkeletalAnimationSystem.h
@@ -11,6 +11,7 @@
 class IMovieScenePlayer;
 class UAnimMontage;
 class UMovieSceneSkeletalAnimationSection;
+enum class ESwapRootBone : uint8;
 
 namespace UE::MovieScene
 {
@@ -29,6 +30,22 @@ struct FActiveSkeletalAnimation
 	bool bWantsRestoreState;
 };
 
+/** DelegateHandle and Skeletal Mesh for bone transform finalized */
+struct FBoneTransformFinalizeData
+{
+	FBoneTransformFinalizeData();
+	USkeletalMeshComponent* SkeletalMeshComponent;
+	ESwapRootBone SwapRootBone;
+	FTransform CurrentTransform;
+	TOptional<FTransform> InitialTransform;
+	/** Delegate Handle for skel mesh bone transform finalized callback*/
+	FDelegateHandle OnBoneTransformsFinalizedHandle;
+
+	void Register(USkeletalMeshComponent* InSkeleletalMeshCompononent, ESwapRootBone InSwapRootBone, FTransform& InCurrentTransform, TOptional<FTransform> InTransform);
+	void Unregister();
+	void BoneTransformFinalized();
+};
+
 /** Information for all skeletal animations playing on a bound object */
 struct FBoundObjectActiveSkeletalAnimations
 {
@@ -38,6 +55,9 @@ struct FBoundObjectActiveSkeletalAnimations
 	FAnimationArray Animations;
 	/** Motion vector simulation animations on the corresponding bound object */
 	FAnimationArray SimulatedAnimations;
+	/** SkelMesh and the bone finalize Delegate*/
+	FBoneTransformFinalizeData  BoneTransformFinalizeData;
+
 };
 
 /** Temporary information about montage setups. */
@@ -49,6 +69,8 @@ struct FMontagePlayerPerSectionData
 
 struct FSkeletalAnimationSystemData
 {
+	void ResetSkeletalAnimations();
+
 	/** Map of active skeletal animations for each bound object */
 	TMap<UObject*, FBoundObjectActiveSkeletalAnimations> SkeletalAnimations;
 
diff --git a/Engine/Source/Runtime/MovieSceneTracks/Public/Tracks/MovieSceneSkeletalAnimationTrack.h b/Engine/Source/Runtime/MovieSceneTracks/Public/Tracks/MovieSceneSkeletalAnimationTrack.h
--- a/Engine/Source/Runtime/MovieSceneTracks/Public/Tracks/MovieSceneSkeletalAnimationTrack.h
+++ b/Engine/Source/Runtime/MovieSceneTracks/Public/Tracks/MovieSceneSkeletalAnimationTrack.h
@@ -8,6 +8,7 @@
 #include "Sections/MovieSceneSkeletalAnimationSection.h"
 #include "MovieSceneSkeletalAnimationTrack.generated.h"
 
+enum class ESwapRootBone : uint8;
 
 /**Struct to hold the cached root motion positions based upon how we calculated them.
 * Also provides way to get the root motion at a particular time.
@@ -120,6 +121,15 @@ public:
 	UPROPERTY(EditAnywhere, Category = "Root Motions")
 	bool bBlendFirstChildOfRoot;
 
+	/** If on the root bone transform will be swapped to the specified root*/
+	UPROPERTY(BlueprintReadWrite, EditAnywhere, Setter,Getter, Category = "Root Motions")
+	ESwapRootBone SwapRootBone;
+
+	UFUNCTION()
+	MOVIESCENETRACKS_API void SetSwapRootBone(ESwapRootBone InValue);
+	UFUNCTION()
+	MOVIESCENETRACKS_API ESwapRootBone GetSwapRootBone() const;
+
 #if WITH_EDITORONLY_DATA
 public:
 
