Sequencer: don't cache pre-animated state on entities that are being torn down

diff --git a/Engine/Source/Runtime/MovieScene/Public/Evaluation/PreAnimatedState/MovieScenePreAnimatedStateStorage.h b/Engine/Source/Runtime/MovieScene/Public/Evaluation/PreAnimatedState/MovieScenePreAnimatedStateStorage.h
--- a/Engine/Source/Runtime/MovieScene/Public/Evaluation/PreAnimatedState/MovieScenePreAnimatedStateStorage.h
+++ b/Engine/Source/Runtime/MovieScene/Public/Evaluation/PreAnimatedState/MovieScenePreAnimatedStateStorage.h
@@ -574,10 +574,13 @@ public:
 
 		FBuiltInComponentTypes* BuiltInComponents = FBuiltInComponentTypes::Get();
 
+		FComponentMask ExcludeMask({ BuiltInComponents->Tags.NeedsUnlink, BuiltInComponents->Tags.Finished, BuiltInComponents->Tags.Ignored });
+
 		FEntityTaskBuilder()
 		.ReadEntityIDs()
 		.Read(BuiltInComponents->RootInstanceHandle)
 		.ReadAllOf(InComponentTypes...)
+		.FilterNone(ExcludeMask)
 		.CombineFilter(InParams.AdditionalFilter)
 		.Iterate_PerAllocation(&Linker->EntityManager, VisitAllocation);
 	}
diff --git a/Engine/Source/Runtime/MovieSceneTracks/Private/Systems/MovieSceneComponentAttachmentSystem.cpp b/Engine/Source/Runtime/MovieSceneTracks/Private/Systems/MovieSceneComponentAttachmentSystem.cpp
--- a/Engine/Source/Runtime/MovieSceneTracks/Private/Systems/MovieSceneComponentAttachmentSystem.cpp
+++ b/Engine/Source/Runtime/MovieSceneTracks/Private/Systems/MovieSceneComponentAttachmentSystem.cpp
@@ -257,6 +257,8 @@ void UMovieSceneComponentAttachmentSystem::SavePreAnimatedState(const FPreAnimat
 		FilterMask.Set(BuiltInComponents->Tags.NeedsLink);
 	}
 
+	FComponentMask ExcludeMask({ BuiltInComponents->Tags.NeedsUnlink, BuiltInComponents->Tags.Finished, BuiltInComponents->Tags.Ignored });
+
 	// Attachments change transforms
 	FPreAnimatedEntityCaptureSource* EntityMetaData = InParameters.CacheExtension->GetOrCreateEntityMetaData();
 	TSharedPtr<FPreAnimatedComponentTransformStorage>  ComponentTransformStorage  = InParameters.CacheExtension->GetOrCreateStorage<FPreAnimatedComponentTransformStorage>();
@@ -269,6 +271,7 @@ void UMovieSceneComponentAttachmentSystem::SavePreAnimatedState(const FPreAnimat
 	.Read(BuiltInComponents->BoundObject)
 	.Read(TrackComponents->AttachComponent)
 	.FilterAll(FilterMask)
+	.FilterNone(ExcludeMask)
 	.Iterate_PerAllocation(&Linker->EntityManager,
 		[EntityMetaData, ComponentTransformStorage, ComponentAttachmentStorage](FEntityAllocationIteratorItem Item, TRead<FMovieSceneEntityID> EntityIDs, TRead<FRootInstanceHandle> InstanceHandles, TRead<UObject*> BoundObjects, TRead<FAttachmentComponent> AttachComponents)
 		{
diff --git a/Engine/Source/Runtime/MovieSceneTracks/Private/Systems/MovieSceneComponentMobilitySystem.cpp b/Engine/Source/Runtime/MovieSceneTracks/Private/Systems/MovieSceneComponentMobilitySystem.cpp
--- a/Engine/Source/Runtime/MovieSceneTracks/Private/Systems/MovieSceneComponentMobilitySystem.cpp
+++ b/Engine/Source/Runtime/MovieSceneTracks/Private/Systems/MovieSceneComponentMobilitySystem.cpp
@@ -220,12 +220,15 @@ void UMovieSceneComponentMobilitySystem::SavePreAnimatedState(const FPreAnimatio
 		RestrictiveMask.Set(BuiltInComponents->Tags.NeedsLink);
 	}
 
+	FComponentMask ExcludeMask({ BuiltInComponents->Tags.NeedsUnlink, BuiltInComponents->Tags.Finished, BuiltInComponents->Tags.Ignored });
+
 	FEntityTaskBuilder()
 	.ReadEntityIDs()
 	.Read(BuiltInComponents->RootInstanceHandle)
 	.Read(BuiltInComponents->BoundObject)
 	.FilterAll(RestrictiveMask)
 	.FilterAny({ TrackComponents->ComponentTransform.PropertyTag, TrackComponents->AttachParent })
+	.FilterNone(ExcludeMask)
 	.Iterate_PerAllocation(&Linker->EntityManager, IterNewObjects);
 }
 
diff --git a/Engine/Source/Runtime/MovieSceneTracks/Private/Systems/MovieScenePropertySystem.cpp b/Engine/Source/Runtime/MovieSceneTracks/Private/Systems/MovieScenePropertySystem.cpp
--- a/Engine/Source/Runtime/MovieSceneTracks/Private/Systems/MovieScenePropertySystem.cpp
+++ b/Engine/Source/Runtime/MovieSceneTracks/Private/Systems/MovieScenePropertySystem.cpp
@@ -67,19 +67,22 @@ void UMovieScenePropertySystem::SavePreAnimatedState(const FPreAnimationParamete
 
 	PreAnimatedStorageID = PreAnimatedStorage->GetStorageType();
 
-	FComponentMask ComponentMask({ Definition.PropertyType });
+	FComponentMask IncludeMask({ Definition.PropertyType });
 	if (!InParameters.CacheExtension->AreEntriesInvalidated())
 	{
-		ComponentMask.Set(BuiltInComponents->Tags.NeedsLink);
+		IncludeMask.Set(BuiltInComponents->Tags.NeedsLink);
 	}
 
+	FComponentMask ExcludeMask({ BuiltInComponents->Tags.NeedsUnlink, BuiltInComponents->Tags.Finished, BuiltInComponents->Tags.Ignored });
+
 	if (IPreAnimatedObjectEntityStorage* ObjectStorage = PreAnimatedStorage->AsObjectStorage())
 	{
 		FEntityTaskBuilder()
 		.ReadEntityIDs()
 		.Read(BuiltInComponents->RootInstanceHandle)
 		.Read(BuiltInComponents->BoundObject)
-		.FilterAll(ComponentMask)
+		.FilterAll(IncludeMask)
+		.FilterNone(ExcludeMask)
 		.Iterate_PerAllocation(&Linker->EntityManager,
 			[ObjectStorage](FEntityAllocationIteratorItem Item, TRead<FMovieSceneEntityID> EntityIDs, TRead<FRootInstanceHandle> RootInstanceHandles, TRead<UObject*> BoundObjects)
 			{
@@ -89,7 +92,8 @@ void UMovieScenePropertySystem::SavePreAnimatedState(const FPreAnimationParamete
 
 		FEntityTaskBuilder()
 		.Read(BuiltInComponents->BoundObject)
-		.FilterAll(ComponentMask)
+		.FilterAll(IncludeMask)
+		.FilterNone(ExcludeMask)
 		.Iterate_PerAllocation(&Linker->EntityManager,
 			[ObjectStorage](FEntityAllocationIteratorItem Item, TRead<UObject*> Objects)
 			{
@@ -105,7 +109,8 @@ void UMovieScenePropertySystem::SavePreAnimatedState(const FPreAnimationParamete
 		.Read(BuiltInComponents->RootInstanceHandle)
 		.Read(BuiltInComponents->BoundObject)
 		.Read(BuiltInComponents->PropertyBinding)
-		.FilterAll(ComponentMask)
+		.FilterAll(IncludeMask)
+		.FilterNone(ExcludeMask)
 		.Iterate_PerAllocation(&Linker->EntityManager,
 			[PropertyStorage](FEntityAllocationIteratorItem Item, TRead<FMovieSceneEntityID> EntityIDs, TRead<FRootInstanceHandle> RootInstanceHandles, TRead<UObject*> BoundObjects, TRead<FMovieScenePropertyBinding> PropertyBindings)
 			{
@@ -117,7 +122,8 @@ void UMovieScenePropertySystem::SavePreAnimatedState(const FPreAnimationParamete
 		.Read(BuiltInComponents->BoundObject)
 		.Read(BuiltInComponents->PropertyBinding)
 		.ReadOneOf(BuiltInComponents->CustomPropertyIndex, BuiltInComponents->FastPropertyOffset, BuiltInComponents->SlowProperty)
-		.FilterAll(ComponentMask)
+		.FilterAll(IncludeMask)
+		.FilterNone(ExcludeMask)
 		.Iterate_PerAllocation(&Linker->EntityManager,
 			[PropertyStorage](FEntityAllocationIteratorItem Item, TRead<UObject*> Objects, TRead<FMovieScenePropertyBinding> PropertyBindings, FThreeWayAccessor ResolvedProperties)
 			{
