Sequencer: Added the ability to blend between hierarchical bias levels

  - 2 new component types control this behavior:
      - HierarchicalBlendTarget (int16): This defines the highest hbias that contributes to any given 'group' of entities. When present, it causes additional hierarchical easing channels to be created for that HBias and all its lesser hbiases in order to provide accurate blend weight accumulation for that specific chain of contributors.
      - BlendHierarchicalBias: A new tag that causes HierarchicalBlendTarget components to be added to any contributors of a blend channel.

  - The property instantiator has been expanded to provide the necessary bookkeeping for these components, and material parameters need their own bookkeeping. Ultimately we need to look at ways to centralize this behavior.
  - The weight and easing instantiator systems now detect HierarchicalBlendTarget components and create the necessary blend channels in order to correctly evaluate and combine these weights.

Fast path is now re-enabled for previously blended entities when eligible.
Removed the MigratedFromFastPath tag since this was never used.

diff --git a/Engine/Source/Runtime/MovieScene/Private/EntitySystem/BuiltInComponentTypes.cpp b/Engine/Source/Runtime/MovieScene/Private/EntitySystem/BuiltInComponentTypes.cpp
--- a/Engine/Source/Runtime/MovieScene/Private/EntitySystem/BuiltInComponentTypes.cpp
+++ b/Engine/Source/Runtime/MovieScene/Private/EntitySystem/BuiltInComponentTypes.cpp
@@ -121,7 +121,8 @@ FBuiltInComponentTypes::FBuiltInComponentTypes()
 
 	ComponentRegistry->NewComponentType(&Easing,                  TEXT("Easing"));
 	ComponentRegistry->NewComponentType(&EasingResult,            TEXT("Easing Result"));
-	ComponentRegistry->NewComponentType(&HierarchicalEasingChannel, TEXT("Hierarchical Easing Channel"));
+	ComponentRegistry->NewComponentType(&HierarchicalEasingChannel,  TEXT("Hierarchical Easing Channel"));
+	ComponentRegistry->NewComponentType(&HierarchicalBlendTarget,    TEXT("Hierarchical Blend Target"));
 	ComponentRegistry->NewComponentType(&HierarchicalEasingProvider, TEXT("Hierarchical Easing Provider"));
 
 	ComponentRegistry->NewComponentType(&BlenderType,           TEXT("Blender System Type"), EComponentTypeFlags::CopyToChildren);
@@ -177,6 +178,7 @@ FBuiltInComponentTypes::FBuiltInComponentTypes()
 
 	Tags.RestoreState            = ComponentRegistry->NewTag(TEXT("Is Restore State Entity"));
 	Tags.IgnoreHierarchicalBias  = ComponentRegistry->NewTag(TEXT("Ignore Hierarchical Bias"));
+	Tags.BlendHierarchicalBias   = ComponentRegistry->NewTag(TEXT("Blend Hierarchical Bias"));
 
 	Tags.AbsoluteBlend           = ComponentRegistry->NewTag(TEXT("Is Absolute Blend"));
 	Tags.RelativeBlend           = ComponentRegistry->NewTag(TEXT("Is Relative Blend"));
@@ -186,7 +188,7 @@ FBuiltInComponentTypes::FBuiltInComponentTypes()
 	Tags.NeedsLink               = ComponentRegistry->NewTag(TEXT("Needs Link"));
 	Tags.NeedsUnlink             = ComponentRegistry->NewTag(TEXT("Needs Unlink"));
 	Tags.HasUnresolvedBinding    = ComponentRegistry->NewTag(TEXT("Has Unresolved Binding"));
-	Tags.MigratedFromFastPath    = ComponentRegistry->NewTag(TEXT("Migrated From Fast Path"));
+	Tags.HasAssignedInitialValue = ComponentRegistry->NewTag(TEXT("Has Assigned Initial Value"));
 	Tags.Root                    = ComponentRegistry->NewTag(TEXT("Root"));
 	Tags.SubInstance             = ComponentRegistry->NewTag(TEXT("Sub Instance"));
 	Tags.ImportedEntity          = ComponentRegistry->NewTag(TEXT("Imported Entity"));
@@ -215,6 +217,8 @@ FBuiltInComponentTypes::FBuiltInComponentTypes()
 
 	// Always copy these tags over to children
 	ComponentRegistry->Factories.DefineChildComponent(Tags.RestoreState,  Tags.RestoreState);
+	ComponentRegistry->Factories.DefineChildComponent(Tags.IgnoreHierarchicalBias,  Tags.IgnoreHierarchicalBias);
+	ComponentRegistry->Factories.DefineChildComponent(Tags.BlendHierarchicalBias,  Tags.BlendHierarchicalBias);
 	ComponentRegistry->Factories.DefineChildComponent(Tags.AbsoluteBlend, Tags.AbsoluteBlend);
 	ComponentRegistry->Factories.DefineChildComponent(Tags.RelativeBlend, Tags.RelativeBlend);
 	ComponentRegistry->Factories.DefineChildComponent(Tags.AdditiveBlend, Tags.AdditiveBlend);
diff --git a/Engine/Source/Runtime/MovieScene/Private/MovieSceneSection.cpp b/Engine/Source/Runtime/MovieScene/Private/MovieSceneSection.cpp
--- a/Engine/Source/Runtime/MovieScene/Private/MovieSceneSection.cpp
+++ b/Engine/Source/Runtime/MovieScene/Private/MovieSceneSection.cpp
@@ -359,6 +359,7 @@ void UMovieSceneSection::BuildDefaultComponents(UMovieSceneEntitySystemLinker* E
 		.AddConditional(Components->EvalTime,                       Params.EntityMetaData ? Params.EntityMetaData->ForcedTime : 0, bHasForcedTime)
 		.AddTagConditional(Components->Tags.RestoreState,           bShouldRestoreState)
 		.AddTagConditional(Components->Tags.IgnoreHierarchicalBias, EnumHasAnyFlags(Params.Sequence.SubSectionFlags, EMovieSceneSubSectionFlags::IgnoreHierarchicalBias))
+		.AddTagConditional(Components->Tags.BlendHierarchicalBias,  EnumHasAnyFlags(Params.Sequence.SubSectionFlags, EMovieSceneSubSectionFlags::BlendHierarchicalBias))
 		.AddTagConditional(Components->Tags.FixedTime,              bHasForcedTime)
 		.AddTagConditional(Components->Tags.SectionPreRoll,         bHasSectionPreRoll)
 		.AddTagConditional(Components->Tags.PreRoll,                bHasSequencePreRoll || bHasSectionPreRoll)
diff --git a/Engine/Source/Runtime/MovieScene/Public/EntitySystem/BuiltInComponentTypes.h b/Engine/Source/Runtime/MovieScene/Public/EntitySystem/BuiltInComponentTypes.h
--- a/Engine/Source/Runtime/MovieScene/Public/EntitySystem/BuiltInComponentTypes.h
+++ b/Engine/Source/Runtime/MovieScene/Public/EntitySystem/BuiltInComponentTypes.h
@@ -440,6 +440,9 @@ public:
 	// The sub-sequence ID that should receive ease in/out as a whole
 	TComponentTypeID<FMovieSceneSequenceID> HierarchicalEasingProvider;
 
+	// Defines an HBias level that is the highest blend target for a given set of components that need to blend together
+	TComponentTypeID<int16>       HierarchicalBlendTarget;
+
 	// A float representing the evaluated easing weight
 	TComponentTypeID<double> WeightAndEasingResult;
 
@@ -475,6 +478,7 @@ public:
 
 		// A tag that specifies this entity should always contribute to the output, regardless of hbias
 		FComponentTypeID IgnoreHierarchicalBias;
+		FComponentTypeID BlendHierarchicalBias;
 
 		FComponentTypeID AbsoluteBlend;
 		FComponentTypeID RelativeBlend;
@@ -487,7 +491,7 @@ public:
 		/** Tag that is added to imported entities with a GenericObjectBinding or SceneComponentBinding whose binding did not resolve */
 		FComponentTypeID HasUnresolvedBinding;
 
-		FComponentTypeID MigratedFromFastPath;
+		FComponentTypeID HasAssignedInitialValue;
 
 		FComponentTypeID ImportedEntity;
 
diff --git a/Engine/Source/Runtime/MovieScene/Public/Evaluation/MovieSceneSectionParameters.h b/Engine/Source/Runtime/MovieScene/Public/Evaluation/MovieSceneSectionParameters.h
--- a/Engine/Source/Runtime/MovieScene/Public/Evaluation/MovieSceneSectionParameters.h
+++ b/Engine/Source/Runtime/MovieScene/Public/Evaluation/MovieSceneSectionParameters.h
@@ -26,6 +26,9 @@ enum class EMovieSceneSubSectionFlags
 	/** Everything inside this sub-sequence should ignore hierarchical bias and always be relevant */
 	IgnoreHierarchicalBias = 1 << 2,
 
+	/** Blend this sub sequence's hierarchical bias level using a higher -> lower override. Values from higher biases will override those in lower biases until a combined weight of 1 is reached. */
+	BlendHierarchicalBias  = 1 << 3,
+
 	AnyRestoreStateOverride = OverrideKeepState | OverrideRestoreState,
 };
 ENUM_CLASS_FLAGS(EMovieSceneSubSectionFlags)
diff --git a/Engine/Source/Runtime/MovieSceneTracks/Private/EntitySystem/Interrogation/MovieSceneInterrogatedPropertyInstantiator.cpp b/Engine/Source/Runtime/MovieSceneTracks/Private/EntitySystem/Interrogation/MovieSceneInterrogatedPropertyInstantiator.cpp
--- a/Engine/Source/Runtime/MovieSceneTracks/Private/EntitySystem/Interrogation/MovieSceneInterrogatedPropertyInstantiator.cpp
+++ b/Engine/Source/Runtime/MovieSceneTracks/Private/EntitySystem/Interrogation/MovieSceneInterrogatedPropertyInstantiator.cpp
@@ -238,7 +238,6 @@ void UMovieSceneInterrogatedPropertyInstantiatorSystem::UpdateOutput(UE::MovieSc
 		FMovieSceneEntityID NewEntityID = FEntityBuilder()
 		.Add(BuiltInComponents->Interrogation.OutputKey, Key)
 		.Add(BuiltInComponents->BlendChannelOutput, BlendChannel)
-		.AddTagConditional(BuiltInComponents->Tags.MigratedFromFastPath, Output->PropertyEntityID.IsValid())
 		.AddTag(BuiltInComponents->Tags.NeedsLink)
 		.AddTag(BlenderSystem->GetBlenderTypeTag())
 		.AddMutualComponents()
diff --git a/Engine/Source/Runtime/MovieSceneTracks/Private/Systems/MovieSceneHierarchicalBiasSystem.cpp b/Engine/Source/Runtime/MovieSceneTracks/Private/Systems/MovieSceneHierarchicalBiasSystem.cpp
--- a/Engine/Source/Runtime/MovieSceneTracks/Private/Systems/MovieSceneHierarchicalBiasSystem.cpp
+++ b/Engine/Source/Runtime/MovieSceneTracks/Private/Systems/MovieSceneHierarchicalBiasSystem.cpp
@@ -47,8 +47,10 @@ struct FHierarchicalBiasTask
 	void ForEachAllocation(FEntityAllocationIteratorItem Iterator, TRead<FMovieSceneEntityID> EntityIDs, TRead<FRootInstanceHandle> RootInstanceHandles, TRead<FMovieSceneBlendChannelID> BlendChannels, TReadOptional<int16> OptHBiases)
 	{
 		const int32 Num = Iterator.GetAllocation()->Num();
+		const FComponentMask& AllocationType = Iterator.GetAllocationType();
+		const bool bIgnoreBias = AllocationType.Contains(FBuiltInComponentTypes::Get()->Tags.IgnoreHierarchicalBias)
+			|| AllocationType.Contains(FBuiltInComponentTypes::Get()->HierarchicalBlendTarget);
 
-		const bool bIgnoreBias = Iterator.GetAllocationType().Contains(FBuiltInComponentTypes::Get()->Tags.IgnoreHierarchicalBias);
 		if (bIgnoreBias)
 		{
 			for (int32 Index = 0; Index < Num; ++Index)
diff --git a/Engine/Source/Runtime/MovieSceneTracks/Private/Systems/MovieSceneMaterialParameterSystem.cpp b/Engine/Source/Runtime/MovieSceneTracks/Private/Systems/MovieSceneMaterialParameterSystem.cpp
--- a/Engine/Source/Runtime/MovieSceneTracks/Private/Systems/MovieSceneMaterialParameterSystem.cpp
+++ b/Engine/Source/Runtime/MovieSceneTracks/Private/Systems/MovieSceneMaterialParameterSystem.cpp
@@ -9,6 +9,7 @@
 #include "Systems/MovieSceneMaterialSystem.h"
 #include "Systems/MovieSceneInitialValueSystem.h"
 #include "Systems/MovieScenePreAnimatedMaterialParameters.h"
+#include "Systems/WeightAndEasingEvaluatorSystem.h"
 
 #include "Evaluation/PreAnimatedState/MovieScenePreAnimatedObjectStorage.h"
 
@@ -194,6 +195,8 @@ struct FScalarMixin
 				Linker->EntityManager.AddComponent(SoleContributor, TracksComponents->FloatParameter.InitialValue, InitialValue);
 			}
 		}
+
+		Linker->EntityManager.RemoveComponent(SoleContributor, FBuiltInComponentTypes::Get()->HierarchicalBlendTarget);
 	}
 };
 
@@ -331,6 +334,40 @@ struct TOverlappingMaterialParameterHandler : Mixin
 
 			const FComponentTypeID BlenderTypeTag = System->DoubleBlenderSystem->GetBlenderTypeTag();
 
+			struct FBlendInfo
+			{
+				int16 HBias = TNumericLimits<int16>::Min();
+				bool bBlendHierarchicalBias = true;
+			};
+			FBlendInfo IgnoredBlendInfo;
+			FBlendInfo BlendInfo;
+			for (FMovieSceneEntityID Input : Inputs)
+			{
+				FBlendInfo& BlendInfoToUpdate = Linker->EntityManager.HasComponent(Input, BuiltInComponents->Tags.Ignored) ? IgnoredBlendInfo : BlendInfo;
+
+				TOptionalComponentReader<int16> HBiasComponent = Linker->EntityManager.ReadComponent(Input, BuiltInComponents->HierarchicalBias);
+				const int16 HBias = HBiasComponent ? *HBiasComponent : 0;
+
+				if (HBias > BlendInfoToUpdate.HBias)
+				{
+					BlendInfoToUpdate.HBias = HBias;
+					BlendInfoToUpdate.bBlendHierarchicalBias = Linker->EntityManager.HasComponent(Input, BuiltInComponents->Tags.BlendHierarchicalBias);
+				}
+				else if (HBias == BlendInfoToUpdate.HBias && !BlendInfoToUpdate.bBlendHierarchicalBias)
+				{
+					BlendInfoToUpdate.bBlendHierarchicalBias = Linker->EntityManager.HasComponent(Input, BuiltInComponents->Tags.BlendHierarchicalBias);
+				}
+			}
+
+			if (BlendInfo.HBias == TNumericLimits<int16>::Min())
+			{
+				BlendInfo = IgnoredBlendInfo;
+			}
+			else if (IgnoredBlendInfo.HBias != TNumericLimits<int16>::Min())
+			{
+				BlendInfo.bBlendHierarchicalBias |= IgnoredBlendInfo.bBlendHierarchicalBias;
+			}
+
 			for (FMovieSceneEntityID Input : Inputs)
 			{
 				if (!Linker->EntityManager.HasComponent(Input, BuiltInComponents->BlendChannelInput))
@@ -343,6 +380,18 @@ struct TOverlappingMaterialParameterHandler : Mixin
 					Linker->EntityManager.WriteComponentChecked(Input, BuiltInComponents->BlendChannelInput, Output->BlendChannelID);
 				}
 
+				if (BlendInfo.bBlendHierarchicalBias)
+				{
+					if (!Linker->EntityManager.HasComponent(Input, BuiltInComponents->HierarchicalBlendTarget))
+					{
+						Linker->EntityManager.AddComponent(Input, BuiltInComponents->HierarchicalBlendTarget, BlendInfo.HBias);
+					}
+					else
+					{
+						Linker->EntityManager.WriteComponentChecked(Input, BuiltInComponents->HierarchicalBlendTarget, BlendInfo.HBias);
+					}
+				}
+
 				// Ensure we have the blender type tag on the inputs.
 				Linker->EntityManager.AddComponent(Input, BlenderTypeTag);
 			}
@@ -392,6 +441,7 @@ UMovieSceneMaterialParameterSystem::UMovieSceneMaterialParameterSystem(const FOb
 	if (HasAnyFlags(RF_ClassDefaultObject))
 	{
 		DefineComponentConsumer(GetClass(), TracksComponents->BoundMaterial);
+		DefineComponentProducer(GetClass(), BuiltInComponents->HierarchicalBlendTarget);
 
 		DefineImplicitPrerequisite(UFloatChannelEvaluatorSystem::StaticClass(), GetClass());
 		DefineImplicitPrerequisite(UDoubleChannelEvaluatorSystem::StaticClass(), GetClass());
@@ -401,6 +451,7 @@ UMovieSceneMaterialParameterSystem::UMovieSceneMaterialParameterSystem(const FOb
 			DefineComponentConsumer(GetClass(), BuiltInComponents->DoubleResult[Index]);
 		}
 
+		DefineImplicitPrerequisite(UMovieSceneHierarchicalEasingInstantiatorSystem::StaticClass(), GetClass());
 		DefineImplicitPrerequisite(UMovieScenePiecewiseDoubleBlenderSystem::StaticClass(), GetClass());
 		DefineImplicitPrerequisite(GetClass(), UMovieSceneHierarchicalBiasSystem::StaticClass());
 		DefineImplicitPrerequisite(GetClass(), UMovieSceneInitialValueSystem::StaticClass());
diff --git a/Engine/Source/Runtime/MovieSceneTracks/Private/Systems/MovieScenePropertyInstantiator.cpp b/Engine/Source/Runtime/MovieSceneTracks/Private/Systems/MovieScenePropertyInstantiator.cpp
--- a/Engine/Source/Runtime/MovieSceneTracks/Private/Systems/MovieScenePropertyInstantiator.cpp
+++ b/Engine/Source/Runtime/MovieSceneTracks/Private/Systems/MovieScenePropertyInstantiator.cpp
@@ -20,6 +20,31 @@ DECLARE_CYCLE_STAT(TEXT("DiscoverInvalidatedProperties"), MovieSceneEval_Discove
 DECLARE_CYCLE_STAT(TEXT("ProcessInvalidatedProperties"), MovieSceneEval_ProcessInvalidatedProperties, STATGROUP_MovieSceneECS);
 DECLARE_CYCLE_STAT(TEXT("InitializePropertyMetaData"), MovieSceneEval_InitializePropertyMetaData, STATGROUP_MovieSceneECS);
 
+
+void UMovieScenePropertyInstantiatorSystem::FHierarchicalMetaData::CombineWith(const FHierarchicalMetaData& Other)
+{
+	NumContributors        += Other.NumContributors;
+	bWantsRestoreState     |= Other.bWantsRestoreState;
+	bSupportsFastPath      &= Other.bSupportsFastPath;
+	bNeedsInitialValue     |= Other.bNeedsInitialValue;
+	bBlendHierarchicalBias |= Other.bBlendHierarchicalBias;
+}
+
+void UMovieScenePropertyInstantiatorSystem::FHierarchicalMetaData::ResetTracking()
+{
+	NumContributors = 0;
+	bWantsRestoreState = false;
+	bSupportsFastPath = true;
+	bNeedsInitialValue = false;
+	bBlendHierarchicalBias = false;
+	bInUse = false;
+}
+
+UMovieScenePropertyInstantiatorSystem::FContributorKey UMovieScenePropertyInstantiatorSystem::FPropertyParameters::MakeContributorKey() const
+{
+	return PropertyInfo->HierarchicalMetaData.bBlendHierarchicalBias ? FContributorKey(PropertyInfoIndex) : FContributorKey(PropertyInfoIndex, PropertyInfo->HierarchicalMetaData.HBias);
+}
+
 UMovieScenePropertyInstantiatorSystem::UMovieScenePropertyInstantiatorSystem(const FObjectInitializer& ObjInit)
 	: Super(ObjInit)
 {
@@ -36,6 +61,7 @@ UMovieScenePropertyInstantiatorSystem::UMovieScenePropertyInstantiatorSystem(con
 	{
 		DefineComponentConsumer(GetClass(), BuiltInComponents->BoundObject);
 		DefineComponentProducer(GetClass(), BuiltInComponents->BlendChannelInput);
+		DefineComponentProducer(GetClass(), BuiltInComponents->HierarchicalBlendTarget);
 		DefineComponentProducer(GetClass(), BuiltInComponents->SymbolicTags.CreatesEntities);
 	}
 }
@@ -149,8 +175,11 @@ void UMovieScenePropertyInstantiatorSystem::DiscoverInvalidatedProperties(TBitAr
 
 	PropertyStats.SetNum(Properties.Num());
 
-	auto VisitNewProperties = [this, Properties, &OutInvalidatedProperties](const FEntityAllocation* Allocation, const FMovieSceneEntityID* EntityIDs, UObject* const * ObjectPtrs, const FMovieScenePropertyBinding* PropertyPtrs)
+	auto VisitNewProperties = [this, Properties, &OutInvalidatedProperties](FEntityAllocationIteratorItem AllocationItem, const FMovieSceneEntityID* EntityIDs, UObject* const * ObjectPtrs, const FMovieScenePropertyBinding* PropertyPtrs, const int16* HierarchicalBiases)
 	{
+		const FEntityAllocation* Allocation     = AllocationItem.GetAllocation();
+		const FComponentMask&    AllocationType = AllocationItem.GetAllocationType();
+
 		const int32 PropertyDefinitionIndex = Algo::IndexOfByPredicate(Properties, [=](const FPropertyDefinition& InDefinition){ return Allocation->HasComponent(InDefinition.PropertyType); });
 		if (PropertyDefinitionIndex == INDEX_NONE)
 		{
@@ -161,6 +190,15 @@ void UMovieScenePropertyInstantiatorSystem::DiscoverInvalidatedProperties(TBitAr
 
 		FCustomAccessorView CustomAccessors = PropertyDefinition.CustomPropertyRegistration ? PropertyDefinition.CustomPropertyRegistration->GetAccessors() : FCustomAccessorView();
 
+		// Figure out the hbias we should use for this allocation.
+		// If the allocation is tagged to ignore hbias, we null out the hbias components and use the ANY_HBIAS symbolic value
+		int16 DefaultHierarchicalBias = 0;
+		if (AllocationType.Contains(this->BuiltInComponents->Tags.IgnoreHierarchicalBias))
+		{
+			HierarchicalBiases = nullptr;
+			DefaultHierarchicalBias = FContributorKey::ANY_HBIAS;
+		}
+
 		for (int32 Index = 0; Index < Allocation->Num(); ++Index)
 		{
 			const int32 PropertyIndex = this->ResolveProperty(CustomAccessors, ObjectPtrs[Index], PropertyPtrs[Index], PropertyDefinitionIndex);
@@ -172,8 +210,9 @@ void UMovieScenePropertyInstantiatorSystem::DiscoverInvalidatedProperties(TBitAr
 
 			if (PropertyIndex != INDEX_NONE)
 			{
-				this->Contributors.Add(PropertyIndex, EntityIDs[Index]);
-				this->NewContributors.Add(PropertyIndex, EntityIDs[Index]);
+				FContributorKey Key { PropertyIndex, HierarchicalBiases ? HierarchicalBiases[Index] : DefaultHierarchicalBias };
+				this->Contributors.Add(Key, EntityIDs[Index]);
+				this->NewContributors.Add(Key, EntityIDs[Index]);
 
 				OutInvalidatedProperties.PadToNum(PropertyIndex + 1, false);
 				OutInvalidatedProperties[PropertyIndex] = true;
@@ -185,6 +224,7 @@ void UMovieScenePropertyInstantiatorSystem::DiscoverInvalidatedProperties(TBitAr
 	.ReadEntityIDs()
 	.Read(BuiltInComponents->BoundObject)
 	.Read(BuiltInComponents->PropertyBinding)
+	.ReadOptional(BuiltInComponents->HierarchicalBias)
 	.FilterNone({ BuiltInComponents->BlendChannelOutput })
 	.FilterAll({ BuiltInComponents->Tags.NeedsLink })
 	.Iterate_PerAllocation(&Linker->EntityManager, VisitNewProperties);
@@ -202,15 +242,6 @@ void UMovieScenePropertyInstantiatorSystem::DiscoverInvalidatedProperties(TBitAr
 				OutInvalidatedProperties[PropertyIndex] = true;
 
 				this->Contributors.Remove(PropertyIndex, EntityID);
-				
-				// The output entity of a property could be going away if that property is using the
-				// fast path, which means its sole contributor is that output entity. It can go away when
-				// that track or section gets reimported.
-				FObjectPropertyInfo& PropertyInfo = this->ResolvedProperties[PropertyIndex];
-				if (PropertyInfo.PropertyEntityID == EntityID)
-				{
-					PropertyInfo.PropertyEntityID = FMovieSceneEntityID::Invalid();
-				}
 			}
 
 			// Always remove the entity ID from the LUT
@@ -331,7 +362,8 @@ void UMovieScenePropertyInstantiatorSystem::UpgradeFloatToDoubleProperties(const
 		if (ensure(OldPropertyTag && NewPropertyTag))
 		{
 			// Swap out the property tag
-			for (auto ContribIt = NewContributors.CreateKeyIterator(PropertyIndex); ContribIt; ++ContribIt)
+			FContributorKey Key(PropertyIndex);
+			for (auto ContribIt = NewContributors.CreateKeyIterator(Key); ContribIt; ++ContribIt)
 			{
 				const FMovieSceneEntityID CurID(ContribIt.Value());
 				FComponentMask EntityType = Linker->EntityManager.GetEntityType(CurID);
@@ -385,7 +417,7 @@ void UMovieScenePropertyInstantiatorSystem::ProcessInvalidatedProperties(const T
 			StaleProperties[PropertyIndex] = true;
 		}
 		// Does it support fast path?
-		else if (PropertySupportsFastPath(Params))
+		else if (Params.PropertyInfo->HierarchicalMetaData.bSupportsFastPath)
 		{
 			InitializeFastPath(Params);
 		}
@@ -411,10 +443,10 @@ void UMovieScenePropertyInstantiatorSystem::ProcessInvalidatedProperties(const T
 					const FMovieSceneBlendChannelID BlendChannelID(Blender->GetBlenderSystemID(), PropertyInfo->BlendChannel);
 					Blender->ReleaseBlendChannel(BlendChannelID);
 				}
-				Linker->EntityManager.AddComponents(PropertyInfo->PropertyEntityID, BuiltInComponents->FinishedMask);
+				Linker->EntityManager.AddComponents(PropertyInfo->FinalBlendOutputID, BuiltInComponents->FinishedMask);
 			}
 
-			if (PropertyInfo->EmptyChannels.Find(true) != INDEX_NONE)
+			if (PropertyInfo->bIsPartiallyAnimated)
 			{
 				--PropertyStats[PropertyInfo->PropertyDefinitionIndex].NumPartialProperties;
 			}
@@ -434,6 +466,8 @@ void UMovieScenePropertyInstantiatorSystem::ProcessInvalidatedProperties(const T
 				It.RemoveCurrent();
 			}
 		}
+
+		ResolvedProperties.Shrink();
 	}
 
 	NewContributors.Empty();
@@ -445,129 +479,262 @@ void UMovieScenePropertyInstantiatorSystem::UpdatePropertyInfo(const FPropertyPa
 
 	TArrayView<const FPropertyCompositeDefinition> Composites = BuiltInComponents->PropertyRegistry.GetComposites(*Params.PropertyDefinition);
 
-	FChannelMask EmptyChannels(true, Params.PropertyDefinition->CompositeSize);
+	// This function updates the meta-data associated with a property for each hbias that it is animated from
+	// There are 3 possible 'modes' for hbias to be considered:
+	//    - Blended means that greater biases are allowed to override lower biases, but still blend with them when the weight is < 1
+	//    - Non-blended hbias simply disables contribution from any lower biases
+	//    - Entities tagged with IgnoreHierarchicalBias will always be relevant and contribute with the highest bias
+
+
+	// Channel masks for all entities, entities within the active hbias, and within the 'ignored hbias' buckets
+	// Set bits denote channels that are not animated by entities in these contexts
+	FChannelMask BlendedBiasEmptyChannels(true, Params.PropertyDefinition->CompositeSize);
+	FChannelMask ActiveBiasEmptyChannels(true, Params.PropertyDefinition->CompositeSize);
+	FChannelMask IgnoredBiasEmptyChannels(true, Params.PropertyDefinition->CompositeSize);
+
+	// Key that visits any contributor to this property regardless of hbias
+	FContributorKey AnyContributor(Params.PropertyInfoIndex);
 
-	bool bWantsRestoreState = false;
-	int32 NumContributors = 0;
+	FHierarchicalMetaData IgnoredHBiasEntry;
+	FHierarchicalMetaData ActiveBiasEntry;
 
-	for (auto ContributorIt = Contributors.CreateConstKeyIterator(Params.PropertyInfoIndex); ContributorIt; ++ContributorIt)
+	ActiveBiasEntry.HBias = TNumericLimits<int16>::Lowest();
+
+	// Iterate all contributors for this property to re-generate the meta-data
+	for (auto ContributorIt = Contributors.CreateConstKeyIterator(AnyContributor); ContributorIt; ++ContributorIt)
 	{
-		++NumContributors;
+		FMovieSceneEntityID Contributor = ContributorIt.Value();
+		const int16         ThisHBias   = ContributorIt.Key().HBias;
+		const bool          bIgnoreBias = ThisHBias == FContributorKey::ANY_HBIAS;
 
-		if (!bWantsRestoreState && Linker->EntityManager.HasComponent(ContributorIt.Value(), BuiltInComponents->Tags.RestoreState))
+		FHierarchicalMetaData* MetaDataToUpdate = nullptr;
+		if (bIgnoreBias)
 		{
-			bWantsRestoreState = true;
+			MetaDataToUpdate = &IgnoredHBiasEntry;
 		}
-
-		for (int32 CompositeIndex = 0; CompositeIndex < Params.PropertyDefinition->CompositeSize; ++CompositeIndex)
+		else if (ThisHBias >= ActiveBiasEntry.HBias)
 		{
-			if (EmptyChannels[CompositeIndex] == false)
+			MetaDataToUpdate = &ActiveBiasEntry;
+
+			if (ThisHBias > ActiveBiasEntry.HBias)
 			{
-				continue;
+				// We found a greater bias than any we've encountered this far
+				// Reset the empty channel list for the active hbias
+				ActiveBiasEmptyChannels = FChannelMask(true, Params.PropertyDefinition->CompositeSize);
+				ActiveBiasEntry.HBias = ThisHBias;
+
+				MetaDataToUpdate = &ActiveBiasEntry;
 			}
+		}
+
+		const bool bContributorIsActive = MetaDataToUpdate != nullptr;
+
+		// Update the various empty channel masks
+		for (int32 CompositeIndex = 0; CompositeIndex < Params.PropertyDefinition->CompositeSize; ++CompositeIndex)
+		{
+			const bool bCheckChannel =
+				BlendedBiasEmptyChannels[CompositeIndex] == true ||
+				(bIgnoreBias && IgnoredBiasEmptyChannels[CompositeIndex] == true) ||
+				(bContributorIsActive && ActiveBiasEmptyChannels[CompositeIndex] == true);
 
-			FComponentTypeID ThisChannel = Composites[CompositeIndex].ComponentTypeID;
-			if (ThisChannel && Linker->EntityManager.HasComponent(ContributorIt.Value(), ThisChannel))
+			if (bCheckChannel)
 			{
-				EmptyChannels[CompositeIndex] = false;
+				FComponentTypeID ThisChannel = Composites[CompositeIndex].ComponentTypeID;
+				if (ThisChannel && Linker->EntityManager.HasComponent(Contributor, ThisChannel))
+				{
+					BlendedBiasEmptyChannels[CompositeIndex] = false;
+					if (bIgnoreBias)
+					{
+						IgnoredBiasEmptyChannels[CompositeIndex] = false;
+					}
+					if (bContributorIsActive)
+					{
+						ActiveBiasEmptyChannels[CompositeIndex] = false;
+					}
+				}
 			}
 		}
-	}
 
-	// Reset the restore state status of the property if we still have contributors
-	// We do not do this if there are no contributors to ensure that stale properties are restored correctly
-	if (NumContributors > 0)
-	{
-		const bool bWasPartial = Params.PropertyInfo->EmptyChannels.Find(true) != INDEX_NONE;
-		const bool bIsPartial  = EmptyChannels.Find(true) != INDEX_NONE;
+		if (!bContributorIsActive)
+		{
+			continue;
+		}
+
+		MetaDataToUpdate->bInUse = true;
+		++MetaDataToUpdate->NumContributors;
 
-		if (bWasPartial != bIsPartial)
+		// Update whether this meta-data entry wants restore state
+		if (!MetaDataToUpdate->bWantsRestoreState && Linker->EntityManager.HasComponent(Contributor, BuiltInComponents->Tags.RestoreState))
 		{
-			const int32 StatIndex = Params.PropertyInfo->PropertyDefinitionIndex;
-			PropertyStats[StatIndex].NumPartialProperties += bIsPartial ? 1 : -1;
+			MetaDataToUpdate->bWantsRestoreState = true;
 		}
 
-		Params.PropertyInfo->EmptyChannels = EmptyChannels;
-		Params.PropertyInfo->bWantsRestoreState = bWantsRestoreState;
-	}
-}
+		// Update whether this meta-data entry needs an initial value or not
+		if (!MetaDataToUpdate->bNeedsInitialValue && Linker->EntityManager.HasComponent(Contributor, BuiltInComponents->Tags.AlwaysCacheInitialValue))
+		{
+			MetaDataToUpdate->bNeedsInitialValue = true;
+		}
 
-bool UMovieScenePropertyInstantiatorSystem::PropertySupportsFastPath(const FPropertyParameters& Params) const
-{
-	using namespace UE::MovieScene;
+		const FComponentMask& Type = Linker->EntityManager.GetEntityType(Contributor);
 
-	// Properties that are already blended, or are currently animated must use the blend path
-	if (ResolvedProperties[Params.PropertyInfoIndex].BlendChannel != INVALID_BLEND_CHANNEL || Params.PropertyInfo->PropertyEntityID.IsValid())
-	{
-		return false;
+		if (!MetaDataToUpdate->bBlendHierarchicalBias && Type.Contains(BuiltInComponents->Tags.BlendHierarchicalBias))
+		{
+			MetaDataToUpdate->bBlendHierarchicalBias = true;
+			MetaDataToUpdate->bSupportsFastPath = false;
+		}
+
+		// Update whether this property supports fast path
+		if (MetaDataToUpdate->bSupportsFastPath)
+		{
+			if (MetaDataToUpdate->NumContributors > 1)
+			{
+				MetaDataToUpdate->bSupportsFastPath = false;
+			}
+			else
+			{
+				if (Type.Contains(BuiltInComponents->Tags.RelativeBlend) || 
+						Type.Contains(BuiltInComponents->Tags.AdditiveBlend) || 
+						Type.Contains(BuiltInComponents->Tags.AdditiveFromBaseBlend) || 
+						Type.Contains(BuiltInComponents->WeightAndEasingResult))
+				{
+					MetaDataToUpdate->bSupportsFastPath = false;
+				}
+			}
+		}
 	}
 
-	int32 NumContributors = 0;
-	for (auto It = Contributors.CreateConstKeyIterator(Params.PropertyInfoIndex); It; ++It)
+	// Add any ignored hbias entries to the first tracked meta data that is in use
+	if (IgnoredHBiasEntry.bInUse)
 	{
-		++NumContributors;
-		if (NumContributors > 1)
+		// Combine using a bitwise & since channels are only empty if they are empty in both
+		ActiveBiasEmptyChannels.CombineWithBitwiseAND(IgnoredBiasEmptyChannels, EBitwiseOperatorFlags::MaintainSize);
+		if (ActiveBiasEntry.bInUse)
 		{
-			return false;
+			ActiveBiasEntry.CombineWith(IgnoredHBiasEntry);
 		}
-
-		FComponentMask Type = Linker->EntityManager.GetEntityType(It.Value());
-		if (Type.Contains(BuiltInComponents->Tags.RelativeBlend) || 
-				Type.Contains(BuiltInComponents->Tags.AdditiveBlend) || 
-				Type.Contains(BuiltInComponents->Tags.AdditiveFromBaseBlend) || 
-				Type.Contains(BuiltInComponents->WeightAndEasingResult))
+		else
 		{
-			return false;
+			ActiveBiasEntry = IgnoredHBiasEntry;
 		}
 	}
 
-	return true;
+	// Reset the restore state status of the property if we still have contributors
+	// We do not do this if there are no contributors to ensure that stale properties are restored correctly
+	Params.PropertyInfo->EmptyChannels = ActiveBiasEntry.bBlendHierarchicalBias
+		? BlendedBiasEmptyChannels
+		: ActiveBiasEmptyChannels;
+
+	const bool bWasPartial = Params.PropertyInfo->bIsPartiallyAnimated;
+	const bool bIsPartial  = Params.PropertyInfo->EmptyChannels.Find(true) != INDEX_NONE;
+
+	if (bWasPartial != bIsPartial)
+	{
+		const int32 StatIndex = Params.PropertyInfo->PropertyDefinitionIndex;
+		PropertyStats[StatIndex].NumPartialProperties += bIsPartial ? 1 : -1;
+	}
+
+	Params.PropertyInfo->bIsPartiallyAnimated = bIsPartial;
+	Params.PropertyInfo->bMaxHBiasHasChanged  = Params.PropertyInfo->HierarchicalMetaData.HBias != ActiveBiasEntry.HBias;
+	Params.PropertyInfo->HierarchicalMetaData = ActiveBiasEntry;
 }
 
 void UMovieScenePropertyInstantiatorSystem::InitializeFastPath(const FPropertyParameters& Params)
 {
 	using namespace UE::MovieScene;
 
-	FMovieSceneEntityID SoleContributor = Contributors.FindChecked(Params.PropertyInfoIndex);
-
-	// Have we ever seen this property before?
-	if (SoleContributor == Params.PropertyInfo->PropertyEntityID)
+	FTypelessMutation IgnoredContributorMutation;
+	IgnoredContributorMutation.RemoveMask = CleanFastPathMask;
+	IgnoredContributorMutation.RemoveMask.Set(BuiltInComponents->BlendChannelInput);
+	IgnoredContributorMutation.RemoveMask.Set(Params.PropertyDefinition->InitialValueType);
+	IgnoredContributorMutation.RemoveMask.Set(BuiltInComponents->Tags.HasAssignedInitialValue);
+	IgnoredContributorMutation.AddMask.Set(BuiltInComponents->Tags.Ignored);
+
+	// Find the sole contributor with the specific property info and hbias
+	const int16 ActiveHBias = Params.PropertyInfo->HierarchicalMetaData.HBias;
+	FContributorKey AnyContributor(Params.PropertyInfoIndex);
+	for (auto ContributorIt = Contributors.CreateConstKeyIterator(AnyContributor); ContributorIt; ++ContributorIt)
 	{
-		return;
-	}
+		FMovieSceneEntityID Contributor = ContributorIt.Value();
 
-	Params.PropertyInfo->PropertyEntityID = SoleContributor;
+		if (ContributorIt.Key().HBias == ActiveHBias)
+		{
+			FTypelessMutation SoleContributorMutation;
+			SoleContributorMutation.RemoveMask.SetAll({ BuiltInComponents->Tags.Ignored, BuiltInComponents->BlendChannelInput, BuiltInComponents->HierarchicalBlendTarget });
 
-	check(Params.PropertyInfo->BlendChannel == INVALID_BLEND_CHANNEL);
-	switch (Params.PropertyInfo->Property.GetIndex())
-	{
-	case 0:
-		Linker->EntityManager.AddComponent(SoleContributor, BuiltInComponents->FastPropertyOffset,  Params.PropertyInfo->Property.template Get<uint16>());
-		break;
-	case 1:
-		Linker->EntityManager.AddComponent(SoleContributor, BuiltInComponents->CustomPropertyIndex, Params.PropertyInfo->Property.template Get<FCustomPropertyIndex>());
-		break;
-	case 2:
-		Linker->EntityManager.AddComponent(SoleContributor, BuiltInComponents->SlowProperty,        Params.PropertyInfo->Property.template Get<FSlowPropertyPtr>());
-		break;
-	}
+			if (Params.PropertyInfo->HierarchicalMetaData.bNeedsInitialValue)
+			{
+				SoleContributorMutation.AddMask.Set(Params.PropertyDefinition->InitialValueType);
+			}
 
-	if (Linker->EntityManager.HasComponent(SoleContributor, BuiltInComponents->Tags.AlwaysCacheInitialValue))
-	{
-		Linker->EntityManager.AddComponent(SoleContributor, Params.PropertyDefinition->InitialValueType);
+			if (Params.PropertyDefinition->MetaDataTypes.Num() > 0)
+			{
+				InitializePropertyMetaDataTasks.PadToNum(Params.PropertyInfo->PropertyDefinitionIndex+1, false);
+				InitializePropertyMetaDataTasks[Params.PropertyInfo->PropertyDefinitionIndex] = true;
+
+				for (FComponentTypeID Component : Params.PropertyDefinition->MetaDataTypes)
+				{
+					SoleContributorMutation.AddMask.Set(Component);
+				}
+			}
+
+			// Ensure the sole contributor is set up to apply the property as a final output
+			switch (Params.PropertyInfo->Property.GetIndex())
+			{
+			case 0:
+				FEntityBuilder()
+				.Add(BuiltInComponents->FastPropertyOffset, Params.PropertyInfo->Property.template Get<uint16>())
+				.MutateExisting(&Linker->EntityManager, Contributor, SoleContributorMutation);
+				break;
+			case 1:
+				FEntityBuilder()
+				.Add(BuiltInComponents->CustomPropertyIndex, Params.PropertyInfo->Property.template Get<FCustomPropertyIndex>())
+				.MutateExisting(&Linker->EntityManager, Contributor, SoleContributorMutation);
+				break;
+			case 2:
+				FEntityBuilder()
+				.Add(BuiltInComponents->SlowProperty, Params.PropertyInfo->Property.template Get<FSlowPropertyPtr>())
+				.MutateExisting(&Linker->EntityManager, Contributor, SoleContributorMutation);
+				break;
+			}
+
+			// Copy initial values and meta-data back off our old blend output
+			FMovieSceneEntityID OldOutput = Params.PropertyInfo->PreviousFastPathID ? Params.PropertyInfo->PreviousFastPathID : Params.PropertyInfo->FinalBlendOutputID;
+			if (OldOutput)
+			{
+				FComponentMask CopyMask;
+				CopyMask.Set(Params.PropertyDefinition->InitialValueType);
+				CopyMask.Set(BuiltInComponents->Tags.HasAssignedInitialValue);
+				for (FComponentTypeID Component : Params.PropertyDefinition->MetaDataTypes)
+				{
+					CopyMask.Set(Component);
+				}
+
+				Linker->EntityManager.CopyComponents(OldOutput, Contributor, CopyMask);
+			}
+
+			Params.PropertyInfo->PreviousFastPathID = Contributor;
+		}
+		else
+		{
+			Linker->EntityManager.ChangeEntityType(Contributor, IgnoredContributorMutation.MutateType(Linker->EntityManager.GetEntityType(Contributor)));
+		}
 	}
 
-	if (Params.PropertyDefinition->MetaDataTypes.Num() > 0)
+	// If this was previously blended, destroy the blend output
+	if (Params.PropertyInfo->FinalBlendOutputID)
 	{
-		InitializePropertyMetaDataTasks.PadToNum(Params.PropertyInfo->PropertyDefinitionIndex+1, false);
-		InitializePropertyMetaDataTasks[Params.PropertyInfo->PropertyDefinitionIndex] = true;
-
-		FComponentMask NewMask;
-		for (FComponentTypeID Component : Params.PropertyDefinition->MetaDataTypes)
+		UMovieSceneBlenderSystem* Blender = Params.PropertyInfo->Blender.Get();
+		if (Blender && Params.PropertyInfo->BlendChannel != INVALID_BLEND_CHANNEL)
 		{
-			NewMask.Set(Component);
+			const FMovieSceneBlendChannelID BlendChannelID(Blender->GetBlenderSystemID(), Params.PropertyInfo->BlendChannel);
+			Blender->ReleaseBlendChannel(BlendChannelID);
+
+			Params.PropertyInfo->BlendChannel = INVALID_BLEND_CHANNEL;
+			Params.PropertyInfo->Blender = nullptr;
 		}
-		Linker->EntityManager.AddComponents(SoleContributor, NewMask);
+
+		Linker->EntityManager.AddComponent(Params.PropertyInfo->FinalBlendOutputID, BuiltInComponents->Tags.NeedsUnlink);
+		Params.PropertyInfo->FinalBlendOutputID = FMovieSceneEntityID();
 	}
 }
 
@@ -581,7 +748,8 @@ UMovieScenePropertyInstantiatorSystem::FSetupBlenderSystemResult UMovieSceneProp
 	int32 BlenderClassPriority = TNumericLimits<int32>::Lowest();
 
 	// Iterate all the contributors to locate the correct blender system by-priority
-	for (auto ContributorIt = Contributors.CreateConstKeyIterator(Params.PropertyInfoIndex); ContributorIt; ++ContributorIt)
+	FContributorKey ContributorKey = Params.MakeContributorKey();
+	for (auto ContributorIt = Contributors.CreateConstKeyIterator(ContributorKey); ContributorIt; ++ContributorIt)
 	{
 		FMovieSceneEntityID Contributor = ContributorIt.Value();
 		TOptionalComponentReader<TSubclassOf<UMovieSceneBlenderSystem>> BlenderTypeComponent = Linker->EntityManager.ReadComponent(Contributor, BuiltInComponents->BlenderType);
@@ -648,6 +816,8 @@ UMovieScenePropertyInstantiatorSystem::FSetupBlenderSystemResult UMovieSceneProp
 	return FSetupBlenderSystemResult{NewBlenderInfo, OldBlenderInfo};
 }
 
+
+
 void UMovieScenePropertyInstantiatorSystem::InitializeBlendPath(const FPropertyParameters& Params)
 {
 	using namespace UE::MovieScene;
@@ -658,7 +828,7 @@ void UMovieScenePropertyInstantiatorSystem::InitializeBlendPath(const FPropertyP
 
 	// -----------------------------------------------------------------------------------------------------
 	// Situation 1: New or modified contributors (inputs) but we're already set up for blending using the same system
-	if (SetupResult.CurrentInfo.BlenderSystemClass == SetupResult.PreviousInfo.BlenderSystemClass)
+	if (Params.PropertyInfo->FinalBlendOutputID && SetupResult.CurrentInfo.BlenderSystemClass == SetupResult.PreviousInfo.BlenderSystemClass)
 	{
 		UMovieSceneBlenderSystem* Blender = Params.PropertyInfo->Blender.Get();
 		check(Blender);
@@ -666,16 +836,23 @@ void UMovieScenePropertyInstantiatorSystem::InitializeBlendPath(const FPropertyP
 		// Ensure the output entity still matches the correct set of channels of all inputs
 		FComponentMask NewEntityType;
 		Params.MakeOutputComponentType(Linker->EntityManager, Composites, NewEntityType);
-		Linker->EntityManager.ChangeEntityType(Params.PropertyInfo->PropertyEntityID, NewEntityType);
+		Linker->EntityManager.ChangeEntityType(Params.PropertyInfo->FinalBlendOutputID, NewEntityType);
 
 		const FMovieSceneBlendChannelID BlendChannel(Blender->GetBlenderSystemID(), Params.PropertyInfo->BlendChannel);
 
 		// Change new contributors to include the blend input components
-		for (auto ContributorIt = NewContributors.CreateConstKeyIterator(Params.PropertyInfoIndex); ContributorIt; ++ContributorIt)
+		FContributorKey ContributorKey = Params.MakeContributorKey();
+
+		TMultiMap<FContributorKey, FMovieSceneEntityID>::TConstKeyIterator ContributorIt = Params.PropertyInfo->bMaxHBiasHasChanged
+			? Contributors.CreateConstKeyIterator(ContributorKey)
+			: NewContributors.CreateConstKeyIterator(ContributorKey);
+
+		for (; ContributorIt; ++ContributorIt)
 		{
 			FEntityBuilder()
 			.Add(BuiltInComponents->BlendChannelInput, BlendChannel)
 			.AddTag(SetupResult.CurrentInfo.BlenderTypeTag)
+			.AddConditional(BuiltInComponents->HierarchicalBlendTarget, Params.PropertyInfo->HierarchicalMetaData.HBias, Params.PropertyInfo->HierarchicalMetaData.bBlendHierarchicalBias)
 			.MutateExisting(&Linker->EntityManager, ContributorIt.Value());
 		}
 
@@ -690,19 +867,26 @@ void UMovieScenePropertyInstantiatorSystem::InitializeBlendPath(const FPropertyP
 
 	const FMovieSceneBlendChannelID NewBlendChannel = NewBlender->AllocateBlendChannel();
 
-	Params.PropertyInfo->Blender = NewBlender;
-	Params.PropertyInfo->BlendChannel = NewBlendChannel.ChannelID;
-
 	FTypelessMutation InputMutation;
+	// Clean any previously-fast-path entities
+	InputMutation.RemoveMask = CleanFastPathMask;
+	InputMutation.RemoveMask.Set(Params.PropertyDefinition->InitialValueType);
+	InputMutation.RemoveMask.Set(BuiltInComponents->Tags.HasAssignedInitialValue);
+	for (FComponentTypeID Component : Params.PropertyDefinition->MetaDataTypes)
+	{
+		InputMutation.RemoveMask.Set(Component);
+	}
 
 	// -----------------------------------------------------------------------------------------------------
 	// Situation 2.1: We're already set up for blending, but with a different blender system
 	if (OldBlender)
 	{
 		const FMovieSceneBlendChannelID OldBlendChannel(OldBlender->GetBlenderSystemID(), Params.PropertyInfo->BlendChannel);
-
 		OldBlender->ReleaseBlendChannel(OldBlendChannel);
 
+		Params.PropertyInfo->Blender = NewBlender;
+		Params.PropertyInfo->BlendChannel = NewBlendChannel.ChannelID;
+
 		// Change the output entity by adding the new blend channel and tag, while simultaneously
 		// updating the channels and restore state flags etc added by MakeOutputComponentType
 		{
@@ -717,7 +901,7 @@ void UMovieScenePropertyInstantiatorSystem::InitializeBlendPath(const FPropertyP
 			FEntityBuilder()
 			.Add(BuiltInComponents->BlendChannelOutput, NewBlendChannel)
 			.AddTag(SetupResult.CurrentInfo.BlenderTypeTag)
-			.MutateExisting(&Linker->EntityManager, Params.PropertyInfo->PropertyEntityID, OutputMutation);
+			.MutateExisting(&Linker->EntityManager, Params.PropertyInfo->FinalBlendOutputID, OutputMutation);
 		}
 
 		// Ensure that the old blend tag is removed from inputs
@@ -728,6 +912,9 @@ void UMovieScenePropertyInstantiatorSystem::InitializeBlendPath(const FPropertyP
 	// Situation 2.2: Never encountered blending before - need to create a new output entity to receive the blend result
 	else
 	{
+		Params.PropertyInfo->Blender = NewBlender;
+		Params.PropertyInfo->BlendChannel = NewBlendChannel.ChannelID;
+
 		FComponentMask NewMask;
 		NewMask.Set(Params.PropertyDefinition->InitialValueType);
 
@@ -758,8 +945,7 @@ void UMovieScenePropertyInstantiatorSystem::InitializeBlendPath(const FPropertyP
 		.Add(BuiltInComponents->PropertyBinding,         Params.PropertyInfo->PropertyBinding)
 		.Add(BuiltInComponents->BoundObject,             Params.PropertyInfo->BoundObject)
 		.Add(BuiltInComponents->BlendChannelOutput,      NewBlendChannel)
-		.AddTagConditional(BuiltInComponents->Tags.MigratedFromFastPath, Params.PropertyInfo->PropertyEntityID.IsValid())
-		.AddTagConditional(BuiltInComponents->Tags.RestoreState, Params.PropertyInfo->bWantsRestoreState)
+		.AddTagConditional(BuiltInComponents->Tags.RestoreState, Params.PropertyInfo->HierarchicalMetaData.bWantsRestoreState)
 		.AddTag(SetupResult.CurrentInfo.BlenderTypeTag)
 		.AddTag(BuiltInComponents->Tags.NeedsLink)
 		.AddMutualComponents();
@@ -786,29 +972,37 @@ void UMovieScenePropertyInstantiatorSystem::InitializeBlendPath(const FPropertyP
 			break;
 		}
 
-		// Params.PropertyInfo->PropertyEntityID will be set if the entity was previously using the fast path
-		if (Params.PropertyInfo->PropertyEntityID)
+		if (Params.PropertyInfo->PreviousFastPathID)
 		{
-			// Move any copiable/migratable components over from the existing fast-path entity
-			Linker->EntityManager.CopyComponents(Params.PropertyInfo->PropertyEntityID, NewOutputEntityID, Linker->EntityManager.GetComponents()->GetCopyAndMigrationMask());
+			// If this contributor has the initial values on it, we copy its initial values and meta-data components
+			FComponentMask CopyMask = Linker->EntityManager.GetComponents()->GetCopyAndMigrationMask();
+
+			CopyMask.Set(Params.PropertyDefinition->InitialValueType);
+			CopyMask.Set(BuiltInComponents->Tags.HasAssignedInitialValue);
+			for (FComponentTypeID Component : Params.PropertyDefinition->MetaDataTypes)
+			{
+				CopyMask.Set(Component);
+			}
 
-			// Clean the previously-fast-path entity
-			InputMutation.RemoveMask = CleanFastPathMask;
+			Linker->EntityManager.CopyComponents(Params.PropertyInfo->PreviousFastPathID, NewOutputEntityID, CopyMask);
 		}
 
 		// The property entity ID is now the blend output entity
-		Params.PropertyInfo->PropertyEntityID = NewOutputEntityID;
+		Params.PropertyInfo->FinalBlendOutputID = NewOutputEntityID;
+		Params.PropertyInfo->PreviousFastPathID = FMovieSceneEntityID();
 	}
 
 	// Change *all* contributors (not just new ones because the old ones will have the old blender's channel and tag on them)
 	// to include the new blend channel input, and remove the old blender type. No need to remove the clean fast path mask because
 	// that will have already happened as part of the 'completely new blending' branch below
-	for (auto ContributorIt = Contributors.CreateConstKeyIterator(Params.PropertyInfoIndex); ContributorIt; ++ContributorIt)
+	FContributorKey ContributorKey(Params.PropertyInfoIndex);
+	for (auto ContributorIt = Contributors.CreateConstKeyIterator(ContributorKey); ContributorIt; ++ContributorIt)
 	{
-		const FMovieSceneEntityID Contributor = ContributorIt.Value();
+		FMovieSceneEntityID Contributor = ContributorIt.Value();
 
 		FEntityBuilder()
 		.Add(BuiltInComponents->BlendChannelInput, NewBlendChannel)
+		.AddConditional(BuiltInComponents->HierarchicalBlendTarget, Params.PropertyInfo->HierarchicalMetaData.HBias, Params.PropertyInfo->HierarchicalMetaData.bBlendHierarchicalBias)
 		.AddTag(SetupResult.CurrentInfo.BlenderTypeTag)
 		.MutateExisting(&Linker->EntityManager, Contributor, InputMutation);
 	}
@@ -850,17 +1044,20 @@ UE::MovieScene::FPropertyRecomposerPropertyInfo UMovieScenePropertyInstantiatorS
 {
 	using namespace UE::MovieScene;
 
+	TOptionalComponentReader<int16>                      HBiasComponent  = Linker->EntityManager.ReadComponent(EntityID, BuiltInComponents->HierarchicalBias);
 	TOptionalComponentReader<FMovieScenePropertyBinding> PropertyBinding = Linker->EntityManager.ReadComponent(EntityID, BuiltInComponents->PropertyBinding);
 	if (!PropertyBinding)
 	{
 		return FPropertyRecomposerPropertyInfo::Invalid();
 	}
 
+	const int16 HBias = HBiasComponent ? *HBiasComponent : 0;
 	TTuple<UObject*, FName> Key = MakeTuple(Object, PropertyBinding->PropertyPath);
 	if (const int32* PropertyIndex = ObjectPropertyToResolvedIndex.Find(Key))
 	{
+		const uint16 BlendChannel = ResolvedProperties[*PropertyIndex].BlendChannel;
 		const FObjectPropertyInfo& PropertyInfo = ResolvedProperties[*PropertyIndex];
-		return FPropertyRecomposerPropertyInfo { PropertyInfo.BlendChannel, PropertyInfo.Blender.Get(), PropertyInfo.PropertyEntityID };
+		return FPropertyRecomposerPropertyInfo { BlendChannel, PropertyInfo.Blender.Get(), PropertyInfo.FinalBlendOutputID };
 	}
 
 	return FPropertyRecomposerPropertyInfo::Invalid();
@@ -891,18 +1088,28 @@ void UMovieScenePropertyInstantiatorSystem::FPropertyParameters::MakeOutputCompo
 	using namespace UE::MovieScene;
 
 	// Get the existing type
-	OutComponentType = EntityManager.GetEntityType(PropertyInfo->PropertyEntityID);
+	if (PropertyInfo->FinalBlendOutputID)
+	{
+		OutComponentType = EntityManager.GetEntityType(PropertyInfo->FinalBlendOutputID);
+	}
 
 	// Ensure the property has only the exact combination of channels that constitute its animation
 	for (int32 Index = 0; Index < Composites.Num(); ++Index)
 	{
 		FComponentTypeID Composite = Composites[Index].ComponentTypeID;
-		OutComponentType[Composite] = (PropertyInfo->EmptyChannels[Index] != true);
+		if (PropertyInfo->EmptyChannels[Index] != true)
+		{
+			OutComponentType.Set(Composite);
+		}
+		else
+		{
+			OutComponentType.Remove(Composite);
+		}
 	}
 	OutComponentType.Set(PropertyDefinition->PropertyType);
 
 	// Set the restore state tag appropriately
-	if (PropertyInfo->bWantsRestoreState)
+	if (PropertyInfo->HierarchicalMetaData.bWantsRestoreState)
 	{
 		OutComponentType.Set(FBuiltInComponentTypes::Get()->Tags.RestoreState);
 	}
diff --git a/Engine/Source/Runtime/MovieSceneTracks/Private/Systems/WeightAndEasingEvaluatorSystem.cpp b/Engine/Source/Runtime/MovieSceneTracks/Private/Systems/WeightAndEasingEvaluatorSystem.cpp
--- a/Engine/Source/Runtime/MovieSceneTracks/Private/Systems/WeightAndEasingEvaluatorSystem.cpp
+++ b/Engine/Source/Runtime/MovieSceneTracks/Private/Systems/WeightAndEasingEvaluatorSystem.cpp
@@ -34,12 +34,66 @@ namespace UE::MovieScene
 
 static constexpr uint16 INVALID_EASING_CHANNEL = uint16(-1);
 
+float Factorial(int32 In)
+{
+	check(In >= 0);
+
+	int32 Result = 1;
+	for (; In > 0; --In)
+	{
+		Result *= In;
+	}
+	return static_cast<float>(Result);
+}
+
 enum class EHierarchicalSequenceChannelFlags
 {
 	None = 0,
 	SharedWithParent = 1,
+	BlendTarget = 2,
 };
-ENUM_CLASS_FLAGS(EHierarchicalSequenceChannelFlags)
+ENUM_CLASS_FLAGS(EHierarchicalSequenceChannelFlags);
+
+bool FHierarchicalEasingChannelBuffer::IsEmpty() const
+{
+	return Channels.Num() == 0;
+}
+
+uint16 FHierarchicalEasingChannelBuffer::AddBaseChannel(const FHierarchicalEasingChannelData& InChannelData)
+{
+	checkf(BlendTargetStartIndex == INDEX_NONE, TEXT("Cannot add base channels once blend targets have been added"));
+
+	const uint16 ChannelID = static_cast<uint16>(Channels.Num());
+	Channels.Emplace(InChannelData);
+	return ChannelID;
+}
+
+uint16 FHierarchicalEasingChannelBuffer::AddBlendTargetChannel(const FHierarchicalEasingChannelData& InChannelData)
+{
+	if (BlendTargetStartIndex == INDEX_NONE)
+	{
+		BlendTargetStartIndex = Channels.Num();
+	}
+
+	const uint16 ChannelID = static_cast<uint16>(Channels.Num());
+	Channels.Emplace(InChannelData);
+	return ChannelID;
+}
+
+void FHierarchicalEasingChannelBuffer::Reset()
+{
+	Channels.Empty();
+	BlendTargetStartIndex = INDEX_NONE;
+}
+
+void FHierarchicalEasingChannelBuffer::ResetBlendTargets()
+{
+	if (BlendTargetStartIndex != INDEX_NONE)
+	{
+		Channels.RemoveAtSwap(BlendTargetStartIndex, Channels.Num() - BlendTargetStartIndex);
+		BlendTargetStartIndex = INDEX_NONE;
+	}
+}
 
 /**
  * Class that accumulates active easing data for a particular root sequence hierarchy
@@ -78,7 +132,7 @@ struct FRootSequenceData
 	int32 Add(FMovieSceneSequenceID SequenceID, EHierarchicalSequenceChannelFlags InFlags)
 	{
 		const int32 Index = Algo::LowerBoundBy(SequenceIDs, SequenceID, &FSequenceIDAndChannel::SequenceID);
-		check(!SequenceIDs.IsValidIndex(Index) ||SequenceIDs[Index].SequenceID != SequenceID);
+		check(!SequenceIDs.IsValidIndex(Index) || SequenceIDs[Index].SequenceID != SequenceID);
 
 		SequenceIDs.Insert(FSequenceIDAndChannel{ SequenceID, INVALID_EASING_CHANNEL, InFlags }, Index);
 
@@ -102,9 +156,8 @@ struct FRootSequenceData
 		if (ensure(RootHierarchy))
 		{
 			const FMovieSceneSequenceHierarchyNode* Node = RootHierarchy->FindNode(SequenceID);
-			check(Node);
 
-			if (Node->ParentID != MovieSceneSequenceID::Invalid && AddImplicitChild(Node->ParentID))
+			if (ensure(Node) && Node->ParentID != MovieSceneSequenceID::Invalid && AddImplicitChild(Node->ParentID))
 			{
 				Add(SequenceID, EHierarchicalSequenceChannelFlags::SharedWithParent);
 				return true;
@@ -116,9 +169,10 @@ struct FRootSequenceData
 	/**
 	 * Produce channel IDs in the specified output data, ordered parent-first
 	 */
-	void ProduceChannels(TSparseArray<uint16>& OutInstanceIDToChannel, TArray<FHierarchicalEasingChannelData>& OutData)
+	void ProduceChannels(TSparseArray<uint16>& OutInstanceIDToChannel, FHierarchicalEasingChannelBuffer& OutData)
 	{
 		const FSequenceInstance& RootInstance = InstanceRegistry->GetInstance(RootInstanceHandle);
+
 		for (int32 Index = 0; Index < SequenceIDs.Num(); ++Index)
 		{
 			ProduceChannelForIndex(Index, RootInstance, OutInstanceIDToChannel, OutData);
@@ -128,7 +182,7 @@ struct FRootSequenceData
 	/**
 	 * Produce channels for the specified index within our map, including any of its parents
 	 */
-	uint16 ProduceChannelForIndex(int32 Index, const FSequenceInstance& RootInstance, TSparseArray<uint16>& OutInstanceIDToChannel, TArray<FHierarchicalEasingChannelData>& OutData)
+	uint16 ProduceChannelForIndex(int32 Index, const FSequenceInstance& RootInstance, TSparseArray<uint16>& OutInstanceIDToChannel, FHierarchicalEasingChannelBuffer& OutData)
 	{
 		const uint16 ExistingChannel = SequenceIDs[Index].ChannelID;
 		if (ExistingChannel != INVALID_EASING_CHANNEL)
@@ -141,8 +195,9 @@ struct FRootSequenceData
 		if (ChannelData.SequenceID == MovieSceneSequenceID::Root)
 		{
 			// Add a new channel with no parent
-			const uint16 ChannelID = static_cast<uint16>(OutData.Num());
-			OutData.Add(FHierarchicalEasingChannelData{ INVALID_EASING_CHANNEL });
+			FHierarchicalEasingChannelData NewData;
+
+			const uint16 ChannelID = OutData.AddBaseChannel(NewData);
 			OutInstanceIDToChannel.Insert(RootInstanceHandle.InstanceID, ChannelID);
 
 			SequenceIDs[Index].ChannelID = ChannelID;
@@ -168,7 +223,6 @@ struct FRootSequenceData
 		}
 
 		uint16 ParentChannel = INVALID_EASING_CHANNEL;
-		FMovieSceneSequenceID Parent = MovieSceneSequenceID::Invalid;
 		while (Node)
 		{
 			const int32 ParentIndex = IndexOf(Node->ParentID);
@@ -196,8 +250,11 @@ struct FRootSequenceData
 		}
 		else
 		{
-			const uint16 ChannelID = static_cast<uint16>(OutData.Num());
-			OutData.Add(FHierarchicalEasingChannelData{ ParentChannel });
+			FHierarchicalEasingChannelData NewData;
+			NewData.ParentChannel = ParentChannel;
+			NewData.HBias = SubData->HierarchicalBias;
+
+			const uint16 ChannelID = OutData.AddBaseChannel(NewData);
 			OutInstanceIDToChannel.Insert(SubInstanceHandle.InstanceID, ChannelID);
 			SequenceIDs[Index].ChannelID = ChannelID;
 			return ChannelID;
@@ -288,9 +345,9 @@ struct FProviderEasingChannelMutation : FEasingChannelMutationBase
 
 		TComponentReader<FRootInstanceHandle>   RootInstanceHandles   = Allocation->ReadComponents(BuiltInComponents->RootInstanceHandle);
 		TComponentReader<FMovieSceneSequenceID> EasingProviders       = Allocation->ReadComponents(BuiltInComponents->HierarchicalEasingProvider);
-		TOptionalComponentWriter<uint16>        ExistingEasingHandles = Allocation->TryWriteComponents(BuiltInComponents->HierarchicalEasingChannel, FEntityAllocationWriteContext::NewAllocation());
+		TOptionalComponentWriter<uint16>        ExistingEasingChannels = Allocation->TryWriteComponents(BuiltInComponents->HierarchicalEasingChannel, FEntityAllocationWriteContext::NewAllocation());
 
-		const bool bHasExistingEasing = ExistingEasingHandles.IsValid();
+		const bool bHasExistingEasing = ExistingEasingChannels.IsValid();
 
 		// Loop through each entity and check whether it has easing.
 		// If so, either assign the new easing channel or mark it for mutation.
@@ -304,7 +361,7 @@ struct FProviderEasingChannelMutation : FEasingChannelMutationBase
 			{
 				if (bHasExistingEasing)
 				{
-					ExistingEasingHandles[Index] = ChannelID;
+					ExistingEasingChannels[Index] = ChannelID;
 				}
 				else
 				{
@@ -348,9 +405,10 @@ private:
  */
 struct FConsumerEasingChannelMutation : FEasingChannelMutationBase
 {
-	explicit FConsumerEasingChannelMutation(UMovieSceneEntitySystemLinker* Linker, TSparseArray<uint16>* InInstanceIDToChannel)
+	explicit FConsumerEasingChannelMutation(UMovieSceneEntitySystemLinker* Linker, TSparseArray<uint16>* InInstanceIDToChannel, TMap<TTuple<int16, int16>, uint16>* InCachedHierarchicalBlendTargetChannels)
 		: FEasingChannelMutationBase(Linker)
 		, InstanceIDToChannel(InInstanceIDToChannel)
+		, CachedHierarchicalBlendTargetChannels(InCachedHierarchicalBlendTargetChannels)
 	{}
 
 	/**
@@ -361,32 +419,66 @@ struct FConsumerEasingChannelMutation : FEasingChannelMutationBase
 	{
 		const FMovieSceneEntityID* EntityIDs = Allocation->GetRawEntityIDs();
 
-		TComponentReader<FInstanceHandle> InstanceHandles       = Allocation->ReadComponents(BuiltInComponents->InstanceHandle);
-		TOptionalComponentWriter<uint16>  ExistingEasingHandles = Allocation->TryWriteComponents(BuiltInComponents->HierarchicalEasingChannel, FEntityAllocationWriteContext::NewAllocation());
+		TComponentReader<FInstanceHandle> InstanceHandles        = Allocation->ReadComponents(BuiltInComponents->InstanceHandle);
+		TOptionalComponentWriter<uint16>  ExistingEasingChannels = Allocation->TryWriteComponents(BuiltInComponents->HierarchicalEasingChannel, FEntityAllocationWriteContext::NewAllocation());
+		TOptionalComponentReader<int16>   BlendTargets           = Allocation->TryReadComponents(BuiltInComponents->HierarchicalBlendTarget);
 
-		const bool bHasExistingEasing = ExistingEasingHandles.IsValid();
+		const bool bHasExistingEasing = ExistingEasingChannels.IsValid();
 
 		// Loop through each entity and check whether it has easing.
 		// If so, either assign the new easing channel or mark it for mutation.
 		// If not, mark it for removal
-		for (int32 Index = 0; Index < Allocation->Num(); ++Index)
+		if (BlendTargets)
 		{
-			if (InstanceIDToChannel->IsValidIndex(InstanceHandles[Index].InstanceID))
+			check(CachedHierarchicalBlendTargetChannels);
+			TOptionalComponentReader<int16> HierarchicalBias = Allocation->TryReadComponents(BuiltInComponents->HierarchicalBias);
+
+			for (int32 Index = 0; Index < Allocation->Num(); ++Index)
 			{
-				if (bHasExistingEasing)
+				TTuple<int16, int16> BlendTargetKey = MakeTuple(
+					HierarchicalBias ? HierarchicalBias[Index] : 0, BlendTargets[Index]
+					);
+
+				if (const uint16* BlendTargetChannel = CachedHierarchicalBlendTargetChannels->Find(BlendTargetKey))
 				{
-					ExistingEasingHandles[Index] = (*InstanceIDToChannel)[InstanceHandles[Index].InstanceID];
+					if (bHasExistingEasing)
+					{
+						ExistingEasingChannels[Index] = *BlendTargetChannel;
+					}
+					else
+					{
+						// Mark this entity for mutation
+						OutEntitiesToMutate.PadToNum(Index + 1, false);
+						OutEntitiesToMutate[Index] = true;
+					}
 				}
-				else
+				else if (bHasExistingEasing)
 				{
-					// Mark this entity for mutation
-					OutEntitiesToMutate.PadToNum(Index + 1, false);
-					OutEntitiesToMutate[Index] = true;
+					StaleEntities.Add(EntityIDs[Index]);
 				}
 			}
-			else if (bHasExistingEasing)
+		}
+		else
+		{
+			for (int32 Index = 0; Index < Allocation->Num(); ++Index)
 			{
-				StaleEntities.Add(EntityIDs[Index]);
+				if (InstanceIDToChannel->IsValidIndex(InstanceHandles[Index].InstanceID))
+				{
+					if (bHasExistingEasing)
+					{
+						ExistingEasingChannels[Index] = (*InstanceIDToChannel)[InstanceHandles[Index].InstanceID];
+					}
+					else
+					{
+						// Mark this entity for mutation
+						OutEntitiesToMutate.PadToNum(Index + 1, false);
+						OutEntitiesToMutate[Index] = true;
+					}
+				}
+				else if (bHasExistingEasing)
+				{
+					StaleEntities.Add(EntityIDs[Index]);
+				}
 			}
 		}
 	}
@@ -398,19 +490,40 @@ struct FConsumerEasingChannelMutation : FEasingChannelMutationBase
 	{
 		FEntityAllocationWriteContext NewAllocation = FEntityAllocationWriteContext::NewAllocation();
 
-		TComponentWriter<uint16>              EasingChannels      = EntityRange.Allocation->WriteComponents(BuiltInComponents->HierarchicalEasingChannel, NewAllocation);
-		TComponentReader<FInstanceHandle>     InstanceHandles     = EntityRange.Allocation->ReadComponents(BuiltInComponents->InstanceHandle);
+		TComponentWriter<uint16>          EasingChannels      = EntityRange.Allocation->WriteComponents(BuiltInComponents->HierarchicalEasingChannel, NewAllocation);
+		TComponentReader<FInstanceHandle> InstanceHandles     = EntityRange.Allocation->ReadComponents(BuiltInComponents->InstanceHandle);
+		TOptionalComponentReader<int16>   BlendTargets        = EntityRange.Allocation->TryReadComponents(BuiltInComponents->HierarchicalBlendTarget);
 
-		for (int32 Index = 0; Index < EntityRange.Num; ++Index)
+		if (BlendTargets)
 		{
-			const int32 Offset = EntityRange.ComponentStartOffset + Index;
-			EasingChannels[Offset] = (*InstanceIDToChannel)[InstanceHandles[Offset].InstanceID];
+			check(CachedHierarchicalBlendTargetChannels);
+
+			TOptionalComponentReader<int16> HierarchicalBias = EntityRange.Allocation->TryReadComponents(BuiltInComponents->HierarchicalBias);
+
+			for (int32 Index = 0; Index < EntityRange.Num; ++Index)
+			{
+				const int32 Offset = EntityRange.ComponentStartOffset + Index;
+
+				TTuple<int16, int16> BlendTargetKey
+					= MakeTuple(HierarchicalBias ? HierarchicalBias[Index] : 0, BlendTargets[Index]);
+
+				EasingChannels[Offset] = CachedHierarchicalBlendTargetChannels->FindChecked(BlendTargetKey);
+			}
+		}
+		else
+		{
+			for (int32 Index = 0; Index < EntityRange.Num; ++Index)
+			{
+				const int32 Offset = EntityRange.ComponentStartOffset + Index;
+				EasingChannels[Offset] = (*InstanceIDToChannel)[InstanceHandles[Offset].InstanceID];
+			}
 		}
 	}
 
 private:
 
 	TSparseArray<uint16>* InstanceIDToChannel;
+	TMap<TTuple<int16, int16>, uint16>* CachedHierarchicalBlendTargetChannels;
 };
 
 /**
@@ -486,6 +599,7 @@ struct FHarvestHierarchicalEasings
 		for (FHierarchicalEasingChannelData& Data : ComputationData)
 		{
 			Data.FinalResult = 1.0;
+			Data.UnaccumulatedResult = 1.0;
 		}
 	}
 
@@ -493,6 +607,7 @@ struct FHarvestHierarchicalEasings
 	void ForEachEntity(double Result, uint16 EasingChannel)
 	{
 		ComputationData[EasingChannel].FinalResult *= Result;
+		ComputationData[EasingChannel].UnaccumulatedResult *= Result;
 	}
 
 	// Multiply hierarchical weights with sub sequences
@@ -502,11 +617,55 @@ struct FHarvestHierarchicalEasings
 		// This is possible because the results array is already sorted by depth
 		for (int32 Index = 0; Index < ComputationData.Num(); ++Index)
 		{
-			FHierarchicalEasingChannelData ChannelData = ComputationData[Index];
-			if (ChannelData.ParentChannel != uint16(-1))
+			FHierarchicalEasingChannelData& This = ComputationData[Index];
+			switch (This.BlendMode)
+			{
+#if UE_ENABLE_MOVIESCENE_BEZIER_BLENDING
+			case EHierarchicalBlendMode::BezierSeries:
 			{
-				// The parent result has already been multiplied by all its parent weights by this point
-				ComputationData[Index].FinalResult *= ComputationData[ChannelData.ParentChannel].FinalResult;
+				checkSlow(This.ParentChannel != INVALID_EASING_CHANNEL);
+
+				const float T         = ComputationData[This.ParentChannel].UnaccumulatedResult;
+				const float OneMinusT = 1.f - T;
+
+				// Of the form a(1-t)^b + t^c
+				This.FinalResult = This.CoeffA * FMath::Pow(OneMinusT, This.ExpB) * FMath::Pow(T, This.ExpC);
+
+				// Result channel is not used for bezier blends
+				break;
+			}
+#endif
+			case EHierarchicalBlendMode::ChildFirstBlendTarget:
+			{
+				checkSlow(This.ParentChannel != INVALID_EASING_CHANNEL);
+
+				const float T          = ComputationData[This.ParentChannel].UnaccumulatedResult;
+				const float RemainingT = This.FinalResult;
+				const float Result     = RemainingT*T;
+
+				This.FinalResult = FMath::Clamp(Result, 0.f, 1.f);
+				if (This.ResultChannel != INVALID_EASING_CHANNEL)
+				{
+					// Pass on the remaining T to the next channel
+					ComputationData[This.ResultChannel].FinalResult = FMath::Clamp(RemainingT-Result, 0.f, 1.f);
+				}
+				break;
+			}
+			case EHierarchicalBlendMode::AccumulateParentToChild:
+				if (This.ParentChannel != INVALID_EASING_CHANNEL)
+				{
+					FHierarchicalEasingChannelData Parent = ComputationData[This.ParentChannel];
+
+					// The parent result has already been multiplied by all its parent weights by this point
+					This.FinalResult *= Parent.FinalResult;
+				}
+
+				if (This.ResultChannel != INVALID_EASING_CHANNEL)
+				{
+					// Forward onto the result channel
+					ComputationData[This.ResultChannel].UnaccumulatedResult *= This.FinalResult;
+				}
+				break;
 			}
 		}
 	}
@@ -529,7 +688,7 @@ struct FPropagateHierarchicalEasings
 		for (int32 Index = 0; Index < Num; ++Index)
 		{
 			const uint16 HierarchicalEasingChannel = HierarchicalEasingChannels[Index];
-			if (ensure(ComputationData.IsValidIndex(HierarchicalEasingChannel)))
+			if (ensureAlways(ComputationData.IsValidIndex(HierarchicalEasingChannel)))
 			{
 				WeightAndEasingResults[Index] *= ComputationData[HierarchicalEasingChannel].FinalResult;
 			}
@@ -537,7 +696,6 @@ struct FPropagateHierarchicalEasings
 	}
 };
 
-
 } // namespace UE::MovieScene
 
 
@@ -601,7 +759,8 @@ void UMovieSceneHierarchicalEasingInstantiatorSystem::OnRun(FSystemTaskPrerequis
 	ExpiringSubInstanceFilter.All({ BuiltInComponents->Tags.SubInstance, BuiltInComponents->HierarchicalEasingChannel, BuiltInComponents->Tags.NeedsUnlink });
 
 	// Check if we have any new or expiring easing providers
-	if (Linker->EntityManager.Contains(ChangedProviderFilter) || Linker->EntityManager.Contains(NewSubInstanceFilter) || Linker->EntityManager.Contains(ExpiringSubInstanceFilter))
+	bChannelsHaveBeenInvalidated = Linker->EntityManager.Contains(ChangedProviderFilter) || Linker->EntityManager.Contains(NewSubInstanceFilter) || Linker->EntityManager.Contains(ExpiringSubInstanceFilter);
+	if (bChannelsHaveBeenInvalidated)
 	{
 		FRootSequenceDataMap SequenceDataMap;
 
@@ -656,17 +815,17 @@ void UMovieSceneHierarchicalEasingInstantiatorSystem::OnRun(FSystemTaskPrerequis
 		.FilterNone({ BuiltInComponents->Tags.NeedsUnlink })
 		.Iterate_PerEntity(&Linker->EntityManager, VisitSubInstance);
 
-		// Step 3: Produce parent-first channel IDs for anything we collected, then pass on the buffer to the evaluator system
+		// Step 3: Produce parent-first channel IDs for anything we collected
 		//
 		{
-			TArray<FHierarchicalEasingChannelData> AllChannels;
+			FHierarchicalEasingChannelBuffer& Buffer = EvaluatorSystem->GetComputationBuffer();
+			Buffer.Reset();
 			CachedInstanceIDToChannel.Empty();
 
 			for (TPair<FRootInstanceHandle, FRootSequenceData>& Pair : SequenceDataMap)
 			{
-				Pair.Value.ProduceChannels(CachedInstanceIDToChannel, AllChannels);
+				Pair.Value.ProduceChannels(CachedInstanceIDToChannel, Buffer);
 			}
-			EvaluatorSystem->SetComputationBuffer(MoveTemp(AllChannels));
 		}
 
 		// Step 4: Assign channel IDs to easing providers, potentially adding or removing invalid ones
@@ -685,12 +844,12 @@ void UMovieSceneHierarchicalEasingInstantiatorSystem::OnRun(FSystemTaskPrerequis
 		// Step 5: Assign easing channels for anything else that is left and remove stale ones
 		//
 		{
-			FEntityComponentFilter NewFilter;
-			NewFilter.All({ BuiltInComponents->InstanceHandle });
-			NewFilter.None({ BuiltInComponents->Tags.NeedsUnlink, BuiltInComponents->HierarchicalEasingProvider });
+			FEntityComponentFilter AssignEasingChannelsFilter;
+			AssignEasingChannelsFilter.All({ BuiltInComponents->InstanceHandle});
+			AssignEasingChannelsFilter.None({ BuiltInComponents->Tags.NeedsUnlink, BuiltInComponents->HierarchicalEasingProvider, BuiltInComponents->HierarchicalBlendTarget, BuiltInComponents->Tags.ImportedEntity });
 
-			FConsumerEasingChannelMutation ConsumerMutation(Linker, &CachedInstanceIDToChannel);
-			Linker->EntityManager.MutateConditional(NewFilter, ConsumerMutation, EMutuallyInclusiveComponentType::All);
+			FConsumerEasingChannelMutation ConsumerMutation(Linker, &CachedInstanceIDToChannel, nullptr);
+			Linker->EntityManager.MutateConditional(AssignEasingChannelsFilter, ConsumerMutation, EMutuallyInclusiveComponentType::All);
 
 			ConsumerMutation.RemoveStaleComponents(Linker);
 		}
@@ -698,18 +857,201 @@ void UMovieSceneHierarchicalEasingInstantiatorSystem::OnRun(FSystemTaskPrerequis
 	else
 	{
 		// Just add easing channels to NeedsLink entities
-		FEntityComponentFilter Filter;
-		Filter.All({ BuiltInComponents->InstanceHandle, BuiltInComponents->Tags.NeedsLink });
-		Filter.None({ BuiltInComponents->HierarchicalEasingProvider });
+		FEntityComponentFilter AssignEasingChannelsFilter;
+		AssignEasingChannelsFilter.All({ BuiltInComponents->InstanceHandle, BuiltInComponents->Tags.NeedsLink  });
+		AssignEasingChannelsFilter.None({ BuiltInComponents->Tags.NeedsUnlink, BuiltInComponents->HierarchicalEasingProvider, BuiltInComponents->HierarchicalBlendTarget, BuiltInComponents->Tags.ImportedEntity });
 
-		FConsumerEasingChannelMutation ConsumerMutation(Linker, &CachedInstanceIDToChannel);
-		Linker->EntityManager.MutateConditional(Filter, ConsumerMutation, EMutuallyInclusiveComponentType::All);
+		FConsumerEasingChannelMutation ConsumerMutation(Linker, &CachedInstanceIDToChannel, nullptr);
+		Linker->EntityManager.MutateConditional(AssignEasingChannelsFilter, ConsumerMutation, EMutuallyInclusiveComponentType::All);
 
-		ensure(!ConsumerMutation.HasStaleComponents());
 		ConsumerMutation.RemoveStaleComponents(Linker);
 	}
 }
 
+void UMovieSceneHierarchicalEasingInstantiatorSystem::FinalizeBlendTargets()
+{
+	using namespace UE::MovieScene;
+
+	const FBuiltInComponentTypes* const BuiltInComponents = FBuiltInComponentTypes::Get();
+
+	FEntityComponentFilter NewBlendTargetFilter;
+	NewBlendTargetFilter.All({ BuiltInComponents->HierarchicalBlendTarget, BuiltInComponents->Tags.NeedsLink });
+
+	if (bChannelsHaveBeenInvalidated || Linker->EntityManager.Contains(NewBlendTargetFilter))
+	{
+		FHierarchicalEasingChannelBuffer& Buffer = EvaluatorSystem->GetComputationBuffer();
+		Buffer.ResetBlendTargets();
+
+		const int32 NumStartingChannels = Buffer.Channels.Num();
+
+		// Map from HBias to whether it is a blend target or not
+		struct FHBiasChannelData
+		{
+			uint16 AccumulatorChannelID = UE::MovieScene::INVALID_EASING_CHANNEL;
+			bool bIsBlendTarget = false;
+		};
+		TSortedMap<int16, FHBiasChannelData, TInlineAllocator<16>, TGreater<>> HBiasToChannelData;
+
+		// Step 1: Gather all the hierarchical blend targets
+		//
+		{
+			auto GatherHierarchicalBlendTargets = [&HBiasToChannelData](const FEntityAllocation* Allocation, const int16* BlendTargets, const int16* HBiases)
+			{
+				const int32 Num = Allocation->Num();
+				if (HBiases)
+				{
+					for (int32 Index = 0; Index < Num; ++Index)
+					{
+						HBiasToChannelData.FindOrAdd(BlendTargets[Index]).bIsBlendTarget = true;
+						HBiasToChannelData.FindOrAdd(HBiases[Index]);
+					}
+				}
+				else
+				{
+					// No HBias components means 0, so add that to the map once
+					HBiasToChannelData.FindOrAdd(0);
+					for (int32 Index = 0; Index < Num; ++Index)
+					{
+						HBiasToChannelData.FindOrAdd(BlendTargets[Index]).bIsBlendTarget = true;
+					}
+				}
+			};
+
+			auto Task = FEntityTaskBuilder()
+			.Read(BuiltInComponents->HierarchicalBlendTarget)
+			.ReadOptional(BuiltInComponents->HierarchicalBias)
+			.FilterNone({ BuiltInComponents->Tags.NeedsUnlink });
+			if (!bChannelsHaveBeenInvalidated)
+			{
+				Task.FilterAll({ BuiltInComponents->Tags.NeedsLink });
+			}
+			Task.Iterate_PerAllocation(&Linker->EntityManager, GatherHierarchicalBlendTargets);
+		}
+
+		// Step 2: Create channels for each blend target sorted by bias
+		//
+		{
+			for (TPair<int16, FHBiasChannelData>& Pair : HBiasToChannelData)
+			{
+				Pair.Value.AccumulatorChannelID = Buffer.AddBlendTargetChannel(FHierarchicalEasingChannelData());
+			}
+
+			// Point any pre-existing easing channels to their HBias outputs
+			for (int32 ChannelIndex = 0; ChannelIndex < NumStartingChannels; ++ChannelIndex)
+			{
+				FHierarchicalEasingChannelData& Channel = Buffer.Channels[ChannelIndex];
+				const FHBiasChannelData* ChannelData = HBiasToChannelData.Find(Channel.HBias);
+				if (ChannelData)
+				{
+					Channel.ResultChannel = ChannelData->AccumulatorChannelID;
+				}
+			}
+		}
+
+		// Step 3: Create chains of channels for each HBias designated as a blend target
+		{
+			CachedHierarchicalBlendTargetChannels.Empty();
+
+			// HBiasToChannelData is a sorted map which is a sorted array internally.
+			// Unfortunately we don't have access to the underlying array, but we can use iterators to achieve the same results.
+			using FSortedMapIterator = TSortedMap<int16, FHBiasChannelData, TInlineAllocator<16>, TGreater<>>::TConstIterator;
+
+			const int32 NumChannels = HBiasToChannelData.Num();
+			int32 Index = 0;
+			for (FSortedMapIterator It = HBiasToChannelData.CreateConstIterator(); It; ++It, ++Index)
+			{
+				if (It.Value().bIsBlendTarget)
+				{
+					const int16 ThisHBias = It.Key();
+					const int32 N = NumChannels - Index - 1;
+					int32 I = N;
+
+					// This is a blend target - create channels for its entire parent chain
+					uint16 ParentChannel = INVALID_EASING_CHANNEL;
+					for (FSortedMapIterator NextIt = It; NextIt; ++NextIt)
+					{
+						const int16 SourceHBias = NextIt.Key();
+
+						FHierarchicalEasingChannelData BlendTargetChannel;
+						// Get our value from the accumulator channel for the source bias which includes all the accumulated weights for that HBias level.
+						// Normally this should just be a single weight
+						BlendTargetChannel.ParentChannel = HBiasToChannelData.FindChecked(SourceHBias).AccumulatorChannelID;
+						BlendTargetChannel.BlendMode = EHierarchicalBlendMode::ChildFirstBlendTarget;
+
+#if UE_ENABLE_MOVIESCENE_BEZIER_BLENDING
+						BlendTargetChannel.CoeffA = Factorial(N) / (Factorial(I) * Factorial(N-I));
+						BlendTargetChannel.ExpB   = static_cast<float>(N-I);
+						BlendTargetChannel.ExpC   = static_cast<float>(I);
+						BlendTargetChannel.BlendMode = EHierarchicalBlendMode::BezierSeries;
+#else
+						BlendTargetChannel.BlendMode = EHierarchicalBlendMode::ChildFirstBlendTarget;
+#endif
+						const uint16 BlendTargetChannelID = Buffer.AddBlendTargetChannel(BlendTargetChannel);
+
+						if (NextIt != It)
+						{
+							// If this isn't the first one, connect the last one to this one so we can
+							// combine weights hierarchically up the chain
+							Buffer.Channels[Buffer.Channels.Num()-2].ResultChannel = BlendTargetChannelID;
+						}
+
+						CachedHierarchicalBlendTargetChannels.Add(MakeTuple(SourceHBias, ThisHBias), BlendTargetChannelID);
+
+						--I;
+					}
+				}
+			}
+		}
+	}
+
+	// Assign easing channels
+	{
+		FEntityComponentFilter AssignEasingChannelsFilter;
+		AssignEasingChannelsFilter.All({ BuiltInComponents->InstanceHandle, BuiltInComponents->HierarchicalBlendTarget });
+		if (bChannelsHaveBeenInvalidated == false)
+		{
+			AssignEasingChannelsFilter.All({ BuiltInComponents->Tags.NeedsLink });
+		}
+		AssignEasingChannelsFilter.None({ BuiltInComponents->Tags.NeedsUnlink, BuiltInComponents->HierarchicalEasingProvider, BuiltInComponents->Tags.ImportedEntity });
+
+		FConsumerEasingChannelMutation ConsumerMutation(Linker, &CachedInstanceIDToChannel, &CachedHierarchicalBlendTargetChannels);
+		Linker->EntityManager.MutateConditional(AssignEasingChannelsFilter, ConsumerMutation, EMutuallyInclusiveComponentType::All);
+
+		ensure(bChannelsHaveBeenInvalidated || !ConsumerMutation.HasStaleComponents());
+		ConsumerMutation.RemoveStaleComponents(Linker);
+	}
+}
+
+UMovieSceneHierarchicalEasingFinalizationSystem::UMovieSceneHierarchicalEasingFinalizationSystem(const FObjectInitializer& ObjInit)
+	: Super(ObjInit)
+{
+	using namespace UE::MovieScene;
+
+	SystemCategories = EEntitySystemCategory::Core;
+	RelevantComponent = FBuiltInComponentTypes::Get()->HierarchicalBlendTarget;
+
+	if (HasAnyFlags(RF_ClassDefaultObject))
+	{
+		// Has to come after property instantiation property instantiation so that initial values get created correctly
+		DefineImplicitPrerequisite(UMovieScenePropertyInstantiatorSystem::StaticClass(), GetClass());
+		DefineImplicitPrerequisite(UMovieSceneHierarchicalEasingInstantiatorSystem::StaticClass(), GetClass());
+
+		DefineComponentConsumer(GetClass(), FBuiltInComponentTypes::Get()->HierarchicalBlendTarget);
+		DefineComponentProducer(GetClass(), FBuiltInComponentTypes::Get()->WeightAndEasingResult);
+	}
+}
+
+void UMovieSceneHierarchicalEasingFinalizationSystem::OnLink()
+{
+	InstantiatorSystem = Linker->LinkSystem<UMovieSceneHierarchicalEasingInstantiatorSystem>();
+	Linker->SystemGraph.AddReference(this, InstantiatorSystem);
+}
+
+void UMovieSceneHierarchicalEasingFinalizationSystem::OnRun(FSystemTaskPrerequisites& InPrerequisites, FSystemSubsequentTasks& Subsequents)
+{
+	InstantiatorSystem->FinalizeBlendTargets();
+}
+
 UWeightAndEasingEvaluatorSystem::UWeightAndEasingEvaluatorSystem(const FObjectInitializer& ObjInit)
 	: Super(ObjInit)
 {
@@ -739,17 +1081,17 @@ bool UWeightAndEasingEvaluatorSystem::IsRelevantImpl(UMovieSceneEntitySystemLink
 
 void UWeightAndEasingEvaluatorSystem::OnLink()
 {
-	PreAllocatedComputationData.Empty();
+	PreAllocatedComputationData.Reset();
 }
 
 void UWeightAndEasingEvaluatorSystem::OnUnlink()
 {
-	PreAllocatedComputationData.Empty();
+	PreAllocatedComputationData.Reset();
 }
 
-void UWeightAndEasingEvaluatorSystem::SetComputationBuffer(TArray<UE::MovieScene::FHierarchicalEasingChannelData>&& InPreAllocatedComputationData)
+UE::MovieScene::FHierarchicalEasingChannelBuffer& UWeightAndEasingEvaluatorSystem::GetComputationBuffer()
 {
-	PreAllocatedComputationData = MoveTemp(InPreAllocatedComputationData);
+	return PreAllocatedComputationData;
 }
 
 void UWeightAndEasingEvaluatorSystem::OnRun(FSystemTaskPrerequisites& InPrerequisites, FSystemSubsequentTasks& Subsequents)
@@ -777,6 +1119,8 @@ void UWeightAndEasingEvaluatorSystem::OnRun(FSystemTaskPrerequisites& InPrerequi
 	.SetStat(GET_STATID(MovieSceneEval_EvaluateEasingTask))
 	.Dispatch_PerEntity<FEvaluateEasings>(&Linker->EntityManager, ResetWeightsDependencies, &Subsequents);
 
+	ResetWeightsDependencies.AddComponentTask(Components->EasingResult, EvaluateEasing);
+
 	FGraphEventRef AccumulateManualWeights = FEntityTaskBuilder()
 	.ReadOneOrMoreOf(Components->WeightResult, Components->EasingResult)
 	.Write(Components->WeightAndEasingResult)
@@ -784,7 +1128,7 @@ void UWeightAndEasingEvaluatorSystem::OnRun(FSystemTaskPrerequisites& InPrerequi
 	.Dispatch_PerAllocation<FAccumulateManualWeights>(&Linker->EntityManager, ResetWeightsDependencies, &Subsequents);
 
 	// If we have hierarchical easing, we initialize all the weights to their hierarchical defaults
-	if (PreAllocatedComputationData.Num() > 0)
+	if (!PreAllocatedComputationData.IsEmpty())
 	{
 		FSystemTaskPrerequisites HarvestPrereqs {InPrerequisites};
 		HarvestPrereqs.AddComponentTask(Components->WeightAndEasingResult, EvaluateEasing);
@@ -797,7 +1141,7 @@ void UWeightAndEasingEvaluatorSystem::OnRun(FSystemTaskPrerequisites& InPrerequi
 		.Read(Components->HierarchicalEasingChannel)
 		.FilterAll({ Components->HierarchicalEasingProvider })  // Only harvest results from entities that are providing results
 		.SetStat(GET_STATID(MovieSceneEval_HarvestEasingTask))
-		.Dispatch_PerEntity<FHarvestHierarchicalEasings>(&Linker->EntityManager, HarvestPrereqs, nullptr, &PreAllocatedComputationData);
+		.Dispatch_PerEntity<FHarvestHierarchicalEasings>(&Linker->EntityManager, HarvestPrereqs, nullptr, &PreAllocatedComputationData.Channels);
 
 		FSystemTaskPrerequisites PropagatePrereqs {InPrerequisites};
 		PropagatePrereqs.AddRootTask(HarvestTask);
@@ -810,6 +1154,6 @@ void UWeightAndEasingEvaluatorSystem::OnRun(FSystemTaskPrerequisites& InPrerequi
 		.FilterNone({ Components->HierarchicalEasingProvider }) // Do not propagate hierarchical weights onto providers!
 		.SetStat(GET_STATID(MovieSceneEval_EvaluateEasingTask))
 		.Dispatch_PerAllocation<FPropagateHierarchicalEasings>(&Linker->EntityManager, PropagatePrereqs, &Subsequents,
-			PreAllocatedComputationData);
+			PreAllocatedComputationData.Channels);
 	}
 }
\ No newline at end of file
diff --git a/Engine/Source/Runtime/MovieSceneTracks/Public/Systems/MovieScenePropertyInstantiator.h b/Engine/Source/Runtime/MovieSceneTracks/Public/Systems/MovieScenePropertyInstantiator.h
--- a/Engine/Source/Runtime/MovieSceneTracks/Public/Systems/MovieScenePropertyInstantiator.h
+++ b/Engine/Source/Runtime/MovieSceneTracks/Public/Systems/MovieScenePropertyInstantiator.h
@@ -105,19 +105,74 @@ private:
 	using FChannelMask     = TBitArray<TFixedAllocator< 1 >>;
 	using FSlowPropertyPtr = TSharedPtr<FTrackInstancePropertyBindings>;
 
+	struct FContributorKey
+	{
+		static constexpr int16 ANY_HBIAS = TNumericLimits<int16>::Max();
+
+		FContributorKey(int32 InPropertyIndex)
+			: PropertyIndex(InPropertyIndex)
+			, HBias(ANY_HBIAS)
+		{}
+		FContributorKey(int32 InPropertyIndex, int16 InHBias)
+			: PropertyIndex(InPropertyIndex)
+			, HBias(InHBias)
+		{}
+
+		int32 PropertyIndex;
+		int16 HBias;
+
+		friend uint32 GetTypeHash(FContributorKey In)
+		{
+			return GetTypeHash(In.PropertyIndex);
+		}
+
+		friend bool operator==(FContributorKey A, FContributorKey B)
+		{
+			return A.PropertyIndex == B.PropertyIndex && (
+				A.HBias == B.HBias || A.HBias == ANY_HBIAS || B.HBias == ANY_HBIAS
+			);
+		}
+	};
+
+	struct FHierarchicalMetaData
+	{
+		FHierarchicalMetaData()
+			: NumContributors(0)
+			, HBias(0)
+		{
+			bWantsRestoreState = false;
+			bSupportsFastPath = true;
+			bNeedsInitialValue = false;
+			bBlendHierarchicalBias = false;
+			bInUse = false;
+		}
+
+		int32 NumContributors = 0;
+		int16 HBias = 0;
+		uint8 bWantsRestoreState : 1;
+		uint8 bSupportsFastPath : 1;
+		uint8 bNeedsInitialValue : 1;
+		uint8 bBlendHierarchicalBias : 1;
+		uint8 bInUse : 1;
+
+		void CombineWith(const FHierarchicalMetaData& Other);
+
+		void ResetTracking();
+	};
+
 	struct FObjectPropertyInfo
 	{
 		FObjectPropertyInfo(UE::MovieScene::FResolvedProperty&& InProperty)
 			: Property(MoveTemp(InProperty))
 			, BoundObject(nullptr)
-			, BlendChannel(INVALID_BLEND_CHANNEL)
-			, bWantsRestoreState(false)
+			, BlendChannel(FMovieSceneBlendChannelID::INVALID_BLEND_CHANNEL)
+			, bMaxHBiasHasChanged(false)
+			, bIsPartiallyAnimated(false)
 		{}
 
 		/** Variant of the property itself as either a pointer offset, a custom property index, or slow track instance bindings object */
 		UE::MovieScene::FResolvedProperty Property;
-
-		/** POinter to the blender system to use for this property, if its blended */
+		/** Pointer to the blender system to use for this property, if its blended */
 		TWeakObjectPtr<UMovieSceneBlenderSystem> Blender;
 		/** The object being animated */
 		UObject* BoundObject;
@@ -125,16 +180,18 @@ private:
 		FMovieScenePropertyBinding PropertyBinding;
 		/** Mask of composite channels that are not animated (set bits indicate an unanimated channel) */
 		FChannelMask EmptyChannels;
-		/** The entity that contains the property component itself. For fast path properties this is the actual child entity produced from the bound object instantiators. */
-		UE::MovieScene::FMovieSceneEntityID PropertyEntityID;
-		/** Blend channel allocated from Blender, or INVALID_BLEND_CHANNEL if unblended. */
+		/** The entity that contains the property component itself. Invalid for fast path properties. */
+		UE::MovieScene::FMovieSceneEntityID FinalBlendOutputID;
+		UE::MovieScene::FMovieSceneEntityID PreviousFastPathID;
+		/** Final blend channel for this object. */
 		uint16 BlendChannel;
 		/** The index of this property within FPropertyRegistry::GetProperties. */
 		int32 PropertyDefinitionIndex;
 		/** Index of a float-based property if this property has been set for float-to-double conversion */
 		TOptional<int32> ConvertedFromPropertyDefinitionIndex;
-		/** true if any of the contributors to this property need restore state. */
-		bool bWantsRestoreState;
+		FHierarchicalMetaData HierarchicalMetaData;
+		uint8 bMaxHBiasHasChanged : 1;
+		uint8 bIsPartiallyAnimated : 1;
 	};
 
 private:
@@ -161,6 +218,8 @@ private:
 		/** The index of the PropertyInfo member within UMovieScenePropertyInstantiatorSystem::ResolvedProperties */
 		int32 PropertyInfoIndex;
 
+		FContributorKey MakeContributorKey() const;
+
 		void MakeOutputComponentType(
 			const UE::MovieScene::FEntityManager& EntityManager,
 			TArrayView<const UE::MovieScene::FPropertyCompositeDefinition> Composites,
@@ -170,7 +229,6 @@ private:
 	void UpgradeFloatToDoubleProperties(const TBitArray<>& InvalidatedProperties);
 	void ProcessInvalidatedProperties(const TBitArray<>& InvalidatedProperties);
 	void UpdatePropertyInfo(const FPropertyParameters& Params);
-	bool PropertySupportsFastPath(const FPropertyParameters& Params) const;
 	void InitializeFastPath(const FPropertyParameters& Params);
 	void InitializeBlendPath(const FPropertyParameters& Params);
 
@@ -189,8 +247,8 @@ private:
 	static constexpr uint16 INVALID_BLEND_CHANNEL = uint16(-1);
 
 	TSparseArray<FObjectPropertyInfo> ResolvedProperties;
-	TMultiMap<int32, FMovieSceneEntityID> Contributors;
-	TMultiMap<int32, FMovieSceneEntityID> NewContributors;
+	TMultiMap<FContributorKey, FMovieSceneEntityID> Contributors;
+	TMultiMap<FContributorKey, FMovieSceneEntityID> NewContributors;
 
 	/** Reverse lookup from an entity to the index within ResolvedProperties that it animates.
 	 * @note: can contain INDEX_NONE for properties that have not resolved. */
diff --git a/Engine/Source/Runtime/MovieSceneTracks/Public/Systems/WeightAndEasingEvaluatorSystem.h b/Engine/Source/Runtime/MovieSceneTracks/Public/Systems/WeightAndEasingEvaluatorSystem.h
--- a/Engine/Source/Runtime/MovieSceneTracks/Public/Systems/WeightAndEasingEvaluatorSystem.h
+++ b/Engine/Source/Runtime/MovieSceneTracks/Public/Systems/WeightAndEasingEvaluatorSystem.h
@@ -26,6 +26,12 @@ class UWeightAndEasingEvaluatorSystem;
 class UObject;
 struct FMovieSceneSubSequenceData;
 
+// Bezier blending gives a true linear interpolation between biases, but 
+// disabled by default because it misbehaves if you feed it different parameters for 't' to each term
+#ifndef UE_ENABLE_MOVIESCENE_BEZIER_BLENDING
+	#define UE_ENABLE_MOVIESCENE_BEZIER_BLENDING 0 
+#endif
+
 namespace UE::MovieScene
 {
 
@@ -33,13 +39,64 @@ struct FEasingMutationBase;
 struct FAddEasingChannelToProviderMutation;
 struct FAddEasingChannelToConsumerMutation;
 
+enum class EHierarchicalBlendMode : uint8
+{
+	// Blend hierarchical weights by multiplying children with all parent weights (default)
+	AccumulateParentToChild,
+	// Combine weights from child -> parent using an increasingly small factor based on previously accumulated weights
+	// For example:
+	//     values = ( (a,t1), (b, t2), (c, t3), (d, t4))
+	//     float final_weight = 1;
+	//     float t = 1;
+	//     for (i : size(values))
+	//        float this = values[i].value*t*values[i].weight
+	//        t = max(t - this, 0)
+	//        final_weight *= this
+	ChildFirstBlendTarget,
+#if UE_ENABLE_MOVIESCENE_BEZIER_BLENDING
+	// Combine weights from child -> parent using a bezier interpolation series
+	BezierSeries,
+#endif
+};
+
 /** Computation data used for accumulating hierarchical weights for sub sequences */
 struct FHierarchicalEasingChannelData
 {
-	/** Our parent's computation data within PreAllocatedComputationData. Must only access if bResultsNeedResort is false. */
-	uint16 ParentChannel = uint16(-1);
 	/** The final result of this easing channel, accumulated with all parents */
 	double FinalResult = 1.0;
+	/** Initial easing result from manual and easing weights, not accumulated with all parents */
+	double UnaccumulatedResult = 1.0;
+
+	/** Our parent's computation data within PreAllocatedComputationData (in order to support one -> many multiplications). */
+	uint16 ParentChannel = uint16(-1);
+	/** An optional result channel to feed our result to (in order to support many -> one multiplications) */
+	uint16 ResultChannel = uint16(-1);
+
+	/** This channel's HBias */
+	int16 HBias = 0;
+
+#if UE_ENABLE_MOVIESCENE_BEZIER_BLENDING
+	// Bezier blending coefficients of the form a(1-t)^b + t^c
+	float CoeffA = 1.f;
+	float ExpB = 1.f;
+	float ExpC = 1.f;
+#endif
+
+	EHierarchicalBlendMode BlendMode = EHierarchicalBlendMode::AccumulateParentToChild;
+};
+
+struct FHierarchicalEasingChannelBuffer
+{
+	bool IsEmpty() const;
+
+	void Reset();
+	void ResetBlendTargets();
+
+	uint16 AddBaseChannel(const FHierarchicalEasingChannelData& InChannelData);
+	uint16 AddBlendTargetChannel(const FHierarchicalEasingChannelData& InChannelData);
+
+	TArray<FHierarchicalEasingChannelData> Channels;
+	int32 BlendTargetStartIndex = INDEX_NONE;
 };
 
 }  // namespace UE::MovieScene
@@ -57,13 +114,14 @@ public:
 		int32 HierarchicalDepth = -1;
 		uint16 StrongRefCount = 0;
 		uint16 WeakRefCount = 0;
-		uint16 EasingChannelID = uint16(-1);
 	};
 
 	GENERATED_BODY()
 
 	UMovieSceneHierarchicalEasingInstantiatorSystem(const FObjectInitializer& ObjInit);
 
+	void FinalizeBlendTargets();
+
 private:
 
 	virtual bool IsRelevantImpl(UMovieSceneEntitySystemLinker* InLinker) const;
@@ -76,8 +134,37 @@ private:
 	/** Map between a sub-sequence handle and the easing channel affecting it */
 	TSparseArray<uint16> CachedInstanceIDToChannel;
 
+	/** Cached output channels for a specific pair (source, target) of hierarchical blends */
+	TMap<TTuple<int16,int16>, uint16> CachedHierarchicalBlendTargetChannels;
+
 	UPROPERTY()
 	TObjectPtr<UWeightAndEasingEvaluatorSystem> EvaluatorSystem;
+
+	/** True of this system just updated all its channels */
+	bool bChannelsHaveBeenInvalidated;
+};
+
+/**
+ * System that finalizes creation of easing channels by allocating new channels for hierarchical blend targets if necessary
+ * This is required in order to solve the problem that the hierarchical instantiator needs to run before the property instantiator,
+ * but we can't allocate channels for blend targets until _after_ the property instantiator (which sets up the relevant BlendTarget components)
+ * 
+ */
+UCLASS()
+class UMovieSceneHierarchicalEasingFinalizationSystem : public UMovieSceneEntityInstantiatorSystem
+{
+public:
+	GENERATED_BODY()
+
+	UMovieSceneHierarchicalEasingFinalizationSystem(const FObjectInitializer& ObjInit);
+
+	virtual void OnLink() override;
+	virtual void OnRun(FSystemTaskPrerequisites& InPrerequisites, FSystemSubsequentTasks& Subsequents) override;
+
+private:
+
+	UPROPERTY()
+	TObjectPtr<UMovieSceneHierarchicalEasingInstantiatorSystem> InstantiatorSystem;
 };
 
 /**
@@ -95,7 +182,7 @@ public:
 	/**
 	 * Called from the instantiator system once channel IDs have been calculated
 	 */
-	void SetComputationBuffer(TArray<UE::MovieScene::FHierarchicalEasingChannelData>&& PreAllocatedComputationData);
+	UE::MovieScene::FHierarchicalEasingChannelBuffer& GetComputationBuffer();
 
 private:
 
@@ -107,6 +194,6 @@ private:
 private:
 
 	/** Unstable array of preallocated storage for computing easing results sorted by hierarchical depth. */
-	TArray<UE::MovieScene::FHierarchicalEasingChannelData> PreAllocatedComputationData;
+	UE::MovieScene::FHierarchicalEasingChannelBuffer PreAllocatedComputationData;
 };
 
