Sequencer: Expose setters and getters for audio section's looping, suppress subtitles and attenuation

diff --git a/Engine/Source/Runtime/MovieSceneTracks/Public/Sections/MovieSceneAudioSection.h b/Engine/Source/Runtime/MovieSceneTracks/Public/Sections/MovieSceneAudioSection.h
--- a/Engine/Source/Runtime/MovieSceneTracks/Public/Sections/MovieSceneAudioSection.h
+++ b/Engine/Source/Runtime/MovieSceneTracks/Public/Sections/MovieSceneAudioSection.h
@@ -87,34 +87,42 @@ public:
 	/**
 	 * @return Whether to allow looping if the section length is greater than the sound duration
 	 */
-	bool GetLooping() const
-	{
-		return bLooping;
-	}
+	UFUNCTION(BlueprintPure, Category = "Sequencer|Section")
+	bool GetLooping() const { return bLooping; }
+
+	/** Set whether the sound should be looped */
+	UFUNCTION(BlueprintCallable, Category = "Sequencer|Section")
+	void SetLooping(bool bInLooping) { bLooping = bInLooping; }
 
 	/**
 	 * @return Whether subtitles should be suppressed
 	 */
-	bool GetSuppressSubtitles() const
-	{
-		return bSuppressSubtitles;
-	}
+	UFUNCTION(BlueprintPure, Category = "Sequencer|Section")
+	bool GetSuppressSubtitles() const { return bSuppressSubtitles; }
+
+	/** Set whether subtitles should be suppressed */
+	UFUNCTION(BlueprintCallable, Category = "Sequencer|Section")
+	void SetSuppressSubtitles(bool bInSuppressSubtitles) { bSuppressSubtitles = bInSuppressSubtitles; }
 
 	/**
 	 * @return Whether override settings on this section should be used
 	 */
-	bool GetOverrideAttenuation() const
-	{
-		return bOverrideAttenuation;
-	}
+	UFUNCTION(BlueprintPure, Category = "Sequencer|Section")
+	bool GetOverrideAttenuation() const { return bOverrideAttenuation; }
+
+	/** Set whether the attentuation should be overriden */
+	UFUNCTION(BlueprintCallable, Category = "Sequencer|Section")
+	void SetOverrideAttenuation(bool bInOverrideAttenuation) { bOverrideAttenuation = bInOverrideAttenuation; }
 
 	/**
 	 * @return The attenuation settings
 	 */
-	USoundAttenuation* GetAttenuationSettings() const
-	{
-		return AttenuationSettings;
-	}
+	UFUNCTION(BlueprintPure, Category = "Sequencer|Section")
+	USoundAttenuation* GetAttenuationSettings() const { return AttenuationSettings; }
+
+	/** Set the attenuation settings for this audio section */
+	UFUNCTION(BlueprintCallable, Category = "Sequencer|Section")
+	void SetAttenuationSettings(USoundAttenuation* InAttenuationSettings) { AttenuationSettings = InAttenuationSettings; }
 
 	/*
 	 * @return The attach actor data
