Sequencer: Removed redundant stat scope

diff --git a/Engine/Source/Runtime/MovieSceneTracks/Private/Systems/MovieScenePiecewiseDoubleBlenderSystem.cpp b/Engine/Source/Runtime/MovieSceneTracks/Private/Systems/MovieScenePiecewiseDoubleBlenderSystem.cpp
--- a/Engine/Source/Runtime/MovieSceneTracks/Private/Systems/MovieScenePiecewiseDoubleBlenderSystem.cpp
+++ b/Engine/Source/Runtime/MovieSceneTracks/Private/Systems/MovieScenePiecewiseDoubleBlenderSystem.cpp
@@ -14,7 +14,6 @@
 
 #include UE_INLINE_GENERATED_CPP_BY_NAME(MovieScenePiecewiseDoubleBlenderSystem)
 
-DECLARE_CYCLE_STAT(TEXT("Piecewise Double Blender System"),       MovieSceneEval_PiecewiseDoubleBlenderSystem,  STATGROUP_MovieSceneECS);
 DECLARE_CYCLE_STAT(TEXT("Blend double values"),                   MovieSceneEval_BlendDoubleValues,        STATGROUP_MovieSceneECS);
 DECLARE_CYCLE_STAT(TEXT("Default combine blended double values"), MovieSceneEval_BlendCombineDoubleValues, STATGROUP_MovieSceneECS);
  
@@ -440,8 +439,6 @@ void UMovieScenePiecewiseDoubleBlenderSystem::OnRun(FSystemTaskPrerequisites& In
 {
 	using namespace UE::MovieScene;
 
-	SCOPE_CYCLE_COUNTER(MovieSceneEval_PiecewiseDoubleBlenderSystem)
-
 	CompactBlendChannels();
 
 	FEntityManager& EntityManager = Linker->EntityManager;
