Sequencer: Work around dependency between property instantiator and overlapping entity trackers when collecting garbage

  - The property instantiator will update any inputs/outputs that are marked as NeedsUnlink when garbage has been tagged. This in turn can cause subsequent entities to be marked for Unlink.
  - Overlapping entity trackers hook into the CleanTaggedGarbage event in order to clean up any references they hold to NeedsUnlink entities on garbage collection (that would otherwise normally get cleaned up through normal processing)

This introduces a dependency between the property instantiator and any overlapping entity trackers but since they were both hooked into the event in an undeterministic manner, the dependency was not honored. I have fixed this for now by triggering the CleanTaggedGarbage event after systems have their OnCleanTaggedGarbage method called, which the property instantiator now (correctly) utilizes instead of binding to the event directly.

This addresses an ensure in the overlapping entity tracker that I just encountered, and is evident on crash reporter (but is still handled gracefully).

diff --git a/Engine/Source/Runtime/MovieScene/Private/EntitySystem/MovieSceneEntitySystemLinker.cpp b/Engine/Source/Runtime/MovieScene/Private/EntitySystem/MovieSceneEntitySystemLinker.cpp
--- a/Engine/Source/Runtime/MovieScene/Private/EntitySystem/MovieSceneEntitySystemLinker.cpp
+++ b/Engine/Source/Runtime/MovieScene/Private/EntitySystem/MovieSceneEntitySystemLinker.cpp
@@ -398,13 +398,16 @@ void UMovieSceneEntitySystemLinker::CleanGarbage()
 	// the next time a runner gets flushed
 	LastInstantiationVersion = 0;
 
-	// Allow any other system to cleanup garbage
-	Events.CleanTaggedGarbage.Broadcast(this);
-
 	auto RouteCleanTaggedGarbage = [](UMovieSceneEntitySystem* System){ System->CleanTaggedGarbage(); };
 	SystemGraph.IteratePhase(ESystemPhase::Spawn, RouteCleanTaggedGarbage);
 	SystemGraph.IteratePhase(ESystemPhase::Instantiation, RouteCleanTaggedGarbage);
 
+	// Allow any other system to cleanup garbage
+	// NOTE: Order is important here - we need to broadcast this after systems have CleanTaggedGarbage called
+	//       since systems are able to (and more likely to) cause additional entities to be unlinked in response
+	//       to finding entities tagged NeedsUnlink.
+	Events.CleanTaggedGarbage.Broadcast(this);
+
 	TArray<FMovieSceneEntityID> UnresolvedEntities;
 
 	FEntityTaskBuilder()
diff --git a/Engine/Source/Runtime/MovieSceneTracks/Private/Systems/MovieScenePropertyInstantiator.cpp b/Engine/Source/Runtime/MovieSceneTracks/Private/Systems/MovieScenePropertyInstantiator.cpp
--- a/Engine/Source/Runtime/MovieSceneTracks/Private/Systems/MovieScenePropertyInstantiator.cpp
+++ b/Engine/Source/Runtime/MovieSceneTracks/Private/Systems/MovieScenePropertyInstantiator.cpp
@@ -79,8 +79,6 @@ UE::MovieScene::FPropertyStats UMovieScenePropertyInstantiatorSystem::GetStatsFo
 
 void UMovieScenePropertyInstantiatorSystem::OnLink()
 {
-	Linker->Events.CleanTaggedGarbage.AddUObject(this, &UMovieScenePropertyInstantiatorSystem::CleanTaggedGarbage);
-
 	CleanFastPathMask.Reset();
 	CleanFastPathMask.SetAll({ BuiltInComponents->FastPropertyOffset, BuiltInComponents->SlowProperty, BuiltInComponents->CustomPropertyIndex });
 	CleanFastPathMask.CombineWithBitwiseOR(Linker->EntityManager.GetComponents()->GetMigrationMask(), EBitwiseOperatorFlags::MaxSize);
@@ -90,8 +88,6 @@ void UMovieScenePropertyInstantiatorSystem::OnUnlink()
 {
 	using namespace UE::MovieScene;
 
-	Linker->Events.CleanTaggedGarbage.RemoveAll(this);
-
 	const bool bAllPropertiesClean = (
 				ResolvedProperties.Num() == 0 &&
 				Contributors.Num() == 0 &&
@@ -128,7 +124,7 @@ void UMovieScenePropertyInstantiatorSystem::OnUnlink()
 	}
 }
 
-void UMovieScenePropertyInstantiatorSystem::CleanTaggedGarbage(UMovieSceneEntitySystemLinker*)
+void UMovieScenePropertyInstantiatorSystem::OnCleanTaggedGarbage()
 {
 	using namespace UE::MovieScene;
 
diff --git a/Engine/Source/Runtime/MovieSceneTracks/Public/Systems/MovieScenePropertyInstantiator.h b/Engine/Source/Runtime/MovieSceneTracks/Public/Systems/MovieScenePropertyInstantiator.h
--- a/Engine/Source/Runtime/MovieSceneTracks/Public/Systems/MovieScenePropertyInstantiator.h
+++ b/Engine/Source/Runtime/MovieSceneTracks/Public/Systems/MovieScenePropertyInstantiator.h
@@ -99,6 +99,7 @@ private:
 	virtual void OnRun(FSystemTaskPrerequisites& InPrerequisites, FSystemSubsequentTasks& Subsequents) override;
 	virtual void OnLink() override;
 	virtual void OnUnlink() override;
+	virtual void OnCleanTaggedGarbage() override;
 
 private:
 
@@ -240,8 +241,6 @@ private:
 
 	void InitializePropertyMetaData(FSystemTaskPrerequisites& InPrerequisites, FSystemSubsequentTasks& Subsequents);
 
-	void CleanTaggedGarbage(UMovieSceneEntitySystemLinker*);
-
 private:
 
 	static constexpr uint16 INVALID_BLEND_CHANNEL = uint16(-1);
