Sequencer: Added missing file from CL 25233877 that got lost unshelving from main

diff --git a/Engine/Source/Runtime/MovieSceneTracks/Private/Systems/MovieSceneInitialValueSystem.cpp b/Engine/Source/Runtime/MovieSceneTracks/Private/Systems/MovieSceneInitialValueSystem.cpp
--- a/Engine/Source/Runtime/MovieSceneTracks/Private/Systems/MovieSceneInitialValueSystem.cpp
+++ b/Engine/Source/Runtime/MovieSceneTracks/Private/Systems/MovieSceneInitialValueSystem.cpp
@@ -58,7 +58,11 @@ struct FInitialValueMutation : IMovieSceneEntityMutation
 
 	virtual void CreateMutation(FEntityManager* EntityManager, FComponentMask* InOutEntityComponentTypes) const override
 	{
-		InOutEntityComponentTypes->Set(BuiltInComponents->InitialValueIndex);
+		if (IsCached())
+		{
+			InOutEntityComponentTypes->Set(BuiltInComponents->InitialValueIndex);
+		}
+		InOutEntityComponentTypes->Set(BuiltInComponents->Tags.HasAssignedInitialValue);
 	}
 
 	virtual void InitializeAllocation(FEntityAllocation* Allocation, const FComponentMask& AllocationType) const
@@ -137,6 +141,7 @@ void UMovieSceneInitialValueSystem::OnRun(FSystemTaskPrerequisites& InPrerequisi
 			Filter.Any(Mutation.AnyInitialValue);
 			Filter.All({ BuiltInComponents->Tags.NeedsLink });
 			Filter.None({ BuiltInComponents->InitialValueIndex });
+			Filter.None({ BuiltInComponents->Tags.HasAssignedInitialValue });
 
 			Linker->EntityManager.MutateAll(Filter, Mutation);
 		}
@@ -166,20 +171,9 @@ void UMovieSceneInitialValueSystem::OnRun(FSystemTaskPrerequisites& InPrerequisi
 		Filter.Any(Mutation.AnyInitialValue);
 		Filter.Any({ BuiltInComponents->BoundObject, BuiltInComponents->Interrogation.OutputKey });
 		Filter.All({ BuiltInComponents->Tags.NeedsLink });
+		Filter.None({ BuiltInComponents->Tags.HasAssignedInitialValue });
 
-		for (FEntityAllocationIteratorItem Item : Linker->EntityManager.Iterate(&Filter))
-		{
-			const FEntityAllocation* Allocation     = Item.GetAllocation();
-			FComponentMask           AllocationType = Item.GetAllocationType();
-
-			FComponentTypeID InitialValueType = FComponentMask::BitwiseAND(AllocationType, Mutation.AnyInitialValue, EBitwiseOperatorFlags::MinSize).First();
-
-			IInitialValueProcessor* Processor = Mutation.PropertyTypeToProcessor.FindRef(InitialValueType);
-			if (ensure(Processor))
-			{
-				Processor->Process(Allocation, AllocationType);
-			}
-		}
+		Linker->EntityManager.MutateAll(Filter, Mutation);
 	}
 }
 
