Sequencer: Change SetRelativeScale3D to SetRelativeScale3D_Direct and call it before SetRelativeLocationAndRotation so that UpdateComponentToWorld is only called once.

Added FIntermediate3DTransform::ApplyTranslationAndRotationTo so that path track can use it to just set relative location and rotation. This fixes a bug where the component velocity wouldn't be set when using the path track.

diff --git a/Engine/Source/Runtime/MovieSceneTracks/Private/Evaluation/MovieScene3DPathTemplate.cpp b/Engine/Source/Runtime/MovieSceneTracks/Private/Evaluation/MovieScene3DPathTemplate.cpp
--- a/Engine/Source/Runtime/MovieSceneTracks/Private/Evaluation/MovieScene3DPathTemplate.cpp
+++ b/Engine/Source/Runtime/MovieSceneTracks/Private/Evaluation/MovieScene3DPathTemplate.cpp
@@ -3,6 +3,7 @@
 #include "Evaluation/MovieScene3DPathTemplate.h"
 #include "Evaluation/MovieSceneTemplateCommon.h"
 #include "MovieSceneCommonHelpers.h"
+#include "MovieSceneTracksPropertyTypes.h"
 #include "GameFramework/Actor.h"
 #include "Components/SplineComponent.h"
 #include "Evaluation/MovieSceneEvaluation.h"
@@ -162,7 +163,9 @@ struct F3DPathExecutionToken
 				Eval(SceneComponent, SplineComponent, Location, Rotation);
 				
 				SceneComponent->SetMobility(EComponentMobility::Movable);
-				SceneComponent->SetRelativeLocationAndRotation(Location, Rotation);
+
+				UE::MovieScene::FIntermediate3DTransform Transform(Location, Rotation, FVector(1.0));
+				UE::MovieScene::FIntermediate3DTransform::ApplyTranslationAndRotationTo(SceneComponent, Transform);
 			}
 		}
 	}
diff --git a/Engine/Source/Runtime/MovieSceneTracks/Private/MovieSceneTracksComponentTypes.cpp b/Engine/Source/Runtime/MovieSceneTracks/Private/MovieSceneTracksComponentTypes.cpp
--- a/Engine/Source/Runtime/MovieSceneTracks/Private/MovieSceneTracksComponentTypes.cpp
+++ b/Engine/Source/Runtime/MovieSceneTracks/Private/MovieSceneTracksComponentTypes.cpp
@@ -178,7 +178,7 @@ FIntermediate3DTransform GetComponentTransform(const UObject* Object)
 	return Result;
 }
 
-void SetComponentTransform(USceneComponent* SceneComponent, const FIntermediate3DTransform& InTransform)
+void SetComponentTranslationAndRotation(USceneComponent* SceneComponent, const FIntermediate3DTransform& Transform)
 {
 	// If this is a simulating component, teleport since sequencer takes over. 
 	// Teleport will not have no velocity, but it's computed later by sequencer so that it will be correct for physics.
@@ -187,10 +187,30 @@ void SetComponentTransform(USceneComponent* SceneComponent, const FIntermediate3
 	USceneComponent* RootComponent = Actor ? Actor->GetRootComponent() : nullptr;
 	bool bIsSimulatingPhysics = RootComponent ? RootComponent->IsSimulatingPhysics() : false;
 
-	FVector Translation = InTransform.GetTranslation();
-	FRotator Rotation = InTransform.GetRotation();
+	FVector Translation = Transform.GetTranslation();
+	FRotator Rotation = Transform.GetRotation();
+	SceneComponent->SetRelativeLocationAndRotation(Translation, Rotation, false, nullptr, bIsSimulatingPhysics ? ETeleportType::ResetPhysics : ETeleportType::None);
+
+	// Force the location and rotation values to avoid Rot->Quat->Rot conversions
+	SceneComponent->SetRelativeLocation_Direct(Translation);
+	SceneComponent->SetRelativeRotation_Direct(Rotation);
+}
+
+void SetComponentTransform(USceneComponent* SceneComponent, const FIntermediate3DTransform& Transform)
+{
+	// If this is a simulating component, teleport since sequencer takes over. 
+	// Teleport will not have no velocity, but it's computed later by sequencer so that it will be correct for physics.
+	// @todo: We would really rather not 
+	AActor* Actor = SceneComponent->GetOwner();
+	USceneComponent* RootComponent = Actor ? Actor->GetRootComponent() : nullptr;
+	bool bIsSimulatingPhysics = RootComponent ? RootComponent->IsSimulatingPhysics() : false;
+
+	// Set Scale3D direct first to avoid UpdateComponentToWorld being called twice (ie. if calling SetRelativeLocationAndRotation and SetRelativeScale3D)
+	SceneComponent->SetRelativeScale3D_Direct(Transform.GetScale());
+
+	FVector Translation = Transform.GetTranslation();
+	FRotator Rotation = Transform.GetRotation();
 	SceneComponent->SetRelativeLocationAndRotation(Translation, Rotation, false, nullptr, bIsSimulatingPhysics ? ETeleportType::ResetPhysics : ETeleportType::None);
-	SceneComponent->SetRelativeScale3D(InTransform.GetScale());
 
 	// Force the location and rotation values to avoid Rot->Quat->Rot conversions
 	SceneComponent->SetRelativeLocation_Direct(Translation);
@@ -199,7 +219,7 @@ void SetComponentTransform(USceneComponent* SceneComponent, const FIntermediate3
 
 void SetComponentTransformAndVelocity(UObject* Object, const FIntermediate3DTransform& InTransform)
 {
-	InTransform.ApplyTo(CastChecked<USceneComponent>(Object));
+	UE::MovieScene::FIntermediate3DTransform::ApplyTransformTo(CastChecked<USceneComponent>(Object), InTransform);
 }
 
 FIntermediateColor GetLightComponentLightColor(const UObject* Object, EColorPropertyType InColorType)
@@ -537,18 +557,44 @@ void SetSkeletalMeshAnimationMode(UObject* Object, uint8 InAnimationMode)
 }
 
 void FIntermediate3DTransform::ApplyTo(USceneComponent* SceneComponent) const
+{
+	ApplyTransformTo(SceneComponent, *this);
+}
+
+void FIntermediate3DTransform::ApplyTransformTo(USceneComponent* SceneComponent, const FIntermediate3DTransform& Transform)
+{
+	double DeltaTime = FApp::GetDeltaTime();
+	if (DeltaTime <= 0)
+	{
+		SetComponentTransform(SceneComponent, Transform);
+	}
+	else
+	{
+		/* Cache initial absolute position */
+		FVector PreviousPosition = SceneComponent->GetComponentLocation();
+
+		SetComponentTransform(SceneComponent, Transform);
+
+		/* Get current absolute position and set component velocity */
+		FVector CurrentPosition = SceneComponent->GetComponentLocation();
+		FVector ComponentVelocity = (CurrentPosition - PreviousPosition) / DeltaTime;
+		SceneComponent->ComponentVelocity = ComponentVelocity;
+	}
+}
+
+void FIntermediate3DTransform::ApplyTranslationAndRotationTo(USceneComponent* SceneComponent, const FIntermediate3DTransform& Transform)
 {
 	double DeltaTime = FApp::GetDeltaTime();
 	if (DeltaTime <= 0)
 	{
-		SetComponentTransform(SceneComponent, *this);
+		SetComponentTranslationAndRotation(SceneComponent, Transform);
 	}
 	else
 	{
 		/* Cache initial absolute position */
 		FVector PreviousPosition = SceneComponent->GetComponentLocation();
 
-		SetComponentTransform(SceneComponent, *this);
+		SetComponentTranslationAndRotation(SceneComponent, Transform);
 
 		/* Get current absolute position and set component velocity */
 		FVector CurrentPosition = SceneComponent->GetComponentLocation();
diff --git a/Engine/Source/Runtime/MovieSceneTracks/Public/MovieSceneTracksPropertyTypes.h b/Engine/Source/Runtime/MovieSceneTracks/Public/MovieSceneTracksPropertyTypes.h
--- a/Engine/Source/Runtime/MovieSceneTracks/Public/MovieSceneTracksPropertyTypes.h
+++ b/Engine/Source/Runtime/MovieSceneTracks/Public/MovieSceneTracksPropertyTypes.h
@@ -184,6 +184,9 @@ struct FIntermediate3DTransform
 	}
 
 	MOVIESCENETRACKS_API void ApplyTo(USceneComponent* SceneComponent) const;
+	
+	MOVIESCENETRACKS_API static void ApplyTransformTo(USceneComponent* SceneComponent, const FIntermediate3DTransform& Transform);
+	MOVIESCENETRACKS_API static void ApplyTranslationAndRotationTo(USceneComponent* SceneComponent, const FIntermediate3DTransform& Transform);
 };
 
 MOVIESCENETRACKS_API FIntermediate3DTransform GetComponentTransform(const UObject* Object);
