Sequencer: Follow up to 25360034. Need to just call SetTransform instead of separating out etRelativeScale3D_Direct/SetRelativeScale3D and SetRelativeLocationAndRotation because UpdateComponentToWorld will not be called if any of the parameters are the same as the cached values.

diff --git a/Engine/Source/Runtime/MovieSceneTracks/Private/MovieSceneTracksComponentTypes.cpp b/Engine/Source/Runtime/MovieSceneTracks/Private/MovieSceneTracksComponentTypes.cpp
--- a/Engine/Source/Runtime/MovieSceneTracks/Private/MovieSceneTracksComponentTypes.cpp
+++ b/Engine/Source/Runtime/MovieSceneTracks/Private/MovieSceneTracksComponentTypes.cpp
@@ -205,12 +205,9 @@ void SetComponentTransform(USceneComponent* SceneComponent, const FIntermediate3
 	USceneComponent* RootComponent = Actor ? Actor->GetRootComponent() : nullptr;
 	bool bIsSimulatingPhysics = RootComponent ? RootComponent->IsSimulatingPhysics() : false;
 
-	// Set Scale3D direct first to avoid UpdateComponentToWorld being called twice (ie. if calling SetRelativeLocationAndRotation and SetRelativeScale3D)
-	SceneComponent->SetRelativeScale3D_Direct(Transform.GetScale());
-
 	FVector Translation = Transform.GetTranslation();
 	FRotator Rotation = Transform.GetRotation();
-	SceneComponent->SetRelativeLocationAndRotation(Translation, Rotation, false, nullptr, bIsSimulatingPhysics ? ETeleportType::ResetPhysics : ETeleportType::None);
+	SceneComponent->SetRelativeTransform(FTransform(Rotation, Translation, Transform.GetScale()), false, nullptr, bIsSimulatingPhysics ? ETeleportType::ResetPhysics : ETeleportType::None);
 
 	// Force the location and rotation values to avoid Rot->Quat->Rot conversions
 	SceneComponent->SetRelativeLocation_Direct(Translation);
