Sequencer: make string properties and string tracks use the ECS evaluation

diff --git a/Engine/Source/Runtime/MovieScene/Private/EntitySystem/BuiltInComponentTypes.cpp b/Engine/Source/Runtime/MovieScene/Private/EntitySystem/BuiltInComponentTypes.cpp
--- a/Engine/Source/Runtime/MovieScene/Private/EntitySystem/BuiltInComponentTypes.cpp
+++ b/Engine/Source/Runtime/MovieScene/Private/EntitySystem/BuiltInComponentTypes.cpp
@@ -106,6 +106,7 @@ FBuiltInComponentTypes::FBuiltInComponentTypes()
 	ComponentRegistry->NewComponentType(&DoubleChannel[7],        TEXT("Double Channel 7"));
 	ComponentRegistry->NewComponentType(&DoubleChannel[8],        TEXT("Double Channel 8"));
 	ComponentRegistry->NewComponentType(&WeightChannel,           TEXT("Weight Channel"));
+	ComponentRegistry->NewComponentType(&StringChannel,           TEXT("String Channel"));
 	ComponentRegistry->NewComponentType(&ObjectPathChannel,       TEXT("Object Path Channel"));
 
 	ComponentRegistry->NewComponentType(&CachedInterpolation[0],   TEXT("Cached Interpolation [0]"));
@@ -146,6 +147,7 @@ FBuiltInComponentTypes::FBuiltInComponentTypes()
 	ComponentRegistry->NewComponentType(&DoubleResult[6],       TEXT("Double Result 6"));
 	ComponentRegistry->NewComponentType(&DoubleResult[7],       TEXT("Double Result 7"));
 	ComponentRegistry->NewComponentType(&DoubleResult[8],       TEXT("Double Result 8"));
+	ComponentRegistry->NewComponentType(&StringResult,          TEXT("String Result"));
 	ComponentRegistry->NewComponentType(&ObjectResult,          TEXT("Object Result"));
 
 	ComponentRegistry->NewComponentType(&BaseByte,			    TEXT("Base Byte"));
@@ -368,6 +370,15 @@ FBuiltInComponentTypes::FBuiltInComponentTypes()
 		ComponentRegistry->Factories.DefineMutuallyInclusiveComponent(EasingResult, WeightAndEasingResult);
 		ComponentRegistry->Factories.DefineMutuallyInclusiveComponent(WeightResult, WeightAndEasingResult);
 	}
+	
+	// String channel relationships
+	{
+		ComponentRegistry->Factories.DuplicateChildComponent(StringChannel);
+		ComponentRegistry->Factories.DuplicateChildComponent(StringResult);
+
+		ComponentRegistry->Factories.DefineMutuallyInclusiveComponent(StringChannel, EvalTime);
+		ComponentRegistry->Factories.DefineMutuallyInclusiveComponent(StringChannel, StringResult);
+	}
 
 	// Object path channel relationships
 	{
diff --git a/Engine/Source/Runtime/MovieScene/Public/EntitySystem/BuiltInComponentTypes.h b/Engine/Source/Runtime/MovieScene/Public/EntitySystem/BuiltInComponentTypes.h
--- a/Engine/Source/Runtime/MovieScene/Public/EntitySystem/BuiltInComponentTypes.h
+++ b/Engine/Source/Runtime/MovieScene/Public/EntitySystem/BuiltInComponentTypes.h
@@ -53,6 +53,7 @@ struct FMovieSceneDoubleChannel;
 struct FMovieSceneFloatChannel;
 struct FMovieSceneIntegerChannel;
 struct FMovieSceneObjectPathChannel;
+struct FMovieSceneStringChannel;
 struct FMovieScenePropertyBinding;
 
 
@@ -183,7 +184,7 @@ struct FSourceFloatChannel
 };
 
 /**
- * The component data for evaluation a double channel
+ * The component data for evaluating a double channel
  */
 struct FSourceDoubleChannel
 {
@@ -198,6 +199,22 @@ struct FSourceDoubleChannel
 	const FMovieSceneDoubleChannel* Source;
 };
 
+/**
+ * The component data for evaluating a string channel
+ */
+struct FSourceStringChannel
+{
+	FSourceStringChannel()
+		: Source(nullptr)
+	{}
+
+	FSourceStringChannel(const FMovieSceneStringChannel* InSource)
+		: Source(InSource)
+	{}
+
+	const FMovieSceneStringChannel* Source;
+};
+
 struct FEvaluationHookFlags
 {
 	bool bHasBegun = false;
@@ -402,6 +419,9 @@ public:
 	// An FMovieSceneDoubleChannel considered to be at index N within the source structure (ie 0 = Location.X, Vector.X; 1 = Location.Y, Vector.Y)
 	TComponentTypeID<FSourceDoubleChannel> DoubleChannel[9];
 
+	// An FMovieSceneStringChannel
+	TComponentTypeID<FSourceStringChannel> StringChannel;
+
 	// A cached interpolation structure relating to either float channels or double channels
 	TComponentTypeID<Interpolation::FCachedInterpolation> CachedInterpolation[9];
 
@@ -426,6 +446,9 @@ public:
 	TComponentTypeID<double> WeightResult;
 	TComponentTypeID<double> EasingResult;
 
+	// The result of an evaluated FMovieSceneStringChannel
+	TComponentTypeID<FString> StringResult;
+
 	// The result of an evaluated FMovieSceneObjectPathChannel
 	TComponentTypeID<UObject*> ObjectResult;
 
diff --git a/Engine/Source/Runtime/MovieSceneTracks/Private/Evaluation/MovieScenePropertyTemplates.cpp b/Engine/Source/Runtime/MovieSceneTracks/Private/Evaluation/MovieScenePropertyTemplates.cpp
--- a/Engine/Source/Runtime/MovieSceneTracks/Private/Evaluation/MovieScenePropertyTemplates.cpp
+++ b/Engine/Source/Runtime/MovieSceneTracks/Private/Evaluation/MovieScenePropertyTemplates.cpp
@@ -8,7 +8,6 @@
 #include "Sections/MovieSceneIntegerSection.h"
 #include "Sections/MovieSceneVectorSection.h"
 #include "Sections/MovieScene3DTransformSection.h"
-#include "Sections/MovieSceneStringSection.h"
 #include "Tracks/MovieScenePropertyTrack.h"
 #include "MovieScene.h"
 #include "Evaluation/MovieSceneEvaluation.h"
@@ -46,20 +45,3 @@ void FMovieSceneBoolPropertySectionTemplate::Evaluate(const FMovieSceneEvaluatio
 	}
 }
 
-//	----------------------------------------------------------------------------
-//	String Property Template
-FMovieSceneStringPropertySectionTemplate::FMovieSceneStringPropertySectionTemplate(const UMovieSceneStringSection& Section, const UMovieScenePropertyTrack& Track)
-	: FMovieScenePropertySectionTemplate(Track.GetPropertyName(), Track.GetPropertyPath().ToString())
-	, StringCurve(Section.GetChannel())
-{}
-
-void FMovieSceneStringPropertySectionTemplate::Evaluate(const FMovieSceneEvaluationOperand& Operand, const FMovieSceneContext& Context, const FPersistentEvaluationData& PersistentData, FMovieSceneExecutionTokens& ExecutionTokens) const
-{
-	const FString* Result = StringCurve.Evaluate(Context.GetTime());
-	if (Result)
-	{
-		ExecutionTokens.Add(TPropertyTrackExecutionToken<FString>(*Result));
-	}
-}
-
-
diff --git a/Engine/Source/Runtime/MovieSceneTracks/Private/Evaluation/MovieScenePropertyTemplates.h b/Engine/Source/Runtime/MovieSceneTracks/Private/Evaluation/MovieScenePropertyTemplates.h
--- a/Engine/Source/Runtime/MovieSceneTracks/Private/Evaluation/MovieScenePropertyTemplates.h
+++ b/Engine/Source/Runtime/MovieSceneTracks/Private/Evaluation/MovieScenePropertyTemplates.h
@@ -6,7 +6,6 @@
 #include "UObject/ObjectMacros.h"
 #include "Evaluation/MovieSceneEvalTemplate.h"
 #include "Evaluation/MovieScenePropertyTemplate.h"
-#include "Sections/MovieSceneStringSection.h"
 #include "Evaluation/Blending/MovieSceneBlendType.h"
 #include "Channels/MovieSceneBoolChannel.h"
 #include "Channels/MovieSceneByteChannel.h"
@@ -43,22 +42,3 @@ protected:
 	FMovieSceneBoolChannel BoolCurve;
 };
 
-
-USTRUCT()
-struct FMovieSceneStringPropertySectionTemplate : public FMovieScenePropertySectionTemplate
-{
-	GENERATED_BODY()
-	
-	FMovieSceneStringPropertySectionTemplate(){}
-	FMovieSceneStringPropertySectionTemplate(const UMovieSceneStringSection& Section, const UMovieScenePropertyTrack& Track);
-
-protected:
-
-	virtual UScriptStruct& GetScriptStructImpl() const override { return *StaticStruct(); }
-	virtual void SetupOverrides() override { EnableOverrides(RequiresSetupFlag); }
-	virtual void Evaluate(const FMovieSceneEvaluationOperand& Operand, const FMovieSceneContext& Context, const FPersistentEvaluationData& PersistentData, FMovieSceneExecutionTokens& ExecutionTokens) const override;
-
-	UPROPERTY()
-	FMovieSceneStringChannel StringCurve;
-};
-
diff --git a/Engine/Source/Runtime/MovieSceneTracks/Private/MovieSceneTracksComponentTypes.cpp b/Engine/Source/Runtime/MovieSceneTracks/Private/MovieSceneTracksComponentTypes.cpp
--- a/Engine/Source/Runtime/MovieSceneTracks/Private/MovieSceneTracksComponentTypes.cpp
+++ b/Engine/Source/Runtime/MovieSceneTracks/Private/MovieSceneTracksComponentTypes.cpp
@@ -848,6 +848,7 @@ FMovieSceneTracksComponentTypes::FMovieSceneTracksComponentTypes()
 	ComponentRegistry->NewPropertyType(Integer, TEXT("int32"));
 	ComponentRegistry->NewPropertyType(FloatVector, TEXT("float vector"));
 	ComponentRegistry->NewPropertyType(DoubleVector, TEXT("double vector"));
+	ComponentRegistry->NewPropertyType(String, TEXT("FString"));
 
 	ComponentRegistry->NewPropertyType(Transform, TEXT("FTransform"));
 	ComponentRegistry->NewPropertyType(EulerTransform, TEXT("FEulerTransform"));
@@ -968,6 +969,12 @@ FMovieSceneTracksComponentTypes::FMovieSceneTracksComponentTypes()
 	.SetCustomAccessors(&Accessors.Color)
 	.Commit(FColorHandler());
 
+	// --------------------------------------------------------------------------------------------
+	// Set up string properties
+	BuiltInComponents->PropertyRegistry.DefineProperty(String, TEXT("Apply String Properties"))
+	.AddSoleChannel(BuiltInComponents->StringResult)
+	.Commit();
+
 	// --------------------------------------------------------------------------------------------
 	// Set up float parameters
 	BuiltInComponents->PropertyRegistry.DefineProperty(FloatParameter, TEXT("Apply Float Parameters"))
diff --git a/Engine/Source/Runtime/MovieSceneTracks/Private/Sections/MovieSceneStringSection.cpp b/Engine/Source/Runtime/MovieSceneTracks/Private/Sections/MovieSceneStringSection.cpp
--- a/Engine/Source/Runtime/MovieSceneTracks/Private/Sections/MovieSceneStringSection.cpp
+++ b/Engine/Source/Runtime/MovieSceneTracks/Private/Sections/MovieSceneStringSection.cpp
@@ -1,8 +1,10 @@
 // Copyright Epic Games, Inc. All Rights Reserved.
 
 #include "Sections/MovieSceneStringSection.h"
-#include "UObject/SequencerObjectVersion.h"
 #include "Channels/MovieSceneChannelProxy.h"
+#include "MovieSceneTracksComponentTypes.h"
+#include "Tracks/MovieScenePropertyTrack.h"
+#include "UObject/SequencerObjectVersion.h"
 
 #include UE_INLINE_GENERATED_CPP_BY_NAME(MovieSceneStringSection)
 
@@ -27,3 +29,27 @@ UMovieSceneStringSection::UMovieSceneStringSection( const FObjectInitializer& Ob
 
 #endif
 }
+
+bool UMovieSceneStringSection::PopulateEvaluationFieldImpl(const TRange<FFrameNumber>& EffectiveRange, const FMovieSceneEvaluationFieldEntityMetaData& InMetaData, FMovieSceneEntityComponentFieldBuilder* OutFieldBuilder)
+{
+	FMovieScenePropertyTrackEntityImportHelper::PopulateEvaluationField(*this, EffectiveRange, InMetaData, OutFieldBuilder);
+	return true;
+}
+
+void UMovieSceneStringSection::ImportEntityImpl(UMovieSceneEntitySystemLinker* EntityLinker, const FEntityImportParams& Params, FImportedEntity* OutImportedEntity)
+{
+	using namespace UE::MovieScene;
+
+	if (!StringCurve.HasAnyData())
+	{
+		return;
+	}
+
+	const FBuiltInComponentTypes* Components = FBuiltInComponentTypes::Get();
+	const FMovieSceneTracksComponentTypes* TracksComponents = FMovieSceneTracksComponentTypes::Get();
+
+	FPropertyTrackEntityImportHelper(TracksComponents->String)
+		.Add(Components->StringChannel, &StringCurve)
+		.Commit(this, Params, OutImportedEntity);
+}
+
diff --git a/Engine/Source/Runtime/MovieSceneTracks/Private/Systems/MovieSceneStringPropertySystem.cpp b/Engine/Source/Runtime/MovieSceneTracks/Private/Systems/MovieSceneStringPropertySystem.cpp
new file mode 100644
--- /dev/null
+++ b/Engine/Source/Runtime/MovieSceneTracks/Private/Systems/MovieSceneStringPropertySystem.cpp
@@ -0,0 +1,25 @@
+// Copyright Epic Games, Inc. All Rights Reserved.
+
+#include "Systems/MovieSceneStringPropertySystem.h"
+#include "Systems/StringChannelEvaluatorSystem.h"
+#include "MovieSceneTracksComponentTypes.h"
+
+#include UE_INLINE_GENERATED_CPP_BY_NAME(MovieSceneStringPropertySystem)
+
+
+UMovieSceneStringPropertySystem::UMovieSceneStringPropertySystem(const FObjectInitializer& ObjInit)
+	: Super(ObjInit)
+{
+	BindToProperty(UE::MovieScene::FMovieSceneTracksComponentTypes::Get()->String);
+
+	if (HasAnyFlags(RF_ClassDefaultObject))
+	{
+		DefineImplicitPrerequisite(UStringChannelEvaluatorSystem::StaticClass(), GetClass());
+	}
+}
+
+void UMovieSceneStringPropertySystem::OnRun(FSystemTaskPrerequisites& InPrerequisites, FSystemSubsequentTasks& Subsequents)
+{
+	Super::OnRun(InPrerequisites, Subsequents);
+}
+
diff --git a/Engine/Source/Runtime/MovieSceneTracks/Private/Systems/StringChannelEvaluatorSystem.cpp b/Engine/Source/Runtime/MovieSceneTracks/Private/Systems/StringChannelEvaluatorSystem.cpp
new file mode 100644
--- /dev/null
+++ b/Engine/Source/Runtime/MovieSceneTracks/Private/Systems/StringChannelEvaluatorSystem.cpp
@@ -0,0 +1,58 @@
+// Copyright Epic Games, Inc. All Rights Reserved.
+
+#include "Systems/StringChannelEvaluatorSystem.h"
+#include "EntitySystem/BuiltInComponentTypes.h"
+#include "EntitySystem/MovieSceneEntitySystemTask.h"
+#include "EntitySystem/MovieSceneEntitySystemLinker.h"
+
+#include "EntitySystem/MovieSceneEvalTimeSystem.h"
+#include "Channels/MovieSceneStringChannel.h"
+
+#include UE_INLINE_GENERATED_CPP_BY_NAME(StringChannelEvaluatorSystem)
+
+DECLARE_CYCLE_STAT(TEXT("MovieScene: Evaluate string channels"), MovieSceneEval_EvaluateStringChannelTask, STATGROUP_MovieSceneECS);
+
+UStringChannelEvaluatorSystem::UStringChannelEvaluatorSystem(const FObjectInitializer& ObjInit)
+	: Super(ObjInit)
+{
+	using namespace UE::MovieScene;
+
+	SystemCategories = EEntitySystemCategory::ChannelEvaluators;
+	RelevantComponent = FBuiltInComponentTypes::Get()->StringChannel;
+
+	if (HasAnyFlags(RF_ClassDefaultObject))
+	{
+		DefineImplicitPrerequisite(UMovieSceneEvalTimeSystem::StaticClass(), GetClass());
+	}
+}
+
+void UStringChannelEvaluatorSystem::OnRun(FSystemTaskPrerequisites& InPrerequisites, FSystemSubsequentTasks& Subsequents)
+{
+	using namespace UE::MovieScene;
+
+	FBuiltInComponentTypes* BuiltInComponents = FBuiltInComponentTypes::Get();
+
+	struct FEvaluateStringChannels
+	{
+		static void ForEachEntity(FSourceStringChannel StringChannel, FFrameTime FrameTime, FString& OutResult)
+		{
+			if (const FString* Value = StringChannel.Source->Evaluate(FrameTime))
+			{
+				OutResult = *Value;
+			}
+			else
+			{
+				OutResult = FString();
+			}
+		}
+	};
+
+	FEntityTaskBuilder()
+	.Read(BuiltInComponents->StringChannel)
+	.Read(BuiltInComponents->EvalTime)
+	.Write(BuiltInComponents->StringResult)
+	.FilterNone({ BuiltInComponents->Tags.Ignored })
+	.SetStat(GET_STATID(MovieSceneEval_EvaluateStringChannelTask))
+	.Dispatch_PerEntity<FEvaluateStringChannels>(&Linker->EntityManager, InPrerequisites, &Subsequents);
+}
+
diff --git a/Engine/Source/Runtime/MovieSceneTracks/Private/Tracks/MovieSceneStringTrack.cpp b/Engine/Source/Runtime/MovieSceneTracks/Private/Tracks/MovieSceneStringTrack.cpp
--- a/Engine/Source/Runtime/MovieSceneTracks/Private/Tracks/MovieSceneStringTrack.cpp
+++ b/Engine/Source/Runtime/MovieSceneTracks/Private/Tracks/MovieSceneStringTrack.cpp
@@ -3,7 +3,6 @@
 #include "Tracks/MovieSceneStringTrack.h"
 #include "MovieSceneCommonHelpers.h"
 #include "Sections/MovieSceneStringSection.h"
-#include "Evaluation/MovieScenePropertyTemplates.h"
 
 #include UE_INLINE_GENERATED_CPP_BY_NAME(MovieSceneStringTrack)
 
@@ -25,17 +24,13 @@ bool UMovieSceneStringTrack::SupportsType(TSubclassOf<UMovieSceneSection> Sectio
 	return SectionClass == UMovieSceneStringSection::StaticClass();
 }
 
+
 UMovieSceneSection* UMovieSceneStringTrack::CreateNewSection()
 {
 	return NewObject<UMovieSceneStringSection>(this, NAME_None, RF_Transactional);
 }
 
 
-FMovieSceneEvalTemplatePtr UMovieSceneStringTrack::CreateTemplateForSection(const UMovieSceneSection& InSection) const
-{
-	return FMovieSceneStringPropertySectionTemplate(*CastChecked<UMovieSceneStringSection>(&InSection), *this);
-}
-
 const TArray<UMovieSceneSection*>& UMovieSceneStringTrack::GetAllSections() const
 {
 	return Sections;
diff --git a/Engine/Source/Runtime/MovieSceneTracks/Public/Channels/MovieSceneStringChannel.h b/Engine/Source/Runtime/MovieSceneTracks/Public/Channels/MovieSceneStringChannel.h
--- a/Engine/Source/Runtime/MovieSceneTracks/Public/Channels/MovieSceneStringChannel.h
+++ b/Engine/Source/Runtime/MovieSceneTracks/Public/Channels/MovieSceneStringChannel.h
@@ -68,6 +68,14 @@ struct MOVIESCENETRACKS_API FMovieSceneStringChannel : public FMovieSceneChannel
 		return Times;
 	}
 
+	/**
+	 * Returns whether this channel has any values
+	 */
+	FORCEINLINE bool HasAnyData() const
+	{
+		return Times.Num() != 0 || bHasDefaultValue == true;
+	}
+
 	/**
 	 * Evaluate this channel
 	 *
diff --git a/Engine/Source/Runtime/MovieSceneTracks/Public/MovieSceneTracksComponentTypes.h b/Engine/Source/Runtime/MovieSceneTracks/Public/MovieSceneTracksComponentTypes.h
--- a/Engine/Source/Runtime/MovieSceneTracks/Public/MovieSceneTracksComponentTypes.h
+++ b/Engine/Source/Runtime/MovieSceneTracks/Public/MovieSceneTracksComponentTypes.h
@@ -412,6 +412,7 @@ using FDoublePropertyTraits             = TDirectPropertyTraits<double>;
 using FTransformPropertyTraits          = TIndirectPropertyTraits<FTransform, FIntermediate3DTransform>;
 using FEulerTransformPropertyTraits     = TIndirectPropertyTraits<FEulerTransform, FIntermediate3DTransform>;
 using FComponentTransformPropertyTraits = TDirectPropertyTraits<FIntermediate3DTransform>;
+using FStringPropertyTraits			    = TDirectPropertyTraits<FString>;
 
 using FFloatParameterTraits             = TIndirectPropertyTraits<float, double>;
 using FColorParameterTraits             = TIndirectPropertyTraits<FLinearColor, FIntermediateColor>;
@@ -432,6 +433,7 @@ struct MOVIESCENETRACKS_API FMovieSceneTracksComponentTypes
 	TPropertyComponents<FTransformPropertyTraits> Transform;
 	TPropertyComponents<FEulerTransformPropertyTraits> EulerTransform;
 	TPropertyComponents<FComponentTransformPropertyTraits> ComponentTransform;
+	TPropertyComponents<FStringPropertyTraits> String;
 
 	TPropertyComponents<FFloatParameterTraits> FloatParameter;
 	TPropertyComponents<FColorParameterTraits> ColorParameter;
diff --git a/Engine/Source/Runtime/MovieSceneTracks/Public/Sections/MovieSceneStringSection.h b/Engine/Source/Runtime/MovieSceneTracks/Public/Sections/MovieSceneStringSection.h
--- a/Engine/Source/Runtime/MovieSceneTracks/Public/Sections/MovieSceneStringSection.h
+++ b/Engine/Source/Runtime/MovieSceneTracks/Public/Sections/MovieSceneStringSection.h
@@ -4,6 +4,7 @@
 
 #include "Channels/MovieSceneStringChannel.h"
 #include "CoreMinimal.h"
+#include "EntitySystem/IMovieSceneEntityProvider.h"
 #include "MovieSceneSection.h"
 #include "UObject/ObjectMacros.h"
 #include "UObject/UObjectGlobals.h"
@@ -18,6 +19,7 @@ class UObject;
 UCLASS(MinimalAPI)
 class UMovieSceneStringSection
 	: public UMovieSceneSection
+	, public IMovieSceneEntityProvider
 {
 	GENERATED_UCLASS_BODY()
 
@@ -28,6 +30,10 @@ public:
 	 */
 	const FMovieSceneStringChannel& GetChannel() const { return StringCurve; }
 
+	//~ IMovieSceneEntityProvider interface
+	virtual bool PopulateEvaluationFieldImpl(const TRange<FFrameNumber>& EffectiveRange, const FMovieSceneEvaluationFieldEntityMetaData& InMetaData, FMovieSceneEntityComponentFieldBuilder* OutFieldBuilder) override;
+	virtual void ImportEntityImpl(UMovieSceneEntitySystemLinker* EntityLinker, const FEntityImportParams& Params, FImportedEntity* OutImportedEntity) override;
+
 private:
 
 	/** Curve data */
diff --git a/Engine/Source/Runtime/MovieSceneTracks/Public/Systems/MovieSceneStringPropertySystem.h b/Engine/Source/Runtime/MovieSceneTracks/Public/Systems/MovieSceneStringPropertySystem.h
new file mode 100644
--- /dev/null
+++ b/Engine/Source/Runtime/MovieSceneTracks/Public/Systems/MovieSceneStringPropertySystem.h
@@ -0,0 +1,18 @@
+// Copyright Epic Games, Inc. All Rights Reserved.
+
+#pragma once
+
+#include "Systems/MovieScenePropertySystem.h"
+#include "MovieSceneStringPropertySystem.generated.h"
+
+
+
+UCLASS()
+class MOVIESCENETRACKS_API UMovieSceneStringPropertySystem : public UMovieScenePropertySystem
+{
+	GENERATED_BODY()
+
+	UMovieSceneStringPropertySystem(const FObjectInitializer& ObjInit);
+
+	virtual void OnRun(FSystemTaskPrerequisites& InPrerequisites, FSystemSubsequentTasks& Subsequents) override;
+};
diff --git a/Engine/Source/Runtime/MovieSceneTracks/Public/Systems/StringChannelEvaluatorSystem.h b/Engine/Source/Runtime/MovieSceneTracks/Public/Systems/StringChannelEvaluatorSystem.h
new file mode 100644
--- /dev/null
+++ b/Engine/Source/Runtime/MovieSceneTracks/Public/Systems/StringChannelEvaluatorSystem.h
@@ -0,0 +1,26 @@
+// Copyright Epic Games, Inc. All Rights Reserved.
+
+#pragma once
+
+#include "EntitySystem/MovieSceneEntitySystem.h"
+#include "UObject/ObjectMacros.h"
+#include "UObject/UObjectGlobals.h"
+
+#include "StringChannelEvaluatorSystem.generated.h"
+
+class UObject;
+
+/**
+ * System that is responsible for evaluating string channels.
+ */
+UCLASS(MinimalAPI)
+class UStringChannelEvaluatorSystem : public UMovieSceneEntitySystem
+{
+public:
+
+	GENERATED_BODY()
+
+	UStringChannelEvaluatorSystem(const FObjectInitializer& ObjInit);
+
+	virtual void OnRun(FSystemTaskPrerequisites& InPrerequisites, FSystemSubsequentTasks& Subsequents) override;
+};
diff --git a/Engine/Source/Runtime/MovieSceneTracks/Public/Tracks/MovieSceneStringTrack.h b/Engine/Source/Runtime/MovieSceneTracks/Public/Tracks/MovieSceneStringTrack.h
--- a/Engine/Source/Runtime/MovieSceneTracks/Public/Tracks/MovieSceneStringTrack.h
+++ b/Engine/Source/Runtime/MovieSceneTracks/Public/Tracks/MovieSceneStringTrack.h
@@ -5,7 +5,6 @@
 #include "CoreMinimal.h"
 #include "UObject/ObjectMacros.h"
 #include "Tracks/MovieScenePropertyTrack.h"
-#include "Compilation/IMovieSceneTrackTemplateProducer.h"
 #include "MovieSceneStringTrack.generated.h"
 
 /**
@@ -14,7 +13,6 @@
 UCLASS(MinimalAPI)
 class UMovieSceneStringTrack
 	: public UMovieScenePropertyTrack
-	, public IMovieSceneTrackTemplateProducer
 {
 	GENERATED_BODY()
 
@@ -35,7 +33,6 @@ public:
 	virtual void AddSection(UMovieSceneSection& Section) override;
 	virtual bool SupportsType(TSubclassOf<UMovieSceneSection> SectionClass) const override;
 	virtual UMovieSceneSection* CreateNewSection() override;
-	virtual FMovieSceneEvalTemplatePtr CreateTemplateForSection(const UMovieSceneSection& InSection) const override;
 	virtual const TArray<UMovieSceneSection*>& GetAllSections() const override;
 	virtual bool HasSection(const UMovieSceneSection& Section) const override;
 	virtual bool IsEmpty() const override;
