Sequencer: upgrade audio track to use ECS evaluation

The audio track was previously using the legacy evaluation templates. Changes include:

- The usual removal of the evaluation template and its references on the audio track.
- The addition of an audio system with associated components, and the audio section importing entities with these components.
- The legacy audio evaluation template was using shared track data to store information about audio components and so on. The audio system is using shared data on itself to the same effect, with the "evaluate audio" task asking it to create/retrieve audio components as needed.
- Some audio assets can have animated parameters, so the audio section has an arbitrary number of channels of different types. An audio section must therefore use as many entities as necessary to store all those channels and their results. The good news is that now these channels are evaluated more efficiently.
- The audio system  runs a "gather pass" to retrive all those parameter channels' result values, and correlate them with the main audio evaluation entities.
- The main chunks of the legacy audio evaluation template are largely preserved (code for creating audio components, processing audio, handling desynchronization, etc.)
- There is a new cvar called "Sequencer.Audio.PlayAudioWhenPlaybackJumps" which plays short audio when jumping time, similarly to when scrubbing. This is useful for debugging.

diff --git a/Engine/Source/Runtime/MovieScene/Public/Channels/MovieSceneAudioTriggerChannel.h b/Engine/Source/Runtime/MovieScene/Public/Channels/MovieSceneAudioTriggerChannel.h
--- a/Engine/Source/Runtime/MovieScene/Public/Channels/MovieSceneAudioTriggerChannel.h
+++ b/Engine/Source/Runtime/MovieScene/Public/Channels/MovieSceneAudioTriggerChannel.h
@@ -79,7 +79,6 @@ struct MOVIESCENE_API FMovieSceneAudioTriggerChannel : public FMovieSceneChannel
 	 * @param OutTriggered   A Value that if set to true combined with the function returning true means there's a trigger
 	 * @return true if the channel was evaluated successfully, false otherwise
 	 */
-;
 	bool EvaluatePossibleTriggers(const FMovieSceneContext& InContext, FMoveSceneAudioTriggerState& InState, bool& OutTriggered) const;
 
 public:
diff --git a/Engine/Source/Runtime/MovieSceneTracks/Private/Evaluation/MovieSceneAudioTemplate.cpp b/Engine/Source/Runtime/MovieSceneTracks/Private/Evaluation/MovieSceneAudioTemplate.cpp
deleted file mode 100644
diff --git a/Engine/Source/Runtime/MovieSceneTracks/Private/Evaluation/MovieSceneAudioTemplate.h b/Engine/Source/Runtime/MovieSceneTracks/Private/Evaluation/MovieSceneAudioTemplate.h
deleted file mode 100644
diff --git a/Engine/Source/Runtime/MovieSceneTracks/Private/MovieSceneTracksComponentTypes.cpp b/Engine/Source/Runtime/MovieSceneTracks/Private/MovieSceneTracksComponentTypes.cpp
--- a/Engine/Source/Runtime/MovieSceneTracks/Private/MovieSceneTracksComponentTypes.cpp
+++ b/Engine/Source/Runtime/MovieSceneTracks/Private/MovieSceneTracksComponentTypes.cpp
@@ -891,6 +891,10 @@ FMovieSceneTracksComponentTypes::FMovieSceneTracksComponentTypes()
 
 	ComponentRegistry->NewComponentType(&Fade,                   TEXT("Fade"), EComponentTypeFlags::CopyToChildren);
 
+	ComponentRegistry->NewComponentType(&Audio,                  TEXT("Audio"), EComponentTypeFlags::CopyToChildren);
+	ComponentRegistry->NewComponentType(&AudioInputs,            TEXT("Audio Inputs"), EComponentTypeFlags::CopyToChildren);
+	ComponentRegistry->NewComponentType(&AudioTriggerName,       TEXT("Audio Trigger Name"), EComponentTypeFlags::CopyToChildren);
+
 	Tags.BoundMaterialChanged = ComponentRegistry->NewTag(TEXT("Bound Material Changed"));
 	FBuiltInComponentTypes::Get()->RequiresInstantiationMask.Set(Tags.BoundMaterialChanged);
 
diff --git a/Engine/Source/Runtime/MovieSceneTracks/Private/Sections/MovieSceneAudioSection.cpp b/Engine/Source/Runtime/MovieSceneTracks/Private/Sections/MovieSceneAudioSection.cpp
--- a/Engine/Source/Runtime/MovieSceneTracks/Private/Sections/MovieSceneAudioSection.cpp
+++ b/Engine/Source/Runtime/MovieSceneTracks/Private/Sections/MovieSceneAudioSection.cpp
@@ -10,6 +10,8 @@
 #include "MovieSceneCommonHelpers.h"
 #include "Misc/FrameRate.h"
 #include "Misc/GeneratedTypeName.h"
+#include "EntitySystem/BuiltInComponentTypes.h"
+#include "MovieSceneTracksComponentTypes.h"
 
 #include UE_INLINE_GENERATED_CPP_BY_NAME(MovieSceneAudioSection)
 
@@ -39,6 +41,28 @@ namespace
 	}
 }
 
+namespace UE::MovieScene
+{
+
+/* 
+ * Entity IDs are an encoded type and index, with the upper bit being the type (scalar iputs vs audio trigger),
+ * and the lower 31 bits as the entity index
+ */
+enum class EAudioSectionEntityType : uint8 { MainEntity, InputsEntity, TriggerEntity };
+
+uint32 EncodeEntityID(int32 InIndex, EAudioSectionEntityType InEntityType)
+{
+	check(InIndex >= 0 && InIndex < int32(0x00FFFFFF));
+	return static_cast<uint32>(InIndex) | ((uint8)InEntityType << 24);
+}
+void DecodeEntityID(uint32 InEntityID, int32& OutIndex, EAudioSectionEntityType& OutEntityType)
+{
+	OutIndex = static_cast<int32>(InEntityID & 0x00FFFFFF);
+	OutEntityType = (EAudioSectionEntityType)(InEntityID >> 24);
+}
+
+}
+
 UMovieSceneAudioSection::UMovieSceneAudioSection( const FObjectInitializer& ObjectInitializer )
 	: Super( ObjectInitializer )
 {
@@ -150,6 +174,17 @@ void UMovieSceneAudioSection::SetupSoundInputParameters(const USoundBase* InSoun
 		TArray<FAudioParameter> DefaultParams;
 		InSoundBase->GetAllDefaultParameters(DefaultParams);
 
+		TSet<FName> OrphanedFloatInputs;
+		Inputs_Float.GetKeys(OrphanedFloatInputs);
+		TSet<FName> OrphanedTriggerInputs;
+		Inputs_Trigger.GetKeys(OrphanedTriggerInputs);
+		TSet<FName> OrphanedBoolInputs;
+		Inputs_Bool.GetKeys(OrphanedBoolInputs);
+		TSet<FName> OrphanedIntInputs;
+		Inputs_Int.GetKeys(OrphanedIntInputs);
+		TSet<FName> OrphanedStringInputs;
+		Inputs_String.GetKeys(OrphanedStringInputs);
+
 		for (const FAudioParameter& Param : DefaultParams)
 		{
 			switch (Param.ParamType)
@@ -157,26 +192,31 @@ void UMovieSceneAudioSection::SetupSoundInputParameters(const USoundBase* InSoun
 			case EAudioParameterType::Float:
 			{
 				Inputs_Float.FindOrAdd(Param.ParamName, FMovieSceneFloatChannel{}).SetDefault(Param.FloatParam);
+				OrphanedFloatInputs.Remove(Param.ParamName);
 				break;
 			}
 			case EAudioParameterType::Trigger:
 			{
 				Inputs_Trigger.FindOrAdd(Param.ParamName, FMovieSceneAudioTriggerChannel{});
+				OrphanedTriggerInputs.Remove(Param.ParamName);
 				break;
 			}
 			case EAudioParameterType::Boolean:
 			{
 				Inputs_Bool.FindOrAdd(Param.ParamName, FMovieSceneBoolChannel{}).SetDefault(Param.BoolParam);
+				OrphanedBoolInputs.Remove(Param.ParamName);
 				break;
 			}
 			case EAudioParameterType::Integer:
 			{
 				Inputs_Int.FindOrAdd(Param.ParamName, FMovieSceneIntegerChannel{}).SetDefault(Param.IntParam);
+				OrphanedIntInputs.Remove(Param.ParamName);
 				break;
 			}
 			case EAudioParameterType::String:
 			{
 				Inputs_String.FindOrAdd(Param.ParamName, FMovieSceneStringChannel{}).SetDefault(Param.StringParam);
+				OrphanedStringInputs.Remove(Param.ParamName);
 				break;
 			}
 			default:
@@ -184,6 +224,27 @@ void UMovieSceneAudioSection::SetupSoundInputParameters(const USoundBase* InSoun
 				break;
 			}
 		}
+
+		for (const FName& Name : OrphanedFloatInputs)
+		{
+			Inputs_Float.Remove(Name);
+		}
+		for (const FName& Name : OrphanedTriggerInputs)
+		{
+			Inputs_Trigger.Remove(Name);
+		}
+		for (const FName& Name : OrphanedBoolInputs)
+		{
+			Inputs_Bool.Remove(Name);
+		}
+		for (const FName& Name : OrphanedIntInputs)
+		{
+			Inputs_Int.Remove(Name);
+		}
+		for (const FName& Name : OrphanedStringInputs)
+		{
+			Inputs_String.Remove(Name);
+		}
 	}
 }
 
@@ -360,4 +421,195 @@ USceneComponent* UMovieSceneAudioSection::GetAttachComponent(const AActor* InPar
 	}
 }
 
+bool UMovieSceneAudioSection::PopulateEvaluationFieldImpl(const TRange<FFrameNumber>& EffectiveRange, const FMovieSceneEvaluationFieldEntityMetaData& InMetaData, FMovieSceneEntityComponentFieldBuilder* OutFieldBuilder)
+{
+	using namespace UE::MovieScene;
+
+	int32 MetaDataIndex = OutFieldBuilder->AddMetaData(InMetaData);
+
+	// Add the default entity first.
+	int32 MainEntityIndex = OutFieldBuilder->FindOrAddEntity(this, EncodeEntityID(0, EAudioSectionEntityType::MainEntity));
+	OutFieldBuilder->AddPersistentEntity(EffectiveRange, MainEntityIndex, MetaDataIndex);
+
+	// See how many additional entities we need to store the audio input data.
+	// We can pack 9 float channels per entity, but we can only have one string/bool/int/audio-trigger
+	// channel per entity.
+	int32 NumInputDataEntities = Inputs_Float.Num() % 9;
+	NumInputDataEntities = FMath::Max(NumInputDataEntities, Inputs_String.Num());
+	NumInputDataEntities = FMath::Max(NumInputDataEntities, Inputs_Bool.Num());
+	NumInputDataEntities = FMath::Max(NumInputDataEntities, Inputs_Int.Num());
+
+	// Add these extra entities to the evaluation field.
+	for (int32 InputDataEntity = 0; InputDataEntity < NumInputDataEntities; ++InputDataEntity)
+	{
+		int32 EntityIndex = OutFieldBuilder->FindOrAddEntity(this, EncodeEntityID(InputDataEntity, EAudioSectionEntityType::InputsEntity));
+		OutFieldBuilder->AddPersistentEntity(EffectiveRange, EntityIndex, MetaDataIndex);
+	}
+
+	// Audio triggers are added differently, as one-shot entities.
+	TArray<FName> AudioTriggerNames;
+	Inputs_Trigger.GetKeys(AudioTriggerNames);
+	AudioTriggerNames.Sort(FNameLexicalLess());
+	for (int32 TriggerIndex = 0; TriggerIndex < AudioTriggerNames.Num(); ++TriggerIndex)
+	{
+		FMovieSceneAudioTriggerChannel& TriggerChannel = Inputs_Trigger[AudioTriggerNames[TriggerIndex]];
+		TArrayView<const FFrameNumber> Times = TriggerChannel.GetTimes();
+		for (int32 Index = 0; Index < Times.Num(); ++Index)
+		{
+			if (EffectiveRange.Contains(Times[Index]))
+			{
+				TRange<FFrameNumber> TriggerRange(Times[Index]);
+				int32 EntityIndex = OutFieldBuilder->FindOrAddEntity(this, EncodeEntityID(TriggerIndex, EAudioSectionEntityType::TriggerEntity));
+				OutFieldBuilder->AddOneShotEntity(TriggerRange, EntityIndex, MetaDataIndex);
+			}
+		}
+	}
+
+	// Return true to indicate we've done everything ourselves.
+	return true;
+}
+
+void UMovieSceneAudioSection::ImportEntityImpl(UMovieSceneEntitySystemLinker* EntityLinker, const FEntityImportParams& Params, FImportedEntity* OutImportedEntity)
+{
+	using namespace UE::MovieScene;
+
+	if (!Sound)
+	{
+		return;
+	}
+
+	const FBuiltInComponentTypes* BuiltInComponents = FBuiltInComponentTypes::Get();
+	const FMovieSceneTracksComponentTypes* TrackComponents = FMovieSceneTracksComponentTypes::Get();
+
+	const FGuid ObjectBindingID = Params.GetObjectBindingID();
+
+	int32 EntityIndex;
+	EAudioSectionEntityType EntityType;
+	DecodeEntityID(Params.EntityID, EntityIndex, EntityType);
+
+	if (EntityType == EAudioSectionEntityType::MainEntity)
+	{
+		// Default entity... we add the main audio component data, plus the volume and pitch channels.
+		OutImportedEntity->AddBuilder(
+			FEntityBuilder()
+			.AddConditional(BuiltInComponents->GenericObjectBinding, ObjectBindingID, ObjectBindingID.IsValid())
+			.AddTagConditional(BuiltInComponents->Tags.Root, !ObjectBindingID.IsValid())
+			.Add(TrackComponents->Audio, FMovieSceneAudioComponentData{ this })
+			.Add(BuiltInComponents->FloatChannel[0], &SoundVolume)
+			.Add(BuiltInComponents->FloatChannel[1], &PitchMultiplier)
+		);
+	}
+	else if (EntityType == EAudioSectionEntityType::InputsEntity)
+	{
+		// Additional entities for custom audio input values.
+		TArray<FName> InputNames;
+		FMovieSceneAudioInputData InputData;
+
+		// There are up to 9 float channels per entity, and we know that all entities are fully packed
+		// until the last one. So add as many as we can for the given entity we're building.
+		int32 FloatInputStartIndex = EntityIndex * 9;
+		if (FloatInputStartIndex < Inputs_Float.Num())
+		{
+			int32 FloatInputNum = FMath::Min(9, Inputs_Float.Num() - FloatInputStartIndex);
+			Inputs_Float.GetKeys(InputNames);
+			for (int32 Offset = 0; Offset < FloatInputNum; ++Offset)
+			{
+				InputData.FloatInputs[Offset] = InputNames[FloatInputStartIndex + Offset];
+			}
+		}
+
+		// Other inputs can only be added once per entity, so add one of each type that exists.
+		int32 OtherInputStartIndex = EntityIndex;
+		if (OtherInputStartIndex < Inputs_String.Num())
+		{
+			InputNames.Reset();
+			Inputs_String.GetKeys(InputNames);
+			InputData.StringInput = InputNames[OtherInputStartIndex];
+		}
+		if (OtherInputStartIndex < Inputs_Bool.Num())
+		{
+			InputNames.Reset();
+			Inputs_Bool.GetKeys(InputNames);
+			InputData.BoolInput = InputNames[OtherInputStartIndex];
+		}
+		if (OtherInputStartIndex < Inputs_Int.Num())
+		{
+			InputNames.Reset();
+			Inputs_Int.GetKeys(InputNames);
+			InputData.IntInput = InputNames[OtherInputStartIndex];
+		}
+
+		// Make this additional entity by adding the component that specifies what audio input channels
+		// are present, plus all of these channels.
+		OutImportedEntity->AddBuilder(
+			FEntityBuilder()
+			.AddConditional(BuiltInComponents->GenericObjectBinding, ObjectBindingID, ObjectBindingID.IsValid())
+			.AddTagConditional(BuiltInComponents->Tags.Root, !ObjectBindingID.IsValid())
+			.Add(TrackComponents->Audio, FMovieSceneAudioComponentData{ this })
+			.Add(TrackComponents->AudioInputs, InputData)
+		);
+		for (int32 Index = 0; Index < 9; ++Index)
+		{
+			FName InputName = InputData.FloatInputs[Index];
+			if (!InputName.IsNone())
+			{
+				OutImportedEntity->AddBuilder(
+					FEntityBuilder()
+					.Add(BuiltInComponents->FloatChannel[Index], &Inputs_Float[InputName])
+				);
+			}
+		}
+		if (!InputData.StringInput.IsNone())
+		{
+			OutImportedEntity->AddBuilder(
+				FEntityBuilder()
+				.Add(BuiltInComponents->StringChannel, &Inputs_String[InputData.StringInput])
+			);
+		}
+		if (!InputData.BoolInput.IsNone())
+		{
+			OutImportedEntity->AddBuilder(
+				FEntityBuilder()
+				.Add(BuiltInComponents->BoolChannel, &Inputs_Bool[InputData.BoolInput])
+			);
+		}
+		if (!InputData.IntInput.IsNone())
+		{
+			OutImportedEntity->AddBuilder(
+				FEntityBuilder()
+				.Add(BuiltInComponents->IntegerChannel, &Inputs_Int[InputData.IntInput])
+			);
+		}
+	}
+	else if (EntityType == EAudioSectionEntityType::TriggerEntity)
+	{
+		// Additional one-shot entities for audio triggers.
+		// The decoded index is the index of the name in the triggers map.
+		TArray<FName> AudioTriggerNames;
+		Inputs_Trigger.GetKeys(AudioTriggerNames);
+		AudioTriggerNames.Sort(FNameLexicalLess());
+		OutImportedEntity->AddBuilder(
+			FEntityBuilder()
+			.AddConditional(BuiltInComponents->GenericObjectBinding, ObjectBindingID, ObjectBindingID.IsValid())
+			.AddTagConditional(BuiltInComponents->Tags.Root, !ObjectBindingID.IsValid())
+			.Add(TrackComponents->Audio, FMovieSceneAudioComponentData{ this })
+			.Add(TrackComponents->AudioTriggerName, AudioTriggerNames[EntityIndex])
+		);
+	}
+}
+
+#if WITH_EDITOR
+
+void UMovieSceneAudioSection::PostEditChangeProperty(FPropertyChangedEvent& PropertyChangedEvent)
+{
+	const FName PropertyName = PropertyChangedEvent.GetPropertyName();
+	if (PropertyName == GET_MEMBER_NAME_CHECKED(UMovieSceneAudioSection, Sound))
+	{
+		// Invalidate channel proxy
+		ChannelProxy = nullptr;
+	}
+
+	Super::PostEditChangeProperty(PropertyChangedEvent);
+}
 
+#endif
diff --git a/Engine/Source/Runtime/MovieSceneTracks/Private/Systems/MovieSceneAudioSystem.cpp b/Engine/Source/Runtime/MovieSceneTracks/Private/Systems/MovieSceneAudioSystem.cpp
new file mode 100644
--- /dev/null
+++ b/Engine/Source/Runtime/MovieSceneTracks/Private/Systems/MovieSceneAudioSystem.cpp
@@ -0,0 +1,947 @@
+// Copyright Epic Games, Inc. All Rights Reserved.
+
+#include "Systems/MovieSceneAudioSystem.h"
+
+#include "AudioDevice.h"
+#include "Components/AudioComponent.h"
+#include "Engine/Engine.h"
+#include "Engine/EngineTypes.h"
+#include "EntitySystem/BuiltInComponentTypes.h"
+#include "Evaluation/PreAnimatedState/MovieScenePreAnimatedObjectStorage.h"
+#include "Evaluation/PreAnimatedState/MovieScenePreAnimatedStateStorage.h"
+#include "Evaluation/PreAnimatedState/MovieScenePreAnimatedStorageID.inl"
+#include "GameFramework/WorldSettings.h"
+#include "IMovieScenePlayer.h"
+#include "MovieScene.h"
+#include "MovieSceneTracksComponentTypes.h"
+#include "Sections/MovieSceneAudioSection.h"
+#include "Sound/SoundAttenuation.h"
+#include "Sound/SoundCue.h"
+#include "Tracks/MovieSceneAudioTrack.h"
+
+#include UE_INLINE_GENERATED_CPP_BY_NAME(MovieSceneAudioSystem)
+
+DECLARE_CYCLE_STAT(TEXT("Audio System Evaluate"), MovieSceneEval_AudioTasks, STATGROUP_MovieSceneECS);
+
+static float MaxSequenceAudioDesyncToleranceCVar = 0.5f;
+FAutoConsoleVariableRef CVarMaxSequenceAudioDesyncTolerance(
+	TEXT("Sequencer.Audio.MaxDesyncTolerance"),
+	MaxSequenceAudioDesyncToleranceCVar,
+	TEXT("Controls how many seconds an audio track can be out of sync in a Sequence before we attempt a time correction.\n"),
+	ECVF_Default);
+
+static bool bIgnoreAudioSyncDuringWorldTimeDilationCVar = true;
+FAutoConsoleVariableRef CVarIgnoreAudioSyncDuringWorldTimeDilation(
+	TEXT("Sequencer.Audio.IgnoreAudioSyncDuringWorldTimeDilation"),
+	bIgnoreAudioSyncDuringWorldTimeDilationCVar,
+	TEXT("Ignore correcting audio if there is world time dilation.\n"),
+	ECVF_Default);
+
+static int32 UseAudioClockForSequencerDesyncCVar = 0;
+FAutoConsoleVariableRef CVaUseAudioClockForSequencerDesync(
+	TEXT("Sequencer.Audio.UseAudioClockForAudioDesync"),
+	UseAudioClockForSequencerDesyncCVar,
+	TEXT("When set to 1, we will use the audio render thread directly to query whether audio has went out of sync with the sequence.\n"),
+	ECVF_Default);
+
+static bool bPlayAudioWhenPlaybackJumps = false;
+FAutoConsoleVariableRef CVarPlayAudioWhenPlaybackJumps(
+	TEXT("Sequencer.Audio.PlayAudioWhenPlaybackJumps"),
+	bPlayAudioWhenPlaybackJumps,
+	TEXT("Play audio when playback jumps.\n"),
+	ECVF_Default);
+
+namespace UE::MovieScene
+{
+
+enum class EPreAnimatedAudioStateType
+{
+	/** Pre-animated state manages the lifespan of the audio component */
+	AudioComponentLifespan,
+	/** Pre-animated state manages whether the audio component is playing */
+	AudioPlaying
+};
+
+template<typename BaseTraits>
+struct FPreAnimatedAudioStateTraits : BaseTraits
+{
+	using KeyType = FObjectKey;
+	using StorageType = EPreAnimatedAudioStateType;
+
+	EPreAnimatedAudioStateType CachePreAnimatedValue(FObjectKey InKey)
+	{
+		check(false);
+		return EPreAnimatedAudioStateType::AudioComponentLifespan;
+	}
+
+	void RestorePreAnimatedValue(FObjectKey InKey, EPreAnimatedAudioStateType InStateType, const FRestoreStateParams& Params)
+	{
+		if (UAudioComponent* AudioComponent = Cast<UAudioComponent>(InKey.ResolveObjectPtr()))
+		{
+			switch (InStateType)
+			{
+				case EPreAnimatedAudioStateType::AudioPlaying:
+					AudioComponent->Stop();
+					break;
+				case EPreAnimatedAudioStateType::AudioComponentLifespan:
+					AudioComponent->DestroyComponent();
+					break;
+			}
+		}
+	}
+};
+
+using FPreAnimatedBoundObjectAudioStateTraits = FPreAnimatedAudioStateTraits<FBoundObjectPreAnimatedStateTraits>;
+
+struct FPreAnimatedAudioStorage : TPreAnimatedStateStorage_ObjectTraits<FPreAnimatedBoundObjectAudioStateTraits>
+{
+	static TAutoRegisterPreAnimatedStorageID<FPreAnimatedAudioStorage> StorageID;
+};
+TAutoRegisterPreAnimatedStorageID<FPreAnimatedAudioStorage> FPreAnimatedAudioStorage::StorageID;
+
+/**
+ * Types of audio evaluation we should run for a given sequence.
+ */
+enum class EAudioEvaluationType
+{
+	Skip,
+	Play,
+	StopAndPlay,
+	Stop
+};
+
+struct FGatherAudioInputs
+{
+	using FAudioInputsBySectionKey = UMovieSceneAudioSystem::FAudioInputsBySectionKey;
+	using FAudioComponentInputEvaluationData = UMovieSceneAudioSystem::FAudioComponentInputEvaluationData;
+
+	UMovieSceneAudioSystem* AudioSystem;
+
+	FGatherAudioInputs(UMovieSceneAudioSystem* InAudioSystem)
+		: AudioSystem(InAudioSystem)
+	{
+	}
+
+	void ForEachAllocation(
+		const FEntityAllocation* Allocation,
+		TRead<FMovieSceneAudioComponentData> AudioDatas,
+		TRead<FMovieSceneAudioInputData> AudioInputDatas,
+		TReadOneOrMoreOf<
+				double, double, double,
+				double, double, double,
+				double, double, double,
+				FString, int32, bool>
+			AudioInputResults)
+	{
+		FAudioInputsBySectionKey& AudioInputsBySectionKey = AudioSystem->AudioInputsBySectionKey;
+
+		const double* DoubleResults[9];
+		{
+			DoubleResults[0] = AudioInputResults.Get<0>();
+			DoubleResults[1] = AudioInputResults.Get<1>();
+			DoubleResults[2] = AudioInputResults.Get<2>();
+			DoubleResults[3] = AudioInputResults.Get<3>();
+			DoubleResults[4] = AudioInputResults.Get<4>();
+			DoubleResults[5] = AudioInputResults.Get<5>();
+			DoubleResults[6] = AudioInputResults.Get<6>();
+			DoubleResults[7] = AudioInputResults.Get<7>();
+			DoubleResults[8] = AudioInputResults.Get<8>();
+		}
+		const FString* StringResults = AudioInputResults.Get<9>();
+		const int32* IntegerResults = AudioInputResults.Get<10>();
+		const bool* BoolResults = AudioInputResults.Get<11>();
+
+		for (int32 Index = 0; Index < Allocation->Num(); ++Index)
+		{
+			const FMovieSceneAudioComponentData& AudioData = AudioDatas[Index];
+			const FMovieSceneAudioInputData& AudioInputNames = AudioInputDatas[Index];
+
+			FObjectKey SectionKey(AudioData.Section);
+			FAudioComponentInputEvaluationData& AudioInputValues = AudioInputsBySectionKey.FindOrAdd(SectionKey);
+			
+			// Gather float inputs.
+			for (int32 FloatIndex = 0; FloatIndex < 9; ++FloatIndex)
+			{
+				if (!AudioInputNames.FloatInputs[FloatIndex].IsNone() && ensure(DoubleResults[FloatIndex]))
+				{
+					AudioInputValues.Inputs_Float.Add(AudioInputNames.FloatInputs[FloatIndex], DoubleResults[FloatIndex][Index]);
+				}
+			}
+
+			// Gather string inputs.
+			if (!AudioInputNames.StringInput.IsNone() && ensure(StringResults))
+			{
+				AudioInputValues.Inputs_String.Add(AudioInputNames.StringInput, StringResults[Index]);
+			}
+
+			// Gather integer inputs.
+			if (!AudioInputNames.IntInput.IsNone() && ensure(IntegerResults))
+			{
+				AudioInputValues.Inputs_Int.Add(AudioInputNames.IntInput, IntegerResults[Index]);
+			}
+
+			// Gather boolean inputs.
+			if (!AudioInputNames.BoolInput.IsNone() && ensure(BoolResults))
+			{
+				AudioInputValues.Inputs_Bool.Add(AudioInputNames.BoolInput, BoolResults[Index]);
+			}
+		}
+	}
+};
+
+struct FGatherAudioTriggers
+{
+	using FAudioInputsBySectionKey = UMovieSceneAudioSystem::FAudioInputsBySectionKey;
+	using FAudioComponentInputEvaluationData = UMovieSceneAudioSystem::FAudioComponentInputEvaluationData;
+
+	UMovieSceneAudioSystem* AudioSystem;
+
+	FGatherAudioTriggers(UMovieSceneAudioSystem* InAudioSystem)
+		: AudioSystem(InAudioSystem)
+	{
+	}
+
+	void ForEachAllocation(
+		const FEntityAllocation* Allocation,
+		TRead<FMovieSceneAudioComponentData> AudioDatas,
+		TRead<FName> AudioTriggerNames)
+	{
+		FAudioInputsBySectionKey& AudioInputsBySectionKey = AudioSystem->AudioInputsBySectionKey;
+
+		for (int32 Index = 0; Index < Allocation->Num(); ++Index)
+		{
+			const FMovieSceneAudioComponentData& AudioData = AudioDatas[Index];
+			const FName& AudioTriggerName = AudioTriggerNames[Index];
+
+			FObjectKey SectionKey(AudioData.Section);
+			FAudioComponentInputEvaluationData& AudioInputValues = AudioInputsBySectionKey.FindOrAdd(SectionKey);
+
+			AudioInputValues.Inputs_Trigger.Add(AudioTriggerName);
+		}
+	}
+};
+
+struct FEvaluateAudio
+{
+	static EAudioEvaluationType GetAudioEvaluationType(const FMovieSceneContext& Context)
+	{
+		if (Context.GetStatus() == EMovieScenePlayerStatus::Jumping &&
+				!bPlayAudioWhenPlaybackJumps)
+		{
+			return EAudioEvaluationType::Skip;
+		}
+
+		if (Context.HasJumped())
+		{
+			// If the status says we jumped, we always stop all sounds, then allow them to be played again 
+			// naturally if status == Playing (for example)
+			return EAudioEvaluationType::StopAndPlay;
+		}
+
+		if (
+			(Context.GetStatus() != EMovieScenePlayerStatus::Playing && 
+			 Context.GetStatus() != EMovieScenePlayerStatus::Scrubbing && 
+			 Context.GetStatus() != EMovieScenePlayerStatus::Stepping)
+			||
+			Context.GetDirection() == EPlayDirection::Backwards)
+		{
+			// stopped, recording, etc
+			return EAudioEvaluationType::Stop;
+		}
+
+		return EAudioEvaluationType::Play;
+	}
+
+	UMovieSceneAudioSystem* AudioSystem;
+	const FInstanceRegistry* InstanceRegistry;
+
+	FEvaluateAudio(UMovieSceneAudioSystem* InAudioSystem)
+		: AudioSystem(InAudioSystem)
+	{
+		InstanceRegistry = AudioSystem->GetLinker()->GetInstanceRegistry();
+	}
+
+	void ForEachAllocation(
+			const FEntityAllocation* Allocation, 
+			TRead<FMovieSceneEntityID> EntityIDs, 
+			TRead<FRootInstanceHandle> RootInstanceHandles,
+			TRead<FInstanceHandle> InstanceHandles,
+			TRead<FMovieSceneAudioComponentData> AudioDatas,
+			TRead<double> VolumeMultipliers,
+			TRead<double> PitchMultipliers,
+			TReadOptional<UObject*> BoundObjects)
+	{
+		const FBuiltInComponentTypes* BuiltInComponents = FBuiltInComponentTypes::Get();
+
+		const int32 Num = Allocation->Num();
+		const bool bWantsRestoreState = Allocation->HasComponent(BuiltInComponents->Tags.RestoreState);
+
+		for (int32 Index = 0; Index < Num; ++Index)
+		{
+			const FMovieSceneEntityID& EntityID = EntityIDs[Index];
+			const FRootInstanceHandle& RootInstanceHandle = RootInstanceHandles[Index];
+			const FInstanceHandle& InstanceHandle = InstanceHandles[Index];
+			const FMovieSceneAudioComponentData& AudioData = AudioDatas[Index];
+
+			const FSequenceInstance& Instance = InstanceRegistry->GetInstance(InstanceHandle);
+
+			double VolumeMultiplier = VolumeMultipliers[Index];
+			double PitchMultiplier = PitchMultipliers[Index];
+			UObject* BoundObject = (BoundObjects.IsValid() ? BoundObjects[Index] : nullptr);
+
+			Evaluate(EntityID, AudioData, Instance, RootInstanceHandle, VolumeMultiplier, PitchMultiplier, BoundObject, bWantsRestoreState);
+		}
+	}
+
+private:
+
+	void Evaluate(
+			const FMovieSceneEntityID& EntityID,
+			const FMovieSceneAudioComponentData& AudioData,
+			const FSequenceInstance& Instance,
+			const FRootInstanceHandle& RootInstanceHandle,
+			double VolumeMultiplier,
+			double PitchMultiplier,
+			UObject* BoundObject,
+			bool bWantsRestoreState)
+	{
+		const FMovieSceneContext& Context = Instance.GetContext();
+		IMovieScenePlayer* Player = Instance.GetPlayer();
+
+		UMovieSceneAudioSection* AudioSection = AudioData.Section;
+		if (!ensureMsgf(AudioSection, TEXT("No valid audio section found in audio track component data!")))
+		{
+			return;
+		}
+
+		FObjectKey ActorKey(BoundObject);
+		FObjectKey SectionKey(AudioSection);
+
+		const EAudioEvaluationType EvalType = GetAudioEvaluationType(Context);
+		if (EvalType == EAudioEvaluationType::StopAndPlay)
+		{
+			AudioSystem->StopSound(ActorKey, AudioData.Section);
+		}
+		else if (EvalType == EAudioEvaluationType::Stop)
+		{
+			AudioSystem->StopSound(ActorKey, AudioData.Section);
+			return;
+		}
+		else if (EvalType == EAudioEvaluationType::Skip)
+		{
+			return;
+		}
+
+		// Root audio track
+		if (BoundObject == nullptr)
+		{
+			UObject* PlaybackContext = Player->GetPlaybackContext();
+
+			const FMovieSceneActorReferenceData& AttachActorData = AudioSection->GetAttachActorData();
+
+			USceneComponent* AttachComponent = nullptr;
+			FMovieSceneActorReferenceKey AttachKey;
+			AttachActorData.Evaluate(Context.GetTime(), AttachKey);
+			FMovieSceneObjectBindingID AttachBindingID = AttachKey.Object;
+			if (AttachBindingID.IsValid())
+			{
+				// If the transform is set, otherwise use the bound actor's transform
+				for (TWeakObjectPtr<> WeakObject : AttachBindingID.ResolveBoundObjects(Instance.GetSequenceID(), *Player))
+				{
+					AActor* AttachActor = Cast<AActor>(WeakObject.Get());
+					if (AttachActor)
+					{
+						AttachComponent = AudioSection->GetAttachComponent(AttachActor, AttachKey);
+					}
+					if (AttachComponent)
+					{
+						break;
+					}
+				}
+			}
+
+			FAudioComponentEvaluationData* EvaluationData = AudioSystem->GetAudioComponentEvaluationData(FObjectKey(), SectionKey);
+			if (!EvaluationData)
+			{
+				// Initialize the sound
+				UWorld* World = PlaybackContext ? PlaybackContext->GetWorld() : nullptr;
+				EvaluationData = AudioSystem->AddRootAudioComponent(AudioSection, World);
+
+				if (EvaluationData)
+				{
+					UAudioComponent* AudioComponent = EvaluationData->AudioComponent.Get();
+					ensure(AudioComponent);
+
+					AudioSystem->PreAnimatedStorage->BeginTrackingEntity(EntityID, bWantsRestoreState, RootInstanceHandle, AudioComponent);
+					AudioSystem->PreAnimatedStorage->CachePreAnimatedValue(
+						FCachePreAnimatedValueParams(), AudioComponent,
+						[](FObjectKey InKey) { return EPreAnimatedAudioStateType::AudioComponentLifespan; });
+
+					if (AudioSection->GetOnQueueSubtitles().IsBound())
+					{
+						AudioComponent->OnQueueSubtitles = AudioSection->GetOnQueueSubtitles();
+					}
+					if (AudioSection->GetOnAudioFinished().IsBound())
+					{
+						AudioComponent->OnAudioFinished = AudioSection->GetOnAudioFinished();
+					}
+					if (AudioSection->GetOnAudioPlaybackPercent().IsBound())
+					{
+						AudioComponent->OnAudioPlaybackPercent = AudioSection->GetOnAudioPlaybackPercent();
+					}
+				}
+			}
+
+			if (EvaluationData)
+			{
+				UAudioComponent* AudioComponent = EvaluationData->AudioComponent.Get();
+
+				if (AudioComponent && AttachComponent && 
+					(AudioComponent->GetAttachParent() != AttachComponent || AudioComponent->GetAttachSocketName() != AttachKey.SocketName))
+				{
+					AudioComponent->AttachToComponent(AttachComponent, FAttachmentTransformRules::KeepRelativeTransform, AttachKey.SocketName);
+				}
+
+				EvaluationData->VolumeMultiplier = VolumeMultiplier;
+				EvaluationData->PitchMultiplier = PitchMultiplier;
+
+				EnsureAudioIsPlaying(nullptr, *AudioSection, *EvaluationData, Context, *Player);
+			}
+		}
+
+		// Object binding audio track
+		else
+		{
+			FAudioComponentEvaluationData* EvaluationData = AudioSystem->GetAudioComponentEvaluationData(ActorKey, SectionKey);
+			if (!EvaluationData)
+			{
+				// Initialize the sound
+				EvaluationData = AudioSystem->AddBoundObjectAudioComponent(AudioSection, BoundObject);
+
+				if (EvaluationData)
+				{
+					UAudioComponent* AudioComponent = EvaluationData->AudioComponent.Get();
+					ensure(AudioComponent);
+
+					AudioSystem->PreAnimatedStorage->BeginTrackingEntity(EntityID, bWantsRestoreState, RootInstanceHandle, AudioComponent);
+					AudioSystem->PreAnimatedStorage->CachePreAnimatedValue(
+						FCachePreAnimatedValueParams(), AudioComponent,
+						[](FObjectKey InKey) { return EPreAnimatedAudioStateType::AudioComponentLifespan; });
+
+					if (AudioSection->GetOnQueueSubtitles().IsBound())
+					{
+						AudioComponent->OnQueueSubtitles = AudioSection->GetOnQueueSubtitles();
+					}
+					if (AudioSection->GetOnAudioFinished().IsBound())
+					{
+						AudioComponent->OnAudioFinished = AudioSection->GetOnAudioFinished();
+					}
+					if (AudioSection->GetOnAudioPlaybackPercent().IsBound())
+					{
+						AudioComponent->OnAudioPlaybackPercent = AudioSection->GetOnAudioPlaybackPercent();
+					}
+				}
+			}
+
+			if (EvaluationData)
+			{
+				EvaluationData->VolumeMultiplier = VolumeMultiplier;
+				EvaluationData->PitchMultiplier = PitchMultiplier;
+
+				EnsureAudioIsPlaying(BoundObject, *AudioSection, *EvaluationData, Context, *Player);
+			}
+		}
+	}
+
+	void EnsureAudioIsPlaying(
+			UObject* BoundObject,
+			UMovieSceneAudioSection& AudioSection,
+			FAudioComponentEvaluationData& EvaluationData,
+			const FMovieSceneContext& Context, 
+			IMovieScenePlayer& Player) const
+	{
+		using FAudioInputsBySectionKey = UMovieSceneAudioSystem::FAudioInputsBySectionKey;
+		using FAudioComponentInputEvaluationData = UMovieSceneAudioSystem::FAudioComponentInputEvaluationData;
+
+		ensureMsgf(EvaluationData.AudioComponent.IsValid(), TEXT("Trying to evaluate audio track on an invalid audio component"));
+		UAudioComponent& AudioComponent = *EvaluationData.AudioComponent.Get();
+
+		AudioSystem->PreAnimatedStorage->CachePreAnimatedValue(
+				FCachePreAnimatedValueParams(), &AudioComponent,
+				[](FObjectKey InKey) { return EPreAnimatedAudioStateType::AudioPlaying; });
+
+		if (AudioComponent.VolumeMultiplier != EvaluationData.VolumeMultiplier)
+		{
+			AudioComponent.SetVolumeMultiplier(EvaluationData.VolumeMultiplier);
+		}
+
+		if (AudioComponent.PitchMultiplier != EvaluationData.PitchMultiplier)
+		{
+			AudioComponent.SetPitchMultiplier(EvaluationData.PitchMultiplier);
+		}
+
+		AudioComponent.bSuppressSubtitles = AudioSection.GetSuppressSubtitles();
+
+		// Allow spatialization if we have any object we've been attached to.
+		const bool bAllowSpatialization = (BoundObject != nullptr || AudioComponent.GetAttachParent() != nullptr);
+
+		// Apply the input params.
+		FAudioInputsBySectionKey& AudioInputsBySectionKey = AudioSystem->AudioInputsBySectionKey;
+		FObjectKey SectionKey(&AudioSection);
+		FAudioComponentInputEvaluationData* AudioInputs = AudioInputsBySectionKey.Find(SectionKey);
+		if (AudioInputs)
+		{
+			SetAudioInputParameters(AudioInputs->Inputs_Float, AudioComponent);
+			SetAudioInputParameters(AudioInputs->Inputs_String, AudioComponent);
+			SetAudioInputParameters(AudioInputs->Inputs_Bool, AudioComponent);
+			SetAudioInputParameters(AudioInputs->Inputs_Int, AudioComponent);
+		}
+
+		FFrameNumber SectionStartFrame = (AudioSection.HasStartFrame() ? AudioSection.GetInclusiveStartFrame() : 0);
+		float SectionStartTimeSeconds = SectionStartFrame / AudioSection.GetTypedOuter<UMovieScene>()->GetTickResolution();
+
+		const FFrameNumber AudioStartOffset = AudioSection.GetStartOffset();
+		USoundBase* Sound = AudioSection.GetSound();
+
+		float AudioTime = (Context.GetTime() / Context.GetFrameRate()) 
+			- SectionStartTimeSeconds 
+			+ (float)Context.GetFrameRate().AsSeconds(AudioStartOffset);
+		if (AudioTime >= 0.f && Sound)
+		{
+			const float Duration = MovieSceneHelpers::GetSoundDuration(Sound);
+
+			if (!AudioSection.GetLooping() && AudioTime > Duration && Duration != 0.f)
+			{
+				AudioComponent.Stop();
+				return;
+			}
+
+			AudioTime = Duration > 0.f ? FMath::Fmod(AudioTime, Duration) : AudioTime;
+		}
+
+		// If the audio component is not playing we (may) need a state change. If the audio component is playing
+		// the wrong sound then we need a state change. If the audio playback time is significantly out of sync 
+		// with the desired time then we need a state change.
+		bool bSoundNeedsStateChange = !AudioComponent.IsPlaying() || AudioComponent.Sound != Sound;
+		bool bSoundNeedsTimeSync = false;
+
+		UObject* PlaybackContext = Player.GetPlaybackContext();
+		UWorld* World = PlaybackContext ? PlaybackContext->GetWorld() : nullptr;
+
+		// Sync only if there is no time dilation because otherwise the system will constantly resync because audio 
+		// playback is not dilated and will never match the expected playback time.
+		const bool bDoTimeSync = 
+			World && World->GetWorldSettings() &&
+			(FMath::IsNearlyEqual(World->GetWorldSettings()->GetEffectiveTimeDilation(), 1.f) ||
+			 !bIgnoreAudioSyncDuringWorldTimeDilationCVar);
+
+		if (bDoTimeSync)
+		{
+			float CurrentGameTime = 0.0f;
+
+			FAudioDevice* AudioDevice = World ? World->GetAudioDeviceRaw() : nullptr;
+			if (UseAudioClockForSequencerDesyncCVar && AudioDevice)
+			{
+				CurrentGameTime = AudioDevice->GetAudioClock();
+			}
+			else
+			{
+				CurrentGameTime = World ? World->GetAudioTimeSeconds() : 0.f;
+			}
+
+			// This tells us how much time has passed in the game world (and thus, reasonably, the audio playback)
+			// so if we calculate that we should be playing say, 15s into the section during evaluation, but
+			// we're only 5s since the last Play call, then we know we're out of sync. 
+			if (EvaluationData.PartialDesyncComputation.IsSet())
+			{
+				const float PartialDesyncComputation = EvaluationData.PartialDesyncComputation.GetValue();
+				float Desync = PartialDesyncComputation + AudioTime - CurrentGameTime;
+
+				if (!FMath::IsNearlyZero(MaxSequenceAudioDesyncToleranceCVar) && FMath::Abs(Desync) > MaxSequenceAudioDesyncToleranceCVar)
+				{
+					UE_LOG(LogMovieScene, Verbose, TEXT("Audio Component detected a significant mismatch in (assumed) playback time versus the desired time. Desync: %6.2f(s) Desired Time: %6.2f(s). Component: %s sound: %s"), Desync, AudioTime, *AudioComponent.GetName(), *GetNameSafe(AudioComponent.Sound));
+					bSoundNeedsTimeSync = true;
+				}
+			}
+		}
+
+		if (bSoundNeedsStateChange || bSoundNeedsTimeSync)
+		{
+			AudioComponent.bAllowSpatialization = bAllowSpatialization;
+
+			if (AudioSection.GetOverrideAttenuation())
+			{
+				AudioComponent.AttenuationSettings = AudioSection.GetAttenuationSettings();
+			}
+
+			// Only call stop on the sound if it is actually playing. This prevents spamming
+			// stop calls when a sound cue with a duration of zero is played.
+			if (AudioComponent.IsPlaying() || bSoundNeedsTimeSync)
+			{
+				UE_LOG(LogMovieScene, Verbose, TEXT("Audio Component stopped due to needing a state change bIsPlaying: %d bNeedsTimeSync: %d. Component: %s sound: %s"), AudioComponent.IsPlaying(), bSoundNeedsTimeSync, *AudioComponent.GetName(), *GetNameSafe(AudioComponent.Sound));
+				AudioComponent.Stop();
+			}
+
+			// Only change the sound clip if it has actually changed. This calls Stop internally if needed.
+			if (AudioComponent.Sound != Sound)
+			{
+				UE_LOG(LogMovieScene, Verbose, TEXT("Audio Component calling SetSound due to new sound. Component: %s OldSound: %s NewSound: %s"), *AudioComponent.GetName(), *GetNameSafe(AudioComponent.Sound), *GetNameSafe(AudioComponent.Sound));
+				AudioComponent.SetSound(Sound);
+			}
+#if WITH_EDITOR
+			if (GIsEditor && World != nullptr && !World->IsPlayInEditor())
+			{
+				AudioComponent.bIsUISound = true;
+				AudioComponent.bIsPreviewSound = true;
+			}
+			else
+#endif // WITH_EDITOR
+			{
+				AudioComponent.bIsUISound = false;
+			}
+
+			if (AudioTime >= 0.f)
+			{
+				UE_LOG(LogMovieScene, Verbose, TEXT("Audio Component Play at Local Time: %6.2f CurrentTime: %6.2f(s) SectionStart: %6.2f(s), SoundDur: %6.2f OffsetIntoClip: %6.2f sound: %s"), AudioTime, (Context.GetTime() / Context.GetFrameRate()), SectionStartTimeSeconds, AudioComponent.Sound ? AudioComponent.Sound->GetDuration() : 0.0f, (float)Context.GetFrameRate().AsSeconds(AudioStartOffset), *GetNameSafe(AudioComponent.Sound));
+				AudioComponent.Play(AudioTime);
+
+				// Keep track of when we asked this audio clip to play (in game time) so that we can figure 
+				// out if there's a significant desync in the future.
+				//
+				// The goal is later to compare:
+				//   (NewAudioTime - PreviousAudioTime) and 
+				//   (NewGameTime - PreviousGameTime)
+				//
+				// If their difference is larger than some threshold, we have a desync. NewGameTime and 
+				// NewAudioTime will be known next update, but PreviousGameTime and PreviousAudioTime
+				// are known now. Let's store (-PreviousAudioTime + PreviousGameTime), and we will only 
+				// need to add (NewAudioTime - NewGameTime).
+				if (World)
+				{
+					FAudioDevice* AudioDevice = World->GetAudioDeviceRaw();
+					if (UseAudioClockForSequencerDesyncCVar && AudioDevice)
+					{
+						EvaluationData.PartialDesyncComputation = AudioDevice->GetAudioClock() - AudioTime;
+					}
+					else
+					{
+						EvaluationData.PartialDesyncComputation = World->GetAudioTimeSeconds() - AudioTime;
+					}
+				}
+			}
+		}
+
+		if (Context.GetStatus() != EMovieScenePlayerStatus::Playing)
+		{
+			// While scrubbing, play the sound for a short time and then cut it.
+			AudioComponent.StopDelayed(AudioTrackConstants::ScrubDuration);
+		}
+
+		if (AudioComponent.IsPlaying() && AudioInputs)
+		{
+			SetAudioInputTriggers(AudioInputs->Inputs_Trigger, AudioComponent);
+		}
+
+		if (bAllowSpatialization)
+		{
+			if (FAudioDevice* AudioDevice = AudioComponent.GetAudioDevice())
+			{
+				DECLARE_CYCLE_STAT(TEXT("FAudioThreadTask.MovieSceneUpdateAudioTransform"), STAT_MovieSceneUpdateAudioTransform, STATGROUP_TaskGraphTasks);
+				AudioDevice->SendCommandToActiveSounds(AudioComponent.GetAudioComponentID(), [ActorTransform = AudioComponent.GetComponentTransform()](FActiveSound& ActiveSound)
+				{
+					ActiveSound.bLocationDefined = true;
+					ActiveSound.Transform = ActorTransform;
+				}, GET_STATID(STAT_MovieSceneUpdateAudioTransform));
+			}
+		}
+	}
+
+	// Helper method for firing all triggered audio triggers.
+	void SetAudioInputTriggers(const TArray<FName>& Inputs, IAudioParameterControllerInterface& InParamaterInterface) const
+	{
+		for (const FName& TriggerName : Inputs)
+		{
+			InParamaterInterface.SetTriggerParameter(TriggerName);
+		};
+	}
+
+	// Helper template to set all audio input values previously evaluated.
+	template<typename ValueType>
+	void SetAudioInputParameters(TMap<FName, ValueType>& Inputs, IAudioParameterControllerInterface& InParamaterInterface) const
+	{
+		for (TPair<FName, ValueType>& Pair : Inputs)
+		{
+			InParamaterInterface.SetParameter<ValueType>(Pair.Key, MoveTempIfPossible(Pair.Value));
+		};
+	}
+};
+
+} // namespace UE::MovieScene
+
+UMovieSceneAudioSystem::UMovieSceneAudioSystem(const FObjectInitializer& ObjInit)
+	: UMovieSceneEntitySystem(ObjInit)
+{
+	using namespace UE::MovieScene;
+
+	RelevantComponent = FMovieSceneTracksComponentTypes::Get()->Audio;
+
+	if (HasAnyFlags(RF_ClassDefaultObject))
+	{
+		const FBuiltInComponentTypes* BuiltInComponents = FBuiltInComponentTypes::Get();
+		const FMovieSceneTracksComponentTypes* TrackComponents = FMovieSceneTracksComponentTypes::Get();
+
+		// We consume the result of all possible audio input channels.
+		for (int32 Index = 0; Index < 9; ++Index)
+		{
+			DefineComponentConsumer(GetClass(), BuiltInComponents->DoubleResult[Index]);
+		}
+		DefineComponentConsumer(GetClass(), BuiltInComponents->StringResult);
+		DefineComponentConsumer(GetClass(), BuiltInComponents->IntegerResult);
+		DefineComponentConsumer(GetClass(), BuiltInComponents->BoolResult);
+		DefineComponentConsumer(GetClass(), TrackComponents->AudioTriggerName);
+	}
+}
+
+void UMovieSceneAudioSystem::OnLink()
+{
+	using namespace UE::MovieScene;
+
+	PreAnimatedStorage = Linker->PreAnimatedState.GetOrCreateStorage<FPreAnimatedAudioStorage>();
+}
+
+void UMovieSceneAudioSystem::OnUnlink()
+{
+	using namespace UE::MovieScene;
+
+	for (const TPair<FObjectKey, FAudioComponentBySectionKey>& AudioComponentsForActor : AudioComponentsByActorKey)
+	{
+		for (const TPair<FObjectKey, FAudioComponentEvaluationData>& AudioComponentForSection : AudioComponentsForActor.Value)
+		{
+			UAudioComponent* AudioComponent = AudioComponentForSection.Value.AudioComponent.Get();
+			if (AudioComponent)
+			{
+				UObject* Actor = AudioComponentsForActor.Key.ResolveObjectPtr();
+				UObject* Section = AudioComponentForSection.Key.ResolveObjectPtr();
+				UE_LOG(LogMovieScene, Warning, TEXT("Cleaning audio component '%s' for section '%s' on actor '%s'"),
+						*AudioComponent->GetPathName(),
+						Section ? *Section->GetPathName() : TEXT("<null>"),
+						Actor ? *Actor->GetPathName() : TEXT("<null>"));
+			}
+		}
+	}
+
+	AudioComponentsByActorKey.Reset();
+	AudioInputsBySectionKey.Reset();
+
+}
+
+void UMovieSceneAudioSystem::OnRun(FSystemTaskPrerequisites& InPrerequisites, FSystemSubsequentTasks& Subsequents)
+{
+	MOVIESCENE_DETAILED_SCOPE_CYCLE_COUNTER(MovieSceneEval_AudioTrack_Evaluate)
+
+	if (!GEngine || !GEngine->UseSound())
+	{
+		return;
+	}
+
+	using namespace UE::MovieScene;
+
+	const FBuiltInComponentTypes* BuiltInComponents = FBuiltInComponentTypes::Get();
+	const FMovieSceneTracksComponentTypes* TrackComponents = FMovieSceneTracksComponentTypes::Get();
+
+	// Reset shared data.
+	AudioInputsBySectionKey.Reset();
+	for (TPair<FObjectKey, FAudioComponentBySectionKey>& AudioComponentsForActor : AudioComponentsByActorKey)
+	{
+		for (TPair<FObjectKey, FAudioComponentEvaluationData>& AudioComponentForSection : AudioComponentsForActor.Value)
+		{
+			AudioComponentForSection.Value.bEvaluatedThisFrame = false;
+		}
+	}
+
+	// Gather audio input values computed by the channel evaluators.
+	FSystemTaskPrerequisites Prereqs;
+
+	FGraphEventRef Task = FEntityTaskBuilder()
+	.Read(TrackComponents->Audio)
+	.Read(TrackComponents->AudioInputs)
+	.ReadOneOrMoreOf(
+			BuiltInComponents->DoubleResult[0], BuiltInComponents->DoubleResult[1], BuiltInComponents->DoubleResult[2], 
+			BuiltInComponents->DoubleResult[3], BuiltInComponents->DoubleResult[4], BuiltInComponents->DoubleResult[5], 
+			BuiltInComponents->DoubleResult[6], BuiltInComponents->DoubleResult[7], BuiltInComponents->DoubleResult[8], 
+			BuiltInComponents->StringResult,
+			BuiltInComponents->IntegerResult,
+			BuiltInComponents->BoolResult)
+	.template Dispatch_PerAllocation<FGatherAudioInputs>(&Linker->EntityManager, InPrerequisites, nullptr, this);
+	if (Task)
+	{
+		Prereqs.AddRootTask(Task);
+	}
+
+	Task = FEntityTaskBuilder()
+	.Read(TrackComponents->Audio)
+	.Read(TrackComponents->AudioTriggerName)
+	.template Dispatch_PerAllocation<FGatherAudioTriggers>(&Linker->EntityManager, InPrerequisites, nullptr, this);
+	if (Task)
+	{
+		Prereqs.AddRootTask(Task);
+	}
+
+	// Next, evaluate audio to play and use the gathered audio input values to set on the audio components.
+	FEntityTaskBuilder()
+	.ReadEntityIDs()
+	.Read(BuiltInComponents->RootInstanceHandle)
+	.Read(BuiltInComponents->InstanceHandle)
+	.Read(TrackComponents->Audio)
+	.Read(BuiltInComponents->DoubleResult[0]) // Volume
+	.Read(BuiltInComponents->DoubleResult[1]) // Pitch multiplier
+	.ReadOptional(BuiltInComponents->BoundObject)
+	.SetDesiredThread(Linker->EntityManager.GetGatherThread())
+	.template Dispatch_PerAllocation<FEvaluateAudio>(&Linker->EntityManager, Prereqs, &Subsequents, this);
+}
+
+UMovieSceneAudioSystem::FAudioComponentEvaluationData* UMovieSceneAudioSystem::GetAudioComponentEvaluationData(FObjectKey ActorKey, FObjectKey SectionKey)
+{
+	FAudioComponentBySectionKey* Map = AudioComponentsByActorKey.Find(ActorKey);
+	if (Map != nullptr)
+	{
+		// First, check for an exact match for this entity
+		FAudioComponentEvaluationData* ExistingData = Map->Find(SectionKey);
+		if (ExistingData != nullptr)
+		{
+			if (ExistingData->AudioComponent.IsValid())
+			{
+				return ExistingData;
+			}
+		}
+
+		// If no exact match, check for any AudioComponent that isn't busy
+		for (FAudioComponentBySectionKey::ElementType& Pair : *Map)
+		{
+			UAudioComponent* ExistingComponent = Pair.Value.AudioComponent.Get();
+			if (ExistingComponent && !ExistingComponent->IsPlaying())
+			{
+				// Replace this entry with the new entity ID to claim it
+				FAudioComponentEvaluationData MovedData(Pair.Value);
+				Map->Remove(Pair.Key);
+				MovedData.PartialDesyncComputation.Reset();
+				return &Map->Add(SectionKey, MovedData);
+			}
+		}
+	}
+
+	return nullptr;
+}
+
+UMovieSceneAudioSystem::FAudioComponentEvaluationData* UMovieSceneAudioSystem::AddBoundObjectAudioComponent(UMovieSceneAudioSection* Section, UObject* PrincipalObject)
+{
+	using namespace UE::MovieScene;
+
+	FObjectKey ObjectKey(PrincipalObject);
+	FObjectKey SectionKey(Section);
+
+	FAudioComponentBySectionKey& ActorAudioComponentMap = AudioComponentsByActorKey.FindOrAdd(ObjectKey);
+
+	FAudioComponentEvaluationData* ExistingData = GetAudioComponentEvaluationData(ObjectKey, SectionKey);
+	if (!ExistingData)
+	{
+		USoundCue* TempPlaybackAudioCue = NewObject<USoundCue>();
+
+		AActor* Actor = nullptr;
+		USceneComponent* SceneComponent = nullptr;
+		FString ObjectName;
+
+		if (PrincipalObject->IsA<AActor>())
+		{
+			Actor = Cast<AActor>(PrincipalObject);
+			SceneComponent = Actor->GetRootComponent();
+			ObjectName =
+#if WITH_EDITOR
+				Actor->GetActorLabel();
+#else
+				Actor->GetName();
+#endif
+		}
+		else if (PrincipalObject->IsA<UActorComponent>())
+		{
+			UActorComponent* ActorComponent = Cast<UActorComponent>(PrincipalObject);
+			Actor = ActorComponent->GetOwner();
+			SceneComponent = Cast<USceneComponent>(ActorComponent);
+			ObjectName = ActorComponent->GetName();
+		}
+
+		if (!Actor || !SceneComponent)
+		{
+			const int32 RowIndex = Section->GetRowIndex();
+			UE_LOG(LogMovieScene, Warning, TEXT("Failed to find scene component for spatialized audio track (row %d)."), RowIndex);
+			return nullptr;
+		}
+
+		FAudioDevice::FCreateComponentParams Params(Actor->GetWorld(), Actor);
+		UAudioComponent* NewComponent = FAudioDevice::CreateComponent(TempPlaybackAudioCue, Params);
+
+		if (!NewComponent)
+		{
+			const int32 RowIndex = Section->GetRowIndex();
+			UE_LOG(LogMovieScene, Warning, TEXT("Failed to create audio component for spatialized audio track (row %d on %s)."), RowIndex, *ObjectName);
+			return nullptr;
+		}
+
+		NewComponent->SetFlags(RF_Transient);
+		NewComponent->AttachToComponent(SceneComponent, FAttachmentTransformRules::KeepRelativeTransform);
+
+		ExistingData = &ActorAudioComponentMap.Add(SectionKey);
+		ExistingData->AudioComponent = NewComponent;
+	}
+
+	return ExistingData;
+}
+
+UMovieSceneAudioSystem::FAudioComponentEvaluationData* UMovieSceneAudioSystem::AddRootAudioComponent(UMovieSceneAudioSection* Section, UWorld* World)
+{
+	using namespace UE::MovieScene;
+
+	FObjectKey NullKey;
+	FObjectKey SectionKey(Section);
+
+	FAudioComponentBySectionKey& RootAudioComponentMap = AudioComponentsByActorKey.FindOrAdd(NullKey);
+
+	FAudioComponentEvaluationData* ExistingData = GetAudioComponentEvaluationData(NullKey, SectionKey);
+	if (!ExistingData)
+	{
+		USoundCue* TempPlaybackAudioCue = NewObject<USoundCue>();
+
+		FAudioDevice::FCreateComponentParams Params(World);
+
+		UAudioComponent* NewComponent = FAudioDevice::CreateComponent(TempPlaybackAudioCue, Params);
+
+		if (!NewComponent)
+		{
+			const int32 RowIndex = Section->GetRowIndex();
+			UE_LOG(LogMovieScene, Warning, TEXT("Failed to create audio component for root audio track (row %d)."), RowIndex);
+			return nullptr;
+		}
+
+		NewComponent->SetFlags(RF_Transient);
+
+		ExistingData = &RootAudioComponentMap.Add(SectionKey);
+		ExistingData->AudioComponent = NewComponent;
+	}
+
+	return ExistingData;
+}
+
+void UMovieSceneAudioSystem::StopSound(FObjectKey ActorKey, FObjectKey SectionKey)
+{
+	if (FAudioComponentBySectionKey* Map = AudioComponentsByActorKey.Find(ActorKey))
+	{
+		if (FAudioComponentEvaluationData* Data = Map->Find(SectionKey))
+		{
+			if (UAudioComponent* AudioComponent = Data->AudioComponent.Get())
+			{
+				AudioComponent->Stop();
+			}
+		}
+	}
+}
+
diff --git a/Engine/Source/Runtime/MovieSceneTracks/Private/Systems/MovieSceneAudioSystem.h b/Engine/Source/Runtime/MovieSceneTracks/Private/Systems/MovieSceneAudioSystem.h
new file mode 100644
--- /dev/null
+++ b/Engine/Source/Runtime/MovieSceneTracks/Private/Systems/MovieSceneAudioSystem.h
@@ -0,0 +1,119 @@
+// Copyright Epic Games, Inc. All Rights Reserved.
+
+#pragma once
+
+#include "CoreTypes.h"
+#include "EntitySystem/MovieSceneEntityIDs.h"
+#include "MovieSceneTracksComponentTypes.h"
+#include "UObject/ObjectMacros.h"
+#include "Channels/MovieSceneAudioTriggerChannel.h"
+#include "EntitySystem/MovieSceneEntitySystem.h"
+
+#include "MovieSceneAudioSystem.generated.h"
+
+class IAudioParameterControllerInterface;
+class UAudioComponent;
+class UMovieSceneAudioSection;
+class USoundBase;
+struct FMoveSceneAudioTriggerState;
+
+namespace UE::MovieScene
+{
+	struct FGatherAudioInputs;
+	struct FGatherAudioTriggers;
+	struct FEvaluateAudio;
+	struct FPreAnimatedAudioStorage;
+
+	struct FAudioComponentInputEvaluationData
+	{
+		TMap<FName, float> Inputs_Float;
+		TMap<FName, FString> Inputs_String;
+		TMap<FName, bool> Inputs_Bool;
+		TMap<FName, int32> Inputs_Int;
+		TArray<FName> Inputs_Trigger;
+	};
+
+	struct FAudioComponentEvaluationData
+	{
+		/** The audio component that was created to play audio */
+		TWeakObjectPtr<UAudioComponent> AudioComponent;
+
+		/** Volume multiplier to use this frame */
+		double VolumeMultiplier;
+
+		/** Pitch multiplier to use this frame */
+		double PitchMultiplier;
+
+		/**
+		 * Set whenever we ask the Audio component to start playing a sound.
+		 * Used to detect desyncs caused when Sequencer evaluates at more-than-real-time.
+		 */
+		TOptional<float> PartialDesyncComputation;
+
+		/** Flag to keep track of audio components evaluated on a given frame */
+		bool bEvaluatedThisFrame;
+	};
+}
+
+/**
+ * System for evaluating audio tracks
+ */
+UCLASS()
+class UMovieSceneAudioSystem : public UMovieSceneEntitySystem
+{
+	GENERATED_BODY()
+
+public:
+
+	using FMovieSceneEntityID = UE::MovieScene::FMovieSceneEntityID;
+	using FAudioComponentEvaluationData = UE::MovieScene::FAudioComponentEvaluationData;
+	using FAudioComponentInputEvaluationData = UE::MovieScene::FAudioComponentInputEvaluationData;
+
+	UMovieSceneAudioSystem(const FObjectInitializer& ObjInit);
+
+	//~ UMovieSceneEntitySystem members
+	virtual void OnLink() override;
+	virtual void OnUnlink() override;
+	virtual void OnRun(FSystemTaskPrerequisites& InPrerequisites, FSystemSubsequentTasks& Subsequents) override;
+
+	/**
+	 * Get the evaluation data for the given actor and section. Pass a null actor key for root (world) audio.
+	 */
+	FAudioComponentEvaluationData* GetAudioComponentEvaluationData(FObjectKey ActorKey, FObjectKey SectionKey);
+
+	/**
+	 * Adds an audio component to the given bound sequencer object.
+	 * WARNING: Only to be called on the game thread.
+	 */
+	FAudioComponentEvaluationData* AddBoundObjectAudioComponent(UMovieSceneAudioSection* Section, UObject* PrincipalObject);
+
+	/**
+	 * Adds an audio component to the world, for playing root audio tracks.
+	 * WARNING: Only to be called on the game thread.
+	 */
+	FAudioComponentEvaluationData* AddRootAudioComponent(UMovieSceneAudioSection* Section, UWorld* World);
+
+	/**
+	 * Stop the audio on the audio component associated with the given audio section.
+	 */
+	void StopSound(FObjectKey ActorKey, FObjectKey SectionKey);
+
+private:
+
+	/** Map of all created audio components */
+	using FAudioComponentBySectionKey = TMap<FObjectKey, FAudioComponentEvaluationData>;
+	using FAudioComponentsByActorKey = TMap<FObjectKey, FAudioComponentBySectionKey>;
+	FAudioComponentsByActorKey AudioComponentsByActorKey;
+
+	/** Map of audio input values, rebuilt every frame */
+	using FAudioInputsBySectionKey = TMap<FObjectKey, FAudioComponentInputEvaluationData>;
+	FAudioInputsBySectionKey AudioInputsBySectionKey;
+
+	/** Pre-animated state */
+	TSharedPtr<UE::MovieScene::FPreAnimatedAudioStorage> PreAnimatedStorage;
+
+	friend struct UE::MovieScene::FGatherAudioInputs;
+	friend struct UE::MovieScene::FGatherAudioTriggers;
+	friend struct UE::MovieScene::FEvaluateAudio;
+};
+
diff --git a/Engine/Source/Runtime/MovieSceneTracks/Private/Tracks/MovieSceneAudioTrack.cpp b/Engine/Source/Runtime/MovieSceneTracks/Private/Tracks/MovieSceneAudioTrack.cpp
--- a/Engine/Source/Runtime/MovieSceneTracks/Private/Tracks/MovieSceneAudioTrack.cpp
+++ b/Engine/Source/Runtime/MovieSceneTracks/Private/Tracks/MovieSceneAudioTrack.cpp
@@ -5,7 +5,6 @@
 #include "Sound/SoundWave.h"
 #include "MovieScene.h"
 #include "Sections/MovieSceneAudioSection.h"
-#include "Evaluation/MovieSceneAudioTemplate.h"
 #include "Kismet/GameplayStatics.h"
 #include "AudioDecompress.h"
 #include "Evaluation/MovieSceneSegment.h"
@@ -27,11 +26,6 @@ UMovieSceneAudioTrack::UMovieSceneAudioTrack( const FObjectInitializer& ObjectIn
 #endif
 }
 
-FMovieSceneEvalTemplatePtr UMovieSceneAudioTrack::CreateTemplateForSection(const UMovieSceneSection& InSection) const
-{
-	return FMovieSceneAudioSectionTemplate(*CastChecked<UMovieSceneAudioSection>(&InSection));
-}
-
 const TArray<UMovieSceneSection*>& UMovieSceneAudioTrack::GetAllSections() const
 {
 	return AudioSections;
diff --git a/Engine/Source/Runtime/MovieSceneTracks/Public/MovieSceneTracksComponentTypes.h b/Engine/Source/Runtime/MovieSceneTracks/Public/MovieSceneTracksComponentTypes.h
--- a/Engine/Source/Runtime/MovieSceneTracks/Public/MovieSceneTracksComponentTypes.h
+++ b/Engine/Source/Runtime/MovieSceneTracks/Public/MovieSceneTracksComponentTypes.h
@@ -18,6 +18,7 @@
 
 class UMaterialParameterCollection;
 class UMovieScene3DTransformSection;
+class UMovieSceneAudioSection;
 class UMovieSceneDataLayerSection;
 class UMovieSceneLevelVisibilitySection;
 class UMovieSceneSkeletalAnimationSection;
@@ -90,6 +91,40 @@ struct FMovieSceneSkeletalAnimationComponentData
 	TObjectPtr<UMovieSceneSkeletalAnimationSection> Section;
 };
 
+/** Component data for audio tracks */
+USTRUCT()
+struct FMovieSceneAudioComponentData
+{
+	GENERATED_BODY()
+
+	UPROPERTY()
+	TObjectPtr<UMovieSceneAudioSection> Section = nullptr;
+};
+
+/**
+ * Component data for audio tracks inputs
+ * This provides the names of the inputs whose values are stored on the
+ * same entity using the DoubleResult, StringResult, BoolResult, and
+ * IntegerResult components.
+ */
+USTRUCT()
+struct FMovieSceneAudioInputData
+{
+	GENERATED_BODY()
+
+	UPROPERTY()
+	FName FloatInputs[9];
+
+	UPROPERTY()
+	FName StringInput;
+
+	UPROPERTY()
+	FName BoolInput;
+
+	UPROPERTY()
+	FName IntInput;
+};
+
 namespace UE
 {
 namespace MovieScene
@@ -464,6 +499,10 @@ struct MOVIESCENETRACKS_API FMovieSceneTracksComponentTypes
 
 	TComponentTypeID<FFadeComponentData> Fade;
 
+	TComponentTypeID<FMovieSceneAudioComponentData> Audio;
+	TComponentTypeID<FMovieSceneAudioInputData> AudioInputs;
+	TComponentTypeID<FName> AudioTriggerName;
+
 	struct
 	{
 		TCustomPropertyRegistration<FBoolPropertyTraits> Bool;
diff --git a/Engine/Source/Runtime/MovieSceneTracks/Public/Sections/MovieSceneAudioSection.h b/Engine/Source/Runtime/MovieSceneTracks/Public/Sections/MovieSceneAudioSection.h
--- a/Engine/Source/Runtime/MovieSceneTracks/Public/Sections/MovieSceneAudioSection.h
+++ b/Engine/Source/Runtime/MovieSceneTracks/Public/Sections/MovieSceneAudioSection.h
@@ -2,17 +2,19 @@
 
 #pragma once
 
-#include "CoreMinimal.h"
-#include "UObject/ObjectMacros.h"
-#include "MovieSceneSection.h"
-#include "Components/AudioComponent.h"
-#include "Sound/SoundAttenuation.h"
-#include "Channels/MovieSceneFloatChannel.h"
-#include "Channels/MovieSceneStringChannel.h"
+#include "Channels/MovieSceneAudioTriggerChannel.h"
 #include "Channels/MovieSceneBoolChannel.h"
+#include "Channels/MovieSceneFloatChannel.h"
 #include "Channels/MovieSceneIntegerChannel.h"
+#include "Channels/MovieSceneStringChannel.h"
+#include "Components/AudioComponent.h"
+#include "CoreMinimal.h"
+#include "EntitySystem/IMovieSceneEntityProvider.h"
+#include "MovieSceneSection.h"
 #include "Sections/MovieSceneActorReferenceSection.h"
-#include "Channels/MovieSceneAudioTriggerChannel.h"
+#include "Sound/SoundAttenuation.h"
+#include "UObject/ObjectMacros.h"
+
 #include "MovieSceneAudioSection.generated.h"
 
 class USoundBase;
@@ -23,6 +25,7 @@ class USoundBase;
 UCLASS()
 class MOVIESCENETRACKS_API UMovieSceneAudioSection
 	: public UMovieSceneSection
+	, public IMovieSceneEntityProvider
 {
 	GENERATED_UCLASS_BODY()
 
@@ -190,20 +193,29 @@ public:
 	virtual void PostEditImport() override;
 	virtual EMovieSceneChannelProxyType CacheChannelProxy() override;
 
+	//~ IMovieSceneEntityProvider interface
+	virtual bool PopulateEvaluationFieldImpl(const TRange<FFrameNumber>& EffectiveRange, const FMovieSceneEvaluationFieldEntityMetaData& InMetaData, FMovieSceneEntityComponentFieldBuilder* OutFieldBuilder) override;
+	virtual void ImportEntityImpl(UMovieSceneEntitySystemLinker* EntityLinker, const FEntityImportParams& Params, FImportedEntity* OutImportedEntity) override;
+
+#if WITH_EDITOR
+	virtual void PostEditChangeProperty(FPropertyChangedEvent& PropertyChangedEvent) override;
+#endif
+
 private:
+
 	template<typename ChannelType, typename ForEachFunction>
-	FORCEINLINE static void ForEachInternal(ForEachFunction InFuncton, const TMap<FName, ChannelType>& InMapToIterate) 
+	FORCEINLINE static void ForEachInternal(ForEachFunction InFunction, const TMap<FName, ChannelType>& InMapToIterate) 
 	{
 		for (auto& Item : InMapToIterate)
 		{
-			InFuncton(Item.Key, Item.Value);
+			InFunction(Item.Key, Item.Value);
 		}	
 	}
 
 	void SetupSoundInputParameters(const USoundBase* InSoundBase);
 
 	/** The sound cue or wave that this section plays */
-	UPROPERTY(EditAnywhere, Category="Audio")
+	UPROPERTY(EditAnywhere, Category="Audio", BlueprintGetter=GetSound, BlueprintSetter=SetSound)
 	TObjectPtr<USoundBase> Sound;
 
 	/** The offset into the beginning of the audio clip */
diff --git a/Engine/Source/Runtime/MovieSceneTracks/Public/Tracks/MovieSceneAudioTrack.h b/Engine/Source/Runtime/MovieSceneTracks/Public/Tracks/MovieSceneAudioTrack.h
--- a/Engine/Source/Runtime/MovieSceneTracks/Public/Tracks/MovieSceneAudioTrack.h
+++ b/Engine/Source/Runtime/MovieSceneTracks/Public/Tracks/MovieSceneAudioTrack.h
@@ -22,7 +22,6 @@ namespace AudioTrackConstants
 UCLASS()
 class MOVIESCENETRACKS_API UMovieSceneAudioTrack
 	: public UMovieSceneNameableTrack
-	, public IMovieSceneTrackTemplateProducer
 {
 	GENERATED_UCLASS_BODY()
 
@@ -58,9 +57,6 @@ public:
 	virtual bool SupportsMultipleRows() const override;
 	virtual UMovieSceneSection* CreateNewSection() override;
 
-	// ~IMovieSceneTrackTemplateProducer interface
-	virtual FMovieSceneEvalTemplatePtr CreateTemplateForSection(const UMovieSceneSection& InSection) const override;
-
 private:
 
 	/** List of all root audio sections */
