Sequencer: fix static analysis warnings

diff --git a/Engine/Source/Runtime/MovieSceneTracks/Private/Systems/MovieSceneAudioSystem.cpp b/Engine/Source/Runtime/MovieSceneTracks/Private/Systems/MovieSceneAudioSystem.cpp
--- a/Engine/Source/Runtime/MovieSceneTracks/Private/Systems/MovieSceneAudioSystem.cpp
+++ b/Engine/Source/Runtime/MovieSceneTracks/Private/Systems/MovieSceneAudioSystem.cpp
@@ -366,12 +366,10 @@ private:
 				// Initialize the sound
 				UWorld* World = PlaybackContext ? PlaybackContext->GetWorld() : nullptr;
 				EvaluationData = AudioSystem->AddRootAudioComponent(AudioSection, World);
+				UAudioComponent* AudioComponent = EvaluationData ? EvaluationData->AudioComponent.Get() : nullptr;
 
-				if (EvaluationData)
+				if (ensure(AudioComponent))
 				{
-					UAudioComponent* AudioComponent = EvaluationData->AudioComponent.Get();
-					ensure(AudioComponent);
-
 					AudioSystem->PreAnimatedStorage->BeginTrackingEntity(EntityID, bWantsRestoreState, RootInstanceHandle, AudioComponent);
 					AudioSystem->PreAnimatedStorage->CachePreAnimatedValue(
 						FCachePreAnimatedValueParams(), AudioComponent,
@@ -417,12 +415,10 @@ private:
 			{
 				// Initialize the sound
 				EvaluationData = AudioSystem->AddBoundObjectAudioComponent(AudioSection, BoundObject);
+				UAudioComponent* AudioComponent = EvaluationData ? EvaluationData->AudioComponent.Get() : nullptr;
 
-				if (EvaluationData)
+				if (AudioComponent)
 				{
-					UAudioComponent* AudioComponent = EvaluationData->AudioComponent.Get();
-					ensure(AudioComponent);
-
 					AudioSystem->PreAnimatedStorage->BeginTrackingEntity(EntityID, bWantsRestoreState, RootInstanceHandle, AudioComponent);
 					AudioSystem->PreAnimatedStorage->CachePreAnimatedValue(
 						FCachePreAnimatedValueParams(), AudioComponent,
