Sequencer: add support for blend types and multiple rows for material parameter tracks (widget and MPC).

- A new RowIndex property is added to FBuildEditWidgetParams, since the material parameter track needs to figure out what section and row to add keys to.
- Parameter track allows adding parameters to a specific row.
- Parameter track now implements and uses the "section to key" API.
- Parameter sections are still used by other tracks such as particle parameter tracks, so this change does not affect the default blend type support.

diff --git a/Engine/Source/Editor/MovieSceneTools/Private/TrackEditors/MaterialParameterCollectionTrackEditor.cpp b/Engine/Source/Editor/MovieSceneTools/Private/TrackEditors/MaterialParameterCollectionTrackEditor.cpp
--- a/Engine/Source/Editor/MovieSceneTools/Private/TrackEditors/MaterialParameterCollectionTrackEditor.cpp
+++ b/Engine/Source/Editor/MovieSceneTools/Private/TrackEditors/MaterialParameterCollectionTrackEditor.cpp
@@ -256,15 +256,26 @@ const FSlateBrush* FMaterialParameterCollectionTrackEditor::GetIconBrush() const
 TSharedPtr<SWidget> FMaterialParameterCollectionTrackEditor::BuildOutlinerEditWidget(const FGuid& ObjectBinding, UMovieSceneTrack* Track, const FBuildEditWidgetParams& Params)
 {
 	UMovieSceneMaterialParameterCollectionTrack* MPCTrack = Cast<UMovieSceneMaterialParameterCollectionTrack>(Track);
-	FOnGetContent MenuContent = FOnGetContent::CreateSP(this, &FMaterialParameterCollectionTrackEditor::OnGetAddParameterMenuContent, MPCTrack);
+	FOnGetContent MenuContent = FOnGetContent::CreateSP(this, &FMaterialParameterCollectionTrackEditor::OnGetAddParameterMenuContent, MPCTrack, Params.RowIndex, Params.TrackInsertRowIndex);
 
 	return FSequencerUtilities::MakeAddButton(LOCTEXT("AddParameterButton", "Parameter"), MenuContent, Params.NodeIsHovered, GetSequencer());
 }
 
-TSharedRef<SWidget> FMaterialParameterCollectionTrackEditor::OnGetAddParameterMenuContent(UMovieSceneMaterialParameterCollectionTrack* MPCTrack)
+TSharedRef<SWidget> FMaterialParameterCollectionTrackEditor::OnGetAddParameterMenuContent(UMovieSceneMaterialParameterCollectionTrack* MPCTrack, int32 RowIndex, int32 TrackInsertRowIndex)
 {
 	FMenuBuilder MenuBuilder(true, nullptr);
 
+	// If this is supported, allow creating other sections with different blend types, and put
+	// the material parameters after a separator. Otherwise, just show the parameters menu.
+	const FMovieSceneBlendTypeField SupportedBlendTypes = MPCTrack->GetSupportedBlendTypes();
+	if (SupportedBlendTypes.Num() > 1)
+	{
+		TWeakPtr<ISequencer> WeakSequencer = GetSequencer();
+		FSequencerUtilities::PopulateMenu_CreateNewSection(MenuBuilder, TrackInsertRowIndex, MPCTrack, WeakSequencer);
+
+		MenuBuilder.AddSeparator();
+	}
+
 	MenuBuilder.BeginSection(NAME_None, LOCTEXT("ScalarParametersHeading", "Scalar"));
 	{
 		TArray<FCollectionScalarParameter> ScalarParameters = MPCTrack->MPC->ScalarParameters;
@@ -276,7 +287,7 @@ TSharedRef<SWidget> FMaterialParameterCollectionTrackEditor::OnGetAddParameterMe
 				FText::FromName(Scalar.ParameterName),
 				FText(),
 				FSlateIcon(),
-				FExecuteAction::CreateSP(this, &FMaterialParameterCollectionTrackEditor::AddScalarParameter, MPCTrack, Scalar)
+				FExecuteAction::CreateSP(this, &FMaterialParameterCollectionTrackEditor::AddScalarParameter, MPCTrack, RowIndex, Scalar)
 				);
 		}
 	}
@@ -293,7 +304,7 @@ TSharedRef<SWidget> FMaterialParameterCollectionTrackEditor::OnGetAddParameterMe
 				FText::FromName(Vector.ParameterName),
 				FText(),
 				FSlateIcon(),
-				FExecuteAction::CreateSP(this, &FMaterialParameterCollectionTrackEditor::AddVectorParameter, MPCTrack, Vector)
+				FExecuteAction::CreateSP(this, &FMaterialParameterCollectionTrackEditor::AddVectorParameter, MPCTrack, RowIndex, Vector)
 				);
 		}
 	}
@@ -303,7 +314,7 @@ TSharedRef<SWidget> FMaterialParameterCollectionTrackEditor::OnGetAddParameterMe
 }
 
 
-void FMaterialParameterCollectionTrackEditor::AddScalarParameter(UMovieSceneMaterialParameterCollectionTrack* Track, FCollectionScalarParameter Parameter)
+void FMaterialParameterCollectionTrackEditor::AddScalarParameter(UMovieSceneMaterialParameterCollectionTrack* Track, int32 RowIndex, FCollectionScalarParameter Parameter)
 {
 	if (!Track->MPC)
 	{
@@ -314,12 +325,12 @@ void FMaterialParameterCollectionTrackEditor::AddScalarParameter(UMovieSceneMate
 
 	const FScopedTransaction Transaction(LOCTEXT("AddScalarParameter", "Add scalar parameter"));
 	Track->Modify();
-	Track->AddScalarParameterKey(Parameter.ParameterName, KeyTime, Parameter.DefaultValue);
+	Track->AddScalarParameterKey(Parameter.ParameterName, KeyTime, RowIndex, Parameter.DefaultValue);
 	GetSequencer()->NotifyMovieSceneDataChanged(EMovieSceneDataChangeType::MovieSceneStructureItemAdded);
 }
 
 
-void FMaterialParameterCollectionTrackEditor::AddVectorParameter(UMovieSceneMaterialParameterCollectionTrack* Track, FCollectionVectorParameter Parameter)
+void FMaterialParameterCollectionTrackEditor::AddVectorParameter(UMovieSceneMaterialParameterCollectionTrack* Track, int32 RowIndex, FCollectionVectorParameter Parameter)
 {
 	if (!Track->MPC)
 	{
@@ -330,7 +341,7 @@ void FMaterialParameterCollectionTrackEditor::AddVectorParameter(UMovieSceneMate
 
 	const FScopedTransaction Transaction(LOCTEXT("AddVectorParameter", "Add vector parameter"));
 	Track->Modify();
-	Track->AddColorParameterKey(Parameter.ParameterName, KeyTime, Parameter.DefaultValue);
+	Track->AddColorParameterKey(Parameter.ParameterName, KeyTime, RowIndex, Parameter.DefaultValue);
 	GetSequencer()->NotifyMovieSceneDataChanged(EMovieSceneDataChangeType::MovieSceneStructureItemAdded);
 }
 
diff --git a/Engine/Source/Editor/MovieSceneTools/Private/TrackEditors/MaterialParameterCollectionTrackEditor.h b/Engine/Source/Editor/MovieSceneTools/Private/TrackEditors/MaterialParameterCollectionTrackEditor.h
--- a/Engine/Source/Editor/MovieSceneTools/Private/TrackEditors/MaterialParameterCollectionTrackEditor.h
+++ b/Engine/Source/Editor/MovieSceneTools/Private/TrackEditors/MaterialParameterCollectionTrackEditor.h
@@ -54,12 +54,12 @@ public:
 private:
 
 	/** Provides the contents of the add parameter menu. */
-	TSharedRef<SWidget> OnGetAddParameterMenuContent(UMovieSceneMaterialParameterCollectionTrack* MaterialTrack);
+	TSharedRef<SWidget> OnGetAddParameterMenuContent(UMovieSceneMaterialParameterCollectionTrack* MaterialTrack, int32 RowIndex, int32 TrackInsertRowIndex);
 
 	void OnSelectMPC(UMaterialParameterCollection* MPC);
 
-	void AddScalarParameter(UMovieSceneMaterialParameterCollectionTrack* Track, FCollectionScalarParameter Parameter);
-	void AddVectorParameter(UMovieSceneMaterialParameterCollectionTrack* Track, FCollectionVectorParameter Parameter);
+	void AddScalarParameter(UMovieSceneMaterialParameterCollectionTrack* Track, int32 RowIndex, FCollectionScalarParameter Parameter);
+	void AddVectorParameter(UMovieSceneMaterialParameterCollectionTrack* Track, int32 RowIndex, FCollectionVectorParameter Parameter);
 
 	void AddTrackToSequence(const FAssetData& InAssetData);
 	void AddTrackToSequenceEnterPressed(const TArray<FAssetData>& InAssetData);
diff --git a/Engine/Source/Editor/MovieSceneTools/Private/TrackEditors/MaterialTrackEditor.cpp b/Engine/Source/Editor/MovieSceneTools/Private/TrackEditors/MaterialTrackEditor.cpp
--- a/Engine/Source/Editor/MovieSceneTools/Private/TrackEditors/MaterialTrackEditor.cpp
+++ b/Engine/Source/Editor/MovieSceneTools/Private/TrackEditors/MaterialTrackEditor.cpp
@@ -37,12 +37,37 @@ TSharedRef<ISequencerSection> FMaterialTrackEditor::MakeSectionInterface( UMovie
 TSharedPtr<SWidget> FMaterialTrackEditor::BuildOutlinerEditWidget( const FGuid& ObjectBinding, UMovieSceneTrack* Track, const FBuildEditWidgetParams& Params )
 {
 	UMovieSceneMaterialTrack* MaterialTrack = Cast<UMovieSceneMaterialTrack>(Track);
-	FOnGetContent MenuContent = FOnGetContent::CreateSP(this, &FMaterialTrackEditor::OnGetAddParameterMenuContent, ObjectBinding, MaterialTrack);
+	FOnGetContent MenuContent = FOnGetContent::CreateSP(this, &FMaterialTrackEditor::OnGetAddMenuContent, ObjectBinding, MaterialTrack, Params.TrackInsertRowIndex);
 
 	return FSequencerUtilities::MakeAddButton(LOCTEXT( "AddParameterButton", "Parameter" ), MenuContent, Params.NodeIsHovered, GetSequencer());
 }
 
 
+TSharedRef<SWidget> FMaterialTrackEditor::OnGetAddMenuContent( FGuid ObjectBinding, UMovieSceneMaterialTrack* MaterialTrack, int32 TrackInsertRowIndex )
+{
+	// IF this is supported, allow creating other sections with different blend types, and put
+	// the material parameters after a separator. Otherwise, just show the parameters menu.
+	const FMovieSceneBlendTypeField SupportedBlendTypes = MaterialTrack->GetSupportedBlendTypes();
+	if (SupportedBlendTypes.Num() > 1)
+	{
+		FMenuBuilder MenuBuilder(true, nullptr);
+
+		TWeakPtr<ISequencer> WeakSequencer = GetSequencer();
+		FSequencerUtilities::PopulateMenu_CreateNewSection(MenuBuilder, TrackInsertRowIndex, MaterialTrack, WeakSequencer);
+
+		MenuBuilder.AddSeparator();
+
+		OnBuildAddParameterMenu(MenuBuilder, ObjectBinding, MaterialTrack);
+
+		return MenuBuilder.MakeWidget();
+	}
+	else
+	{
+		return OnGetAddParameterMenuContent(ObjectBinding, MaterialTrack);
+	}
+}
+
+
 struct FParameterNameAndAction
 {
 	FName ParameterName;
@@ -64,7 +89,13 @@ struct FParameterNameAndAction
 TSharedRef<SWidget> FMaterialTrackEditor::OnGetAddParameterMenuContent( FGuid ObjectBinding, UMovieSceneMaterialTrack* MaterialTrack )
 {
 	FMenuBuilder AddParameterMenuBuilder( true, nullptr );
+	OnBuildAddParameterMenu(AddParameterMenuBuilder, ObjectBinding, MaterialTrack);
+	return AddParameterMenuBuilder.MakeWidget();
+}
 
+
+void FMaterialTrackEditor::OnBuildAddParameterMenu( FMenuBuilder& MenuBuilder, FGuid ObjectBinding, UMovieSceneMaterialTrack* MaterialTrack )
+{
 	UMaterial* Material = GetMaterialForTrack( ObjectBinding, MaterialTrack );
 	if ( Material != nullptr )
 	{
@@ -117,11 +148,9 @@ TSharedRef<SWidget> FMaterialTrackEditor::OnGetAddParameterMenuContent( FGuid Ob
 		ParameterNamesAndActions.Sort();
 		for ( FParameterNameAndAction NameAndAction : ParameterNamesAndActions )
 		{
-			AddParameterMenuBuilder.AddMenuEntry( FText::FromName( NameAndAction.ParameterName ), FText(), FSlateIcon(), NameAndAction.Action );
+			MenuBuilder.AddMenuEntry( FText::FromName( NameAndAction.ParameterName ), FText(), FSlateIcon(), NameAndAction.Action );
 		}
 	}
-
-	return AddParameterMenuBuilder.MakeWidget();
 }
 
 
diff --git a/Engine/Source/Editor/MovieSceneTools/Public/TrackEditors/MaterialTrackEditor.h b/Engine/Source/Editor/MovieSceneTools/Public/TrackEditors/MaterialTrackEditor.h
--- a/Engine/Source/Editor/MovieSceneTools/Public/TrackEditors/MaterialTrackEditor.h
+++ b/Engine/Source/Editor/MovieSceneTools/Public/TrackEditors/MaterialTrackEditor.h
@@ -45,22 +45,28 @@ protected:
 
 private:
 
+	/** Provides the contents of the outliner edit widget */
+	TSharedRef<SWidget> OnGetAddMenuContent( FGuid ObjectBinding, UMovieSceneMaterialTrack* MaterialTrack, int32 TrackInsertRowIndex );
+
 	/** Provides the contents of the add parameter menu. */
 	TSharedRef<SWidget> OnGetAddParameterMenuContent( FGuid ObjectBinding, UMovieSceneMaterialTrack* MaterialTrack );
 
+	/** Provides the contents of the add parameter menu. */
+	void OnBuildAddParameterMenu( FMenuBuilder& MenuBuilder, FGuid ObjectBinding, UMovieSceneMaterialTrack* MaterialTrack );
+
 	/** Gets a material for a specific object binding and track */
 	UMaterial* GetMaterialForTrack( FGuid ObjectBinding, UMovieSceneMaterialTrack* MaterialTrack );
 
 	/** Adds a scalar parameter and initial key to a material track.
 	 * @param ObjectBinding The object binding which owns the material track.
-	 * @param MaterialTrack The track to Add the section to.
+	 * @param MaterialTrack The track in which to look for sections to add the parameter to.
 	 * @param ParameterName The name of the parameter to add an initial key for.
 	 */
 	void AddScalarParameter( FGuid ObjectBinding, UMovieSceneMaterialTrack* MaterialTrack, FName ParameterName );
 
 	/** Adds a color parameter and initial key to a material track.
 	* @param ObjectBinding The object binding which owns the material track.
-	* @param MaterialTrack The track to Add the section to.
+	 * @param MaterialTrack The track in which to look for sections to add the parameter to.
 	* @param ParameterName The name of the parameter to add an initial key for.
 	*/
 	void AddColorParameter( FGuid ObjectBinding, UMovieSceneMaterialTrack* MaterialTrack, FName ParameterName );
diff --git a/Engine/Source/Editor/Sequencer/Private/MVVM/Views/SOutlinerTrackView.cpp b/Engine/Source/Editor/Sequencer/Private/MVVM/Views/SOutlinerTrackView.cpp
--- a/Engine/Source/Editor/Sequencer/Private/MVVM/Views/SOutlinerTrackView.cpp
+++ b/Engine/Source/Editor/Sequencer/Private/MVVM/Views/SOutlinerTrackView.cpp
@@ -69,13 +69,14 @@ void SOutlinerTrackView::Construct(
 
 	FBuildEditWidgetParams Params;
 	Params.NodeIsHovered = HoverState;
+	Params.RowIndex = RowIndex;
 	if (DataModel->FindAncestorOfType<ITrackExtension>())
 	{
-		Params.TrackInsertRowIndex = TrackExtension->GetRowIndex();
+		Params.TrackInsertRowIndex = RowIndex;
 	}
 	else if (Track && Track->SupportsMultipleRows())
 	{
-		Params.TrackInsertRowIndex = Track->GetMaxRowIndex()+1;
+		Params.TrackInsertRowIndex = Track->GetMaxRowIndex() + 1;
 	}
 
 	TSharedPtr<SWidget> CustomWidget = TrackEditor ? TrackEditor->BuildOutlinerEditWidget(ObjectBinding ? ObjectBinding->GetObjectGuid() : FGuid(), Track, Params) : nullptr;
diff --git a/Engine/Source/Editor/Sequencer/Public/ISequencerTrackEditor.h b/Engine/Source/Editor/Sequencer/Public/ISequencerTrackEditor.h
--- a/Engine/Source/Editor/Sequencer/Public/ISequencerTrackEditor.h
+++ b/Engine/Source/Editor/Sequencer/Public/ISequencerTrackEditor.h
@@ -26,12 +26,16 @@ class UMovieSceneSequence;
 struct FBuildEditWidgetParams
 {
 	FBuildEditWidgetParams()
-		: TrackInsertRowIndex(0)
+		: RowIndex(0)
+		, TrackInsertRowIndex(0)
 	{}
 
 	/** Attribute that specifies when the node relating to this edit widget is hovered */
 	TAttribute<bool> NodeIsHovered;
 
+	/** Index of current track row */
+	int32 RowIndex;
+
 	/** Track row index for any newly created sections */
 	int32 TrackInsertRowIndex;
 };
diff --git a/Engine/Source/Runtime/MovieSceneTracks/Private/Sections/MovieSceneParameterSection.cpp b/Engine/Source/Runtime/MovieSceneTracks/Private/Sections/MovieSceneParameterSection.cpp
--- a/Engine/Source/Runtime/MovieSceneTracks/Private/Sections/MovieSceneParameterSection.cpp
+++ b/Engine/Source/Runtime/MovieSceneTracks/Private/Sections/MovieSceneParameterSection.cpp
@@ -30,9 +30,9 @@ void DecodeEntityID(uint32 InEntityID, int32& OutIndex, uint8& OutType)
 
 }// namespace UE::MovieScene
 
-void IMovieSceneParameterSectionExtender::ExtendEntity(UMovieSceneEntitySystemLinker* EntityLinker, const UE::MovieScene::FEntityImportParams& Params, UE::MovieScene::FImportedEntity* OutImportedEntity)
+void IMovieSceneParameterSectionExtender::ExtendEntity(UMovieSceneParameterSection* Section, UMovieSceneEntitySystemLinker* EntityLinker, const UE::MovieScene::FEntityImportParams& Params, UE::MovieScene::FImportedEntity* OutImportedEntity)
 {
-	ExtendEntityImpl(EntityLinker, Params, OutImportedEntity);
+	ExtendEntityImpl(Section, EntityLinker, Params, OutImportedEntity);
 }
 
 FScalarParameterNameAndCurve::FScalarParameterNameAndCurve( FName InParameterName )
@@ -723,7 +723,7 @@ void UMovieSceneParameterSection::ImportEntityImpl(UMovieSceneEntitySystemLinker
 		}
 	}
 
-	Extender->ExtendEntity(EntityLinker, Params, OutImportedEntity);
+	Extender->ExtendEntity(this, EntityLinker, Params, OutImportedEntity);
 }
 
 bool UMovieSceneParameterSection::PopulateEvaluationFieldImpl(const TRange<FFrameNumber>& EffectiveRange, const FMovieSceneEvaluationFieldEntityMetaData& InMetaData, FMovieSceneEntityComponentFieldBuilder* OutFieldBuilder)
diff --git a/Engine/Source/Runtime/MovieSceneTracks/Private/Tracks/MovieSceneMaterialParameterCollectionTrack.cpp b/Engine/Source/Runtime/MovieSceneTracks/Private/Tracks/MovieSceneMaterialParameterCollectionTrack.cpp
--- a/Engine/Source/Runtime/MovieSceneTracks/Private/Tracks/MovieSceneMaterialParameterCollectionTrack.cpp
+++ b/Engine/Source/Runtime/MovieSceneTracks/Private/Tracks/MovieSceneMaterialParameterCollectionTrack.cpp
@@ -14,10 +14,6 @@
 UMovieSceneMaterialParameterCollectionTrack::UMovieSceneMaterialParameterCollectionTrack(const FObjectInitializer& ObjectInitializer)
 	: Super(ObjectInitializer)
 {
-	BuiltInTreePopulationMode = ETreePopulationMode::Blended;
-
-	SupportedBlendTypes.Add(EMovieSceneBlendType::Absolute);
-
 #if WITH_EDITORONLY_DATA
 	TrackTint = FColor(64,192,64,65);
 #endif
@@ -25,9 +21,8 @@ UMovieSceneMaterialParameterCollectionTrack::UMovieSceneMaterialParameterCollect
 
 UMovieSceneSection* UMovieSceneMaterialParameterCollectionTrack::CreateNewSection()
 {
-	UMovieSceneSection* NewSection = NewObject<UMovieSceneParameterSection>(this, NAME_None, RF_Transactional);
+	UMovieSceneSection* NewSection = Super::CreateNewSection();
 	NewSection->SetRange(TRange<FFrameNumber>::All());
-	NewSection->SetBlendType(EMovieSceneBlendType::Absolute);
 	return NewSection;
 }
 
@@ -37,7 +32,7 @@ void UMovieSceneMaterialParameterCollectionTrack::ImportEntityImpl(UMovieSceneEn
 	checkf(false, TEXT("This track should never have created entities for itself - this assertion indicates an error in the entity-component field"));
 }
 
-void UMovieSceneMaterialParameterCollectionTrack::ExtendEntityImpl(UMovieSceneEntitySystemLinker* EntityLinker, const UE::MovieScene::FEntityImportParams& Params, UE::MovieScene::FImportedEntity* OutImportedEntity)
+void UMovieSceneMaterialParameterCollectionTrack::ExtendEntityImpl(UMovieSceneParameterSection* Section, UMovieSceneEntitySystemLinker* EntityLinker, const UE::MovieScene::FEntityImportParams& Params, UE::MovieScene::FImportedEntity* OutImportedEntity)
 {
 	using namespace UE::MovieScene;
 
@@ -48,8 +43,10 @@ void UMovieSceneMaterialParameterCollectionTrack::ExtendEntityImpl(UMovieSceneEn
 	OutImportedEntity->AddBuilder(
 		FEntityBuilder()
 		.Add(TracksComponents->MPC, MPC)
-		.AddTag(BuiltInComponents->Tags.AbsoluteBlend)
 		.AddTag(BuiltInComponents->Tags.Root)
+		// If the section has no valid blend type (legacy data), make it use absolute blending.
+		// Otherwise, the base section class will add the appropriate blend type tag in BuildDefaultComponents.
+		.AddTagConditional(BuiltInComponents->Tags.AbsoluteBlend, !Section->GetBlendType().IsValid())
 	);
 }
 
diff --git a/Engine/Source/Runtime/MovieSceneTracks/Private/Tracks/MovieSceneMaterialTrack.cpp b/Engine/Source/Runtime/MovieSceneTracks/Private/Tracks/MovieSceneMaterialTrack.cpp
--- a/Engine/Source/Runtime/MovieSceneTracks/Private/Tracks/MovieSceneMaterialTrack.cpp
+++ b/Engine/Source/Runtime/MovieSceneTracks/Private/Tracks/MovieSceneMaterialTrack.cpp
@@ -17,7 +17,11 @@ UMovieSceneMaterialTrack::UMovieSceneMaterialTrack(const FObjectInitializer& Obj
 	TrackTint = FColor(64,192,64,65);
 #endif
 
+	BuiltInTreePopulationMode = ETreePopulationMode::Blended;
+
 	SupportedBlendTypes.Add(EMovieSceneBlendType::Absolute);
+	SupportedBlendTypes.Add(EMovieSceneBlendType::Additive);
+	SupportedBlendTypes.Add(EMovieSceneBlendType::AdditiveFromBase);
 }
 
 
@@ -26,15 +30,19 @@ bool UMovieSceneMaterialTrack::SupportsType(TSubclassOf<UMovieSceneSection> Sect
 	return SectionClass == UMovieSceneParameterSection::StaticClass();
 }
 
+
 UMovieSceneSection* UMovieSceneMaterialTrack::CreateNewSection()
 {
-	return NewObject<UMovieSceneParameterSection>(this, NAME_None, RF_Transactional);
+	UMovieSceneSection* NewSection = NewObject<UMovieSceneParameterSection>(this, NAME_None, RF_Transactional);
+	NewSection->SetBlendType(EMovieSceneBlendType::Absolute);
+	return NewSection;
 }
 
 
 void UMovieSceneMaterialTrack::RemoveAllAnimationData()
 {
 	Sections.Empty();
+	SectionToKey = nullptr;
 }
 
 
@@ -47,17 +55,40 @@ bool UMovieSceneMaterialTrack::HasSection(const UMovieSceneSection& Section) con
 void UMovieSceneMaterialTrack::AddSection(UMovieSceneSection& Section)
 {
 	Sections.Add(&Section);
+
+	if (Sections.Num() > 1)
+	{
+		SetSectionToKey(&Section);
+	}
 }
 
 
 void UMovieSceneMaterialTrack::RemoveSection(UMovieSceneSection& Section)
 {
 	Sections.Remove(&Section);
+	if (SectionToKey == &Section)
+	{
+		if (Sections.Num() > 0)
+		{
+			SectionToKey = Sections[0];
+		}
+		else
+		{
+			SectionToKey = nullptr;
+		}
+	}
 }
 
 void UMovieSceneMaterialTrack::RemoveSectionAt(int32 SectionIndex)
 {
+	bool bResetSectionToKey = (SectionToKey == Sections[SectionIndex]);
+
 	Sections.RemoveAt(SectionIndex);
+
+	if (bResetSectionToKey)
+	{
+		SectionToKey = Sections.Num() > 0 ? Sections[0] : nullptr;
+	}
 }
 
 
@@ -71,6 +102,16 @@ bool UMovieSceneMaterialTrack::SupportsMultipleRows() const
 	return true;
 }
 
+void UMovieSceneMaterialTrack::SetSectionToKey(UMovieSceneSection* InSection)
+{
+	SectionToKey = InSection;
+}
+
+UMovieSceneSection* UMovieSceneMaterialTrack::GetSectionToKey() const
+{
+	return SectionToKey;
+}
+
 const TArray<UMovieSceneSection*>& UMovieSceneMaterialTrack::GetAllSections() const
 {
 	return Sections;
@@ -79,7 +120,17 @@ const TArray<UMovieSceneSection*>& UMovieSceneMaterialTrack::GetAllSections() co
 
 void UMovieSceneMaterialTrack::AddScalarParameterKey(FName ParameterName, FFrameNumber Time, float Value)
 {
-	UMovieSceneParameterSection* NearestSection = Cast<UMovieSceneParameterSection>(MovieSceneHelpers::FindNearestSectionAtTime(Sections, Time));
+	AddScalarParameterKey(ParameterName, Time, INDEX_NONE, Value);
+}
+
+
+void UMovieSceneMaterialTrack::AddScalarParameterKey(FName ParameterName, FFrameNumber Time, int32 RowIndex, float Value)
+{
+	UMovieSceneParameterSection* NearestSection = Cast<UMovieSceneParameterSection>(SectionToKey);
+	if (NearestSection == nullptr || (RowIndex != INDEX_NONE && NearestSection->GetRowIndex() != RowIndex))
+	{
+		NearestSection = Cast<UMovieSceneParameterSection>(MovieSceneHelpers::FindNearestSectionAtTime(Sections, Time, RowIndex));
+	}
 	if (NearestSection == nullptr)
 	{
 		NearestSection = Cast<UMovieSceneParameterSection>(CreateNewSection());
@@ -99,7 +150,17 @@ void UMovieSceneMaterialTrack::AddScalarParameterKey(FName ParameterName, FFrame
 
 void UMovieSceneMaterialTrack::AddColorParameterKey(FName ParameterName, FFrameNumber Time, FLinearColor Value)
 {
-	UMovieSceneParameterSection* NearestSection = Cast<UMovieSceneParameterSection>(MovieSceneHelpers::FindNearestSectionAtTime(Sections, Time));
+	AddColorParameterKey(ParameterName, Time, INDEX_NONE, Value);
+}
+
+
+void UMovieSceneMaterialTrack::AddColorParameterKey(FName ParameterName, FFrameNumber Time, int32 RowIndex, FLinearColor Value)
+{
+	UMovieSceneParameterSection* NearestSection = Cast<UMovieSceneParameterSection>(SectionToKey);
+	if (NearestSection == nullptr || (RowIndex != INDEX_NONE && NearestSection->GetRowIndex() != RowIndex))
+	{
+		NearestSection = Cast<UMovieSceneParameterSection>(MovieSceneHelpers::FindNearestSectionAtTime(Sections, Time));
+	}
 	if (NearestSection == nullptr)
 	{
 		NearestSection = Cast<UMovieSceneParameterSection>(CreateNewSection());
@@ -124,20 +185,13 @@ UMovieSceneComponentMaterialTrack::UMovieSceneComponentMaterialTrack(const FObje
 	BuiltInTreePopulationMode = ETreePopulationMode::Blended;
 }
 
-void UMovieSceneComponentMaterialTrack::AddSection(UMovieSceneSection& Section)
-{
-	// Materials are always blendable now
-	Section.SetBlendType(EMovieSceneBlendType::Absolute);
-	Super::AddSection(Section);
-}
-
 void UMovieSceneComponentMaterialTrack::ImportEntityImpl(UMovieSceneEntitySystemLinker* EntityLinker, const FEntityImportParams& Params, FImportedEntity* OutImportedEntity)
 {
 	// These tracks don't define any entities for themselves
 	checkf(false, TEXT("This track should never have created entities for itself - this assertion indicates an error in the entity-component field"));
 }
 
-void UMovieSceneComponentMaterialTrack::ExtendEntityImpl(UMovieSceneEntitySystemLinker* EntityLinker, const UE::MovieScene::FEntityImportParams& Params, UE::MovieScene::FImportedEntity* OutImportedEntity)
+void UMovieSceneComponentMaterialTrack::ExtendEntityImpl(UMovieSceneParameterSection* Section, UMovieSceneEntitySystemLinker* EntityLinker, const UE::MovieScene::FEntityImportParams& Params, UE::MovieScene::FImportedEntity* OutImportedEntity)
 {
 	using namespace UE::MovieScene;
 
@@ -148,7 +202,9 @@ void UMovieSceneComponentMaterialTrack::ExtendEntityImpl(UMovieSceneEntitySystem
 	OutImportedEntity->AddBuilder(
 		FEntityBuilder()
 		.Add(TracksComponents->ComponentMaterialIndex, MaterialIndex)
-		.AddTag(BuiltInComponents->Tags.AbsoluteBlend)
+		// If the section has no valid blend type (legacy data), make it use absolute blending.
+		// Otherwise, the base section class will add the appropriate blend type tag in BuildDefaultComponents.
+		.AddTagConditional(BuiltInComponents->Tags.AbsoluteBlend, !Section->GetBlendType().IsValid())
 	);
 }
 
diff --git a/Engine/Source/Runtime/MovieSceneTracks/Public/Sections/MovieSceneParameterSection.h b/Engine/Source/Runtime/MovieSceneTracks/Public/Sections/MovieSceneParameterSection.h
--- a/Engine/Source/Runtime/MovieSceneTracks/Public/Sections/MovieSceneParameterSection.h
+++ b/Engine/Source/Runtime/MovieSceneTracks/Public/Sections/MovieSceneParameterSection.h
@@ -30,11 +30,11 @@ public:
 
 	GENERATED_BODY()
 
-	void ExtendEntity(UMovieSceneEntitySystemLinker* EntityLinker, const UE::MovieScene::FEntityImportParams& Params, UE::MovieScene::FImportedEntity* OutImportedEntity);
+	void ExtendEntity(UMovieSceneParameterSection* Section, UMovieSceneEntitySystemLinker* EntityLinker, const UE::MovieScene::FEntityImportParams& Params, UE::MovieScene::FImportedEntity* OutImportedEntity);
 
 private:
 
-	virtual void ExtendEntityImpl(UMovieSceneEntitySystemLinker* EntityLinker, const UE::MovieScene::FEntityImportParams& Params, UE::MovieScene::FImportedEntity* OutImportedEntity) = 0;
+	virtual void ExtendEntityImpl(UMovieSceneParameterSection* Section, UMovieSceneEntitySystemLinker* EntityLinker, const UE::MovieScene::FEntityImportParams& Params, UE::MovieScene::FImportedEntity* OutImportedEntity) = 0;
 };
 
 
@@ -479,4 +479,4 @@ protected:
 	/** The transform  parameter names and their associated curves. */
 	UPROPERTY()
 	TArray<FTransformParameterNameAndCurves> TransformParameterNamesAndCurves;
-};
\ No newline at end of file
+};
diff --git a/Engine/Source/Runtime/MovieSceneTracks/Public/Tracks/MovieSceneMaterialParameterCollectionTrack.h b/Engine/Source/Runtime/MovieSceneTracks/Public/Tracks/MovieSceneMaterialParameterCollectionTrack.h
--- a/Engine/Source/Runtime/MovieSceneTracks/Public/Tracks/MovieSceneMaterialParameterCollectionTrack.h
+++ b/Engine/Source/Runtime/MovieSceneTracks/Public/Tracks/MovieSceneMaterialParameterCollectionTrack.h
@@ -37,7 +37,7 @@ public:
 	virtual bool PopulateEvaluationFieldImpl(const TRange<FFrameNumber>& EffectiveRange, const FMovieSceneEvaluationFieldEntityMetaData& InMetaData, FMovieSceneEntityComponentFieldBuilder* OutFieldBuilder) override;
 
 	/*~ IMovieSceneParameterSectionExtender */
-	virtual void ExtendEntityImpl(UMovieSceneEntitySystemLinker* EntityLinker, const UE::MovieScene::FEntityImportParams& Params, UE::MovieScene::FImportedEntity* OutImportedEntity) override;
+	virtual void ExtendEntityImpl(UMovieSceneParameterSection* Section, UMovieSceneEntitySystemLinker* EntityLinker, const UE::MovieScene::FEntityImportParams& Params, UE::MovieScene::FImportedEntity* OutImportedEntity) override;
 
 #if WITH_EDITORONLY_DATA
 	virtual FText GetDefaultDisplayName() const override;
diff --git a/Engine/Source/Runtime/MovieSceneTracks/Public/Tracks/MovieSceneMaterialTrack.h b/Engine/Source/Runtime/MovieSceneTracks/Public/Tracks/MovieSceneMaterialTrack.h
--- a/Engine/Source/Runtime/MovieSceneTracks/Public/Tracks/MovieSceneMaterialTrack.h
+++ b/Engine/Source/Runtime/MovieSceneTracks/Public/Tracks/MovieSceneMaterialTrack.h
@@ -34,6 +34,8 @@ public:
 	virtual bool IsEmpty() const override;
 	virtual const TArray<UMovieSceneSection*>& GetAllSections() const override;	
 	virtual bool SupportsMultipleRows() const override;
+	virtual void SetSectionToKey(UMovieSceneSection* Section) override;
+	virtual UMovieSceneSection* GetSectionToKey() const override;
 
 public:
 
@@ -46,18 +48,40 @@ public:
 	void AddScalarParameterKey(FName ParameterName, FFrameNumber Position, float Value);
 
 	/**
-	* Adds a color parameter key to the track.
-	* @param ParameterName The name of the parameter to add a key for.
-	* @param Time The time to add the new key.
-	* @param The value for the new key.
-	*/
+	 * Adds a scalar parameter key to the track. 
+	 * @param ParameterName The name of the parameter to add a key for.
+	 * @param Time The time to add the new key.
+	 * @param RowIndex The preferred row index on which to look for sections.
+	 * @param The value for the new key.
+	 */
+	void AddScalarParameterKey(FName ParameterName, FFrameNumber Position, int32 RowIndex, float Value);
+
+	/**
+	 * Adds a color parameter key to the track.
+	 * @param ParameterName The name of the parameter to add a key for.
+	 * @param Time The time to add the new key.
+	 * @param The value for the new key.
+	 */
 	void AddColorParameterKey(FName ParameterName, FFrameNumber Position, FLinearColor Value);
 
+	/**
+	 * Adds a color parameter key to the track.
+	 * @param ParameterName The name of the parameter to add a key for.
+	 * @param Time The time to add the new key.
+	 * @param RowIndex The preferred row index on which to look for sections.
+	 * @param The value for the new key.
+	 */
+	void AddColorParameterKey(FName ParameterName, FFrameNumber Position, int32 RowIndex, FLinearColor Value);
+
 private:
 
 	/** The sections owned by this track .*/
 	UPROPERTY()
 	TArray<TObjectPtr<UMovieSceneSection>> Sections;
+
+	/** Section we should Key */
+	UPROPERTY()
+	TObjectPtr<UMovieSceneSection> SectionToKey;
 };
 
 
@@ -74,15 +98,12 @@ class UMovieSceneComponentMaterialTrack
 
 public:
 
-	// UMovieSceneTrack interface
-	virtual void AddSection(UMovieSceneSection& Section) override;
-
 	/*~ IMovieSceneEntityProvider */
 	virtual void ImportEntityImpl(UMovieSceneEntitySystemLinker* EntityLinker, const FEntityImportParams& Params, FImportedEntity* OutImportedEntity) override;
 	virtual bool PopulateEvaluationFieldImpl(const TRange<FFrameNumber>& EffectiveRange, const FMovieSceneEvaluationFieldEntityMetaData& InMetaData, FMovieSceneEntityComponentFieldBuilder* OutFieldBuilder) override;
 
 	/*~ IMovieSceneParameterSectionExtender */
-	virtual void ExtendEntityImpl(UMovieSceneEntitySystemLinker* EntityLinker, const UE::MovieScene::FEntityImportParams& Params, UE::MovieScene::FImportedEntity* OutImportedEntity) override;
+	virtual void ExtendEntityImpl(UMovieSceneParameterSection* Section, UMovieSceneEntitySystemLinker* EntityLinker, const UE::MovieScene::FEntityImportParams& Params, UE::MovieScene::FImportedEntity* OutImportedEntity) override;
 
 	virtual FName GetTrackName() const { return FName( *FString::FromInt(MaterialIndex) ); }
 
diff --git a/Engine/Source/Runtime/UMG/Private/Animation/MovieSceneWidgetMaterialTrack.cpp b/Engine/Source/Runtime/UMG/Private/Animation/MovieSceneWidgetMaterialTrack.cpp
--- a/Engine/Source/Runtime/UMG/Private/Animation/MovieSceneWidgetMaterialTrack.cpp
+++ b/Engine/Source/Runtime/UMG/Private/Animation/MovieSceneWidgetMaterialTrack.cpp
@@ -10,14 +10,6 @@
 UMovieSceneWidgetMaterialTrack::UMovieSceneWidgetMaterialTrack( const FObjectInitializer& ObjectInitializer )
 	: Super(ObjectInitializer)
 {
-	BuiltInTreePopulationMode = ETreePopulationMode::Blended;
-}
-
-void UMovieSceneWidgetMaterialTrack::AddSection(UMovieSceneSection& Section)
-{
-	// Materials are always blendable now
-	Section.SetBlendType(EMovieSceneBlendType::Absolute);
-	Super::AddSection(Section);
 }
 
 void UMovieSceneWidgetMaterialTrack::ImportEntityImpl(UMovieSceneEntitySystemLinker* EntityLinker, const FEntityImportParams& Params, FImportedEntity* OutImportedEntity)
@@ -26,18 +18,19 @@ void UMovieSceneWidgetMaterialTrack::ImportEntityImpl(UMovieSceneEntitySystemLin
 	checkf(false, TEXT("This track should never have created entities for itself - this assertion indicates an error in the entity-component field"));
 }
 
-void UMovieSceneWidgetMaterialTrack::ExtendEntityImpl(UMovieSceneEntitySystemLinker* EntityLinker, const UE::MovieScene::FEntityImportParams& Params, UE::MovieScene::FImportedEntity* OutImportedEntity)
+void UMovieSceneWidgetMaterialTrack::ExtendEntityImpl(UMovieSceneParameterSection* Section, UMovieSceneEntitySystemLinker* EntityLinker, const UE::MovieScene::FEntityImportParams& Params, UE::MovieScene::FImportedEntity* OutImportedEntity)
 {
 	using namespace UE::MovieScene;
 
 	FBuiltInComponentTypes*       BuiltInComponents = FBuiltInComponentTypes::Get();
 	FMovieSceneUMGComponentTypes* WidgetComponents  = FMovieSceneUMGComponentTypes::Get();
 
-	// Material parameters are always absolute blends for the time being
 	OutImportedEntity->AddBuilder(
 		FEntityBuilder()
 		.Add(WidgetComponents->WidgetMaterialPath, FWidgetMaterialPath(BrushPropertyNamePath))
-		.AddTag(BuiltInComponents->Tags.AbsoluteBlend)
+		// If the section has no valid blend type (legacy data), make it use absolute blending.
+		// Otherwise, the base section class will add the appropriate blend type tag in BuildDefaultComponents.
+		.AddTagConditional(BuiltInComponents->Tags.AbsoluteBlend, !Section->GetBlendType().IsValid())
 	);
 }
 
diff --git a/Engine/Source/Runtime/UMG/Public/Animation/MovieSceneWidgetMaterialTrack.h b/Engine/Source/Runtime/UMG/Public/Animation/MovieSceneWidgetMaterialTrack.h
--- a/Engine/Source/Runtime/UMG/Public/Animation/MovieSceneWidgetMaterialTrack.h
+++ b/Engine/Source/Runtime/UMG/Public/Animation/MovieSceneWidgetMaterialTrack.h
@@ -37,7 +37,6 @@ class UMovieSceneWidgetMaterialTrack
 public:
 
 	// UMovieSceneTrack interface
-	virtual void AddSection(UMovieSceneSection& Section) override;
 	virtual FName GetTrackName() const override;
 
 	/*~ IMovieSceneEntityProvider */
@@ -45,7 +44,7 @@ public:
 	virtual bool PopulateEvaluationFieldImpl(const TRange<FFrameNumber>& EffectiveRange, const FMovieSceneEvaluationFieldEntityMetaData& InMetaData, FMovieSceneEntityComponentFieldBuilder* OutFieldBuilder) override;
 
 	/*~ IMovieSceneParameterSectionExtender */
-	virtual void ExtendEntityImpl(UMovieSceneEntitySystemLinker* EntityLinker, const UE::MovieScene::FEntityImportParams& Params, UE::MovieScene::FImportedEntity* OutImportedEntity) override;
+	virtual void ExtendEntityImpl(UMovieSceneParameterSection* Section, UMovieSceneEntitySystemLinker* EntityLinker, const UE::MovieScene::FEntityImportParams& Params, UE::MovieScene::FImportedEntity* OutImportedEntity) override;
 
 #if WITH_EDITORONLY_DATA
 	virtual FText GetDefaultDisplayName() const override;
