Sequencer: Recache the audio section's channel proxy when the track has been renamed since the attach FAudioChannelEditorData depends on the outer chain

This fixes an issue when you copy/paste an audio track (and it's audio section)

diff --git a/Engine/Source/Runtime/MovieSceneTracks/Private/Tracks/MovieSceneAudioTrack.cpp b/Engine/Source/Runtime/MovieSceneTracks/Private/Tracks/MovieSceneAudioTrack.cpp
--- a/Engine/Source/Runtime/MovieSceneTracks/Private/Tracks/MovieSceneAudioTrack.cpp
+++ b/Engine/Source/Runtime/MovieSceneTracks/Private/Tracks/MovieSceneAudioTrack.cpp
@@ -144,5 +144,19 @@ UMovieSceneSection* UMovieSceneAudioTrack::CreateNewSection()
 	return NewObject<UMovieSceneAudioSection>(this, NAME_None, RF_Transactional);
 }
 
+void UMovieSceneAudioTrack::PostRename(UObject* OldOuter, const FName OldName)
+{
+	Super::PostRename(OldOuter, OldName);
+
+	// Recache the channel proxy because attach in FAudioChannelEditorData is dependent upon the outer chain
+	for (TObjectPtr<UMovieSceneSection> Section : AudioSections)
+	{
+		if (UMovieSceneAudioSection* AudioSection = Cast<UMovieSceneAudioSection>(Section.Get()))
+		{
+			AudioSection->CacheChannelProxy();
+		}
+	}
+}
+
 #undef LOCTEXT_NAMESPACE
 
diff --git a/Engine/Source/Runtime/MovieSceneTracks/Public/Tracks/MovieSceneAudioTrack.h b/Engine/Source/Runtime/MovieSceneTracks/Public/Tracks/MovieSceneAudioTrack.h
--- a/Engine/Source/Runtime/MovieSceneTracks/Public/Tracks/MovieSceneAudioTrack.h
+++ b/Engine/Source/Runtime/MovieSceneTracks/Public/Tracks/MovieSceneAudioTrack.h
@@ -57,6 +57,9 @@ public:
 	virtual bool SupportsMultipleRows() const override;
 	virtual UMovieSceneSection* CreateNewSection() override;
 
+	// ~UObject interface
+	virtual void PostRename(UObject* OldOuter, const FName OldName) override;
+
 private:
 
 	/** List of all root audio sections */
