Sequencer: Added custom, persistent scheduler for dispatching ECS tasks

diff --git a/Engine/Source/Runtime/MovieSceneTracks/Private/Systems/MovieScenePiecewiseDoubleBlenderSystem.cpp b/Engine/Source/Runtime/MovieSceneTracks/Private/Systems/MovieScenePiecewiseDoubleBlenderSystem.cpp
--- a/Engine/Source/Runtime/MovieSceneTracks/Private/Systems/MovieScenePiecewiseDoubleBlenderSystem.cpp
+++ b/Engine/Source/Runtime/MovieSceneTracks/Private/Systems/MovieScenePiecewiseDoubleBlenderSystem.cpp
@@ -22,6 +22,42 @@ namespace UE
 namespace MovieScene
 {
 
+struct FForkedAccumulationTask
+{
+	FForkedAccumulationTask(TArray<FBlendResult>* InAccumulationBuffer)
+		: AccumulationBuffer(InAccumulationBuffer)
+	{}
+
+	void ForEachAllocation(const FEntityAllocation* InAllocation, const double* InResults, const FMovieSceneBlendChannelID* BlendIDs, const double* OptionalEasingAndWeights) const
+	{
+		const int32 Num = InAllocation->Num();
+		if (OptionalEasingAndWeights)
+		{
+			// We have some easing/weight factors to multiply values with.
+			for (int32 Index = 0; Index < Num; ++Index)
+			{
+				FBlendResult& Result = (*AccumulationBuffer)[BlendIDs[Index].ChannelID];
+
+				const float Weight = OptionalEasingAndWeights[Index];
+				Result.Total += InResults[Index] * Weight;
+				Result.Weight += Weight;
+			}
+		}
+		else
+		{
+			// Faster path for when there's no weight to multiply values with.
+			for (int32 Index = 0; Index < Num; ++Index)
+			{
+				FBlendResult& Result = (*AccumulationBuffer)[BlendIDs[Index].ChannelID];
+				Result.Total += InResults[Index];
+				Result.Weight += 1.f;
+			}
+		}
+	}
+
+	TArray<FBlendResult>* AccumulationBuffer;
+};
+
 /** Task for accumulating all weighted blend inputs into arrays based on BlendID. Will be run for Absolute, Additive and Relative blend modes*/
 struct FAccumulationTask
 {
@@ -30,7 +66,7 @@ struct FAccumulationTask
 	{}
 
 	/** Task entry point - iterates the allocation's headers and accumulates results for any required components */
-	void ForEachAllocation(FEntityAllocationIteratorItem InItem, TRead<FMovieSceneBlendChannelID> BlendIDs, TReadOptional<double> OptionalEasingAndWeights)
+	void ForEachAllocation(FEntityAllocationIteratorItem InItem, TRead<FMovieSceneBlendChannelID> BlendIDs, TReadOptional<double> OptionalEasingAndWeights) const
 	{
 		const FEntityAllocation* Allocation = InItem;
 		const FComponentMask& AllocationType = InItem;
@@ -57,7 +93,7 @@ struct FAccumulationTask
 
 private:
 
-	void AccumulateResults(const FEntityAllocation* InAllocation, const double* InResults, const FMovieSceneBlendChannelID* BlendIDs, const double* OptionalEasingAndWeights, TArray<FBlendResult>& OutBlendResults)
+	void AccumulateResults(const FEntityAllocation* InAllocation, const double* InResults, const FMovieSceneBlendChannelID* BlendIDs, const double* OptionalEasingAndWeights, TArray<FBlendResult>& OutBlendResults) const
 	{
 		static const FMovieSceneBlenderSystemID BlenderSystemID = UMovieSceneBlenderSystem::GetBlenderSystemID<UMovieScenePiecewiseDoubleBlenderSystem>();
 
@@ -115,7 +151,7 @@ struct FAdditiveFromBaseBlendTask
 {
 	TSortedMap<FComponentTypeID, FAdditiveFromBaseBuffer>* AccumulationBuffers;
 
-	void ForEachAllocation(FEntityAllocationIteratorItem InItem, TRead<FMovieSceneBlendChannelID> BlendIDs, TReadOptional<double> EasingAndWeightResults)
+	void ForEachAllocation(FEntityAllocationIteratorItem InItem, TRead<FMovieSceneBlendChannelID> BlendIDs, TReadOptional<double> EasingAndWeightResults) const
 	{
 		FEntityAllocation* Allocation = InItem;
 		const FComponentMask& AllocationType = InItem;
@@ -134,7 +170,7 @@ struct FAdditiveFromBaseBlendTask
 
 private:
 
-	void AccumulateResults(const FEntityAllocation* InAllocation, const double* Results, const double* BaseValues, const FMovieSceneBlendChannelID* BlendIDs, const double* OptionalEasingAndWeights, TArray<FBlendResult>& OutBlendResults)
+	void AccumulateResults(const FEntityAllocation* InAllocation, const double* Results, const double* BaseValues, const FMovieSceneBlendChannelID* BlendIDs, const double* OptionalEasingAndWeights, TArray<FBlendResult>& OutBlendResults) const
 	{
 		static const FMovieSceneBlenderSystemID BlenderSystemID = UMovieSceneBlenderSystem::GetBlenderSystemID<UMovieScenePiecewiseDoubleBlenderSystem>();
 
@@ -188,17 +224,116 @@ void BlendResults(const FAccumulationResult& Results, uint16 BlendID, double& Ou
 	}
 }
 
+void BlendResultsWithInitial(const FAccumulationResult& Results, uint16 BlendID, const double InitialValue, double& OutFinalBlendResult)
+{
+	FBlendResult AbsoluteResult = Results.GetAbsoluteResult(BlendID);
+	FBlendResult RelativeResult = Results.GetRelativeResult(BlendID);
+	FBlendResult AdditiveResult = Results.GetAdditiveResult(BlendID);
+	FBlendResult AdditiveFromBaseResult = Results.GetAdditiveFromBaseResult(BlendID);
+
+	if (RelativeResult.Weight != 0)
+	{
+		RelativeResult.Total += InitialValue * RelativeResult.Weight;
+	}
+
+	FBlendResult TotalAdditiveResult = { AdditiveResult.Total + AdditiveFromBaseResult.Total, AdditiveResult.Weight + AdditiveFromBaseResult.Weight };
+
+	const float TotalWeight = AbsoluteResult.Weight + RelativeResult.Weight;
+	if (TotalWeight != 0)
+	{
+		// If the absolute value has some partial weighting (for ease-in/out for instance), we ramp it from/to the initial value. This means
+		// that the "initial value" adds a contribution to the entire blending process, so we add its weight to the total that we
+		// normalize absolutes and relatives with.
+		//
+		// Note that "partial weighting" means strictly between 0 and 100%. At 100% and above, we don't need to do this thing with the initial
+		// value. At 0%, we have no absolute value (only a relative value) and we therefore don't want to include the initial value either.
+		const bool bInitialValueContributes = (0.f < AbsoluteResult.Weight && AbsoluteResult.Weight < 1.f);
+		const double AbsoluteBlendedValue = bInitialValueContributes ?
+			(InitialValue * (1.f - AbsoluteResult.Weight) + AbsoluteResult.Total) :
+			AbsoluteResult.Total;
+		const float FinalTotalWeight = bInitialValueContributes ? (TotalWeight + (1.f - AbsoluteResult.Weight)) : TotalWeight;
+
+		const double Value = (AbsoluteBlendedValue + RelativeResult.Total) / FinalTotalWeight + TotalAdditiveResult.Total;
+		OutFinalBlendResult = Value;
+	}
+	else if (TotalAdditiveResult.Weight != 0)
+	{
+		OutFinalBlendResult = TotalAdditiveResult.Total + InitialValue;
+	}
+	else
+	{
+		OutFinalBlendResult = InitialValue;
+	}
+}
+
+/** Task that combines all accumulated blends for any tracked property type that has blend inputs/outputs */
+struct FScheduleCombineBlendsForProperties
+{
+	const FAccumulationBuffers* AccumulationBuffers;
+	TArrayView<const FPropertyCompositeDefinition> Composites;
+	FEntityAllocationWriteContext WriteContext;
+	uint32 DoubleCompositeMask;
+
+	FScheduleCombineBlendsForProperties(const FEntityManager* EntityManager, const FAccumulationBuffers* InAccumulationBuffers, const FPropertyDefinition* InPropertyDefinition)
+		: AccumulationBuffers(InAccumulationBuffers)
+		, Composites(FBuiltInComponentTypes::Get()->PropertyRegistry.GetComposites(*InPropertyDefinition))
+		, WriteContext(*EntityManager)
+		, DoubleCompositeMask(InPropertyDefinition->DoubleCompositeMask)
+	{}
+
+	void UpdateWriteContext(FEntityAllocationWriteContext InWriteContext)
+	{
+		WriteContext = InWriteContext;
+	}
+
+	void ForEachAllocation(FEntityAllocationIteratorItem Item, const FMovieSceneBlendChannelID* BlendIDs, FReadErasedOptional OptInitialValues) const
+	{
+		const FComponentMask& AllocationType = Item.GetAllocationType();
+		const FEntityAllocation* Allocation = Item.GetAllocation();
+
+		for (int32 CompositeIndex = 0; CompositeIndex < Composites.Num(); ++CompositeIndex)
+		{
+			const bool bIsCompositeSupported = ((DoubleCompositeMask & (1 << CompositeIndex)) != 0);
+			FAccumulationResult Results = AccumulationBuffers->FindResults(Composites[CompositeIndex].ComponentTypeID);
+			if (bIsCompositeSupported && AllocationType.Contains(Composites[CompositeIndex].ComponentTypeID) && Results.IsValid())
+			{
+				TComponentWriter<double> OutValues = Allocation->WriteComponents(Composites[CompositeIndex].ComponentTypeID.ReinterpretCast<double>(), WriteContext);
+				ForEachAllocation(Allocation, Results, BlendIDs, OutValues, OptInitialValues, Composites[CompositeIndex].CompositeOffset);
+			}
+		}
+	}
+
+	void ForEachAllocation(const FEntityAllocation* Allocation, FAccumulationResult Results, const FMovieSceneBlendChannelID* BlendIDs, double* OutValues, FReadErasedOptional OptInitialValues, uint16 InitialValueProjectionOffset) const
+	{
+		if (OptInitialValues.IsValid())
+		{
+			for (int32 Index = 0; Index < Allocation->Num(); ++Index)
+			{
+				const double InitialValue = *reinterpret_cast<const double*>(static_cast<const uint8*>(OptInitialValues[Index]) + InitialValueProjectionOffset);
+				BlendResultsWithInitial(Results, BlendIDs[Index].ChannelID, InitialValue, OutValues[Index]);
+			}
+		}
+		else
+		{
+			for (int32 Index = 0; Index < Allocation->Num(); ++Index)
+			{
+				BlendResults(Results, BlendIDs[Index].ChannelID, OutValues[Index]);
+			}
+		}
+	}
+};
+
 /** Task that combines all accumulated blends for any tracked property type that has blend inputs/outputs */
-struct FCombineBlendsForProperties
+struct FLegacyCombineBlendsForProperties
 {
-	explicit FCombineBlendsForProperties(const TBitArray<>& InCachedRelevantProperties, const FAccumulationBuffers* InAccumulationBuffers, FEntityAllocationWriteContext InWriteContext)
+	explicit FLegacyCombineBlendsForProperties(const TBitArray<>& InCachedRelevantProperties, const FAccumulationBuffers* InAccumulationBuffers, FEntityAllocationWriteContext InWriteContext)
 		: CachedRelevantProperties(InCachedRelevantProperties)
-		, AccumulationBuffers(InAccumulationBuffers)
 		, PropertyRegistry(&FBuiltInComponentTypes::Get()->PropertyRegistry)
+		, AccumulationBuffers(InAccumulationBuffers)
 		, WriteContext(InWriteContext)
 	{}
 
-	void ForEachAllocation(FEntityAllocationIteratorItem InItem, TRead<FMovieSceneBlendChannelID> BlendIDs)
+	void ForEachAllocation(FEntityAllocationIteratorItem InItem, TRead<FMovieSceneBlendChannelID> BlendIDs) const
 	{
 		FEntityAllocation* Allocation = InItem;
 		const FComponentMask& AllocationType = InItem;
@@ -217,7 +352,7 @@ struct FCombineBlendsForProperties
 
 private:
 
-	void ProcessPropertyType(FEntityAllocation* Allocation, const FComponentMask& AllocationType, const FPropertyDefinition& PropertyDefinition, const FMovieSceneBlendChannelID* BlendIDs)
+	void ProcessPropertyType(FEntityAllocation* Allocation, const FComponentMask& AllocationType, const FPropertyDefinition& PropertyDefinition, const FMovieSceneBlendChannelID* BlendIDs) const
 	{
 		static const FMovieSceneBlenderSystemID BlenderSystemID = UMovieSceneBlenderSystem::GetBlenderSystemID<UMovieScenePiecewiseDoubleBlenderSystem>();
 
@@ -268,75 +403,14 @@ private:
 		}
 	}
 
-	void BlendResultsWithInitial(const FAccumulationResult& Results, uint16 BlendID, const double InitialValue, double& OutFinalBlendResult)
-	{
-		FBlendResult AbsoluteResult = Results.GetAbsoluteResult(BlendID);
-		FBlendResult RelativeResult = Results.GetRelativeResult(BlendID);
-		FBlendResult AdditiveResult = Results.GetAdditiveResult(BlendID);
-		FBlendResult AdditiveFromBaseResult = Results.GetAdditiveFromBaseResult(BlendID);
-
-		if (RelativeResult.Weight != 0)
-		{
-			RelativeResult.Total += InitialValue * RelativeResult.Weight;
-		}
-
-		FBlendResult TotalAdditiveResult = { AdditiveResult.Total + AdditiveFromBaseResult.Total, AdditiveResult.Weight + AdditiveFromBaseResult.Weight };
-
-		const float TotalWeight = AbsoluteResult.Weight + RelativeResult.Weight;
-		if (TotalWeight != 0)
-		{
-			// If the absolute value has some partial weighting (for ease-in/out for instance), we ramp it from/to the initial value. This means
-			// that the "initial value" adds a contribution to the entire blending process, so we add its weight to the total that we
-			// normalize absolutes and relatives with.
-			//
-			// Note that "partial weighting" means strictly between 0 and 100%. At 100% and above, we don't need to do this thing with the initial
-			// value. At 0%, we have no absolute value (only a relative value) and we therefore don't want to include the initial value either.
-			const bool bInitialValueContributes = (0.f < AbsoluteResult.Weight && AbsoluteResult.Weight < 1.f);
-			const double AbsoluteBlendedValue = bInitialValueContributes ?
-				(InitialValue * (1.f - AbsoluteResult.Weight) + AbsoluteResult.Total) :
-				AbsoluteResult.Total;
-			const float FinalTotalWeight = bInitialValueContributes ? (TotalWeight + (1.f - AbsoluteResult.Weight)) : TotalWeight;
-
-			const double Value = (AbsoluteBlendedValue + RelativeResult.Total) / FinalTotalWeight + TotalAdditiveResult.Total;
-			OutFinalBlendResult = Value;
-		}
-		else if (TotalAdditiveResult.Weight != 0)
-		{
-			OutFinalBlendResult = TotalAdditiveResult.Total + InitialValue;
-		}
-		else
-		{
-			OutFinalBlendResult = InitialValue;
-		}
-	}
-
-	void BlendResults(const FAccumulationResult& Results, uint16 BlendID, double& OutFinalBlendResult)
-	{
-		FBlendResult AbsoluteResult = Results.GetAbsoluteResult(BlendID);
-		FBlendResult AdditiveResult = Results.GetAdditiveResult(BlendID);
-		FBlendResult AdditiveFromBaseResult = Results.GetAdditiveFromBaseResult(BlendID);
-
-#if DO_GUARD_SLOW
-		ensureMsgf(AbsoluteResult.Weight != 0.f, TEXT("Default blend combine being used for an entity that has no absolute weight. This should have an initial value and should be handled by each system, and excluded by default with UMovieSceneBlenderSystem::FinalCombineExclusionFilter ."));
-#endif
-
-		const float TotalWeight = AbsoluteResult.Weight;
-		if (TotalWeight != 0)
-		{
-			const double Value = AbsoluteResult.Total / AbsoluteResult.Weight + AdditiveResult.Total + AdditiveFromBaseResult.Total;
-			OutFinalBlendResult = Value;
-		}
-	}
-
 private:
 
 	TBitArray<> CachedRelevantProperties;
-	const FAccumulationBuffers* AccumulationBuffers;
 	const FPropertyRegistry* PropertyRegistry;
+	const FAccumulationBuffers* AccumulationBuffers;
 	FEntityAllocationWriteContext WriteContext;
 };
 
-
 /** Task that combines all accumulated blends for any tracked non-property type that has blend inputs/outputs */
 struct FCombineBlends
 {
@@ -345,7 +419,7 @@ struct FCombineBlends
 		, WriteContext(InWriteContext)
 	{}
 
-	void ForEachAllocation(FEntityAllocationIteratorItem InItem, TRead<FMovieSceneBlendChannelID> BlendIDs)
+	void ForEachAllocation(FEntityAllocationIteratorItem InItem, TRead<FMovieSceneBlendChannelID> BlendIDs) const
 	{
 		static const FMovieSceneBlenderSystemID BlenderSystemID = UMovieSceneBlenderSystem::GetBlenderSystemID<UMovieScenePiecewiseDoubleBlenderSystem>();
 
@@ -425,6 +499,10 @@ FAccumulationResult FAccumulationBuffers::FindResults(FComponentTypeID InCompone
 UMovieScenePiecewiseDoubleBlenderSystem::UMovieScenePiecewiseDoubleBlenderSystem(const FObjectInitializer& ObjInit)
 	: Super(ObjInit)
 {
+	using namespace UE::MovieScene;
+
+	Phase = ESystemPhase::Scheduling;
+
 	if (HasAnyFlags(RF_ClassDefaultObject))
 	{
 		DefineImplicitPrerequisite(UFloatChannelEvaluatorSystem::StaticClass(), GetClass());
@@ -435,6 +513,166 @@ UMovieScenePiecewiseDoubleBlenderSystem::UMovieScenePiecewiseDoubleBlenderSystem
 	}
 }
 
+void UMovieScenePiecewiseDoubleBlenderSystem::OnSchedulePersistentTasks(UE::MovieScene::IEntitySystemScheduler* TaskScheduler)
+{
+	using namespace UE::MovieScene;
+
+	CompactBlendChannels();
+
+	FEntityManager& EntityManager = Linker->EntityManager;
+	const TStatId BlendValuesStatId = GET_STATID(MovieSceneEval_BlendDoubleValues);
+	const TStatId CombineBlendsStatId = GET_STATID(MovieSceneEval_BlendCombineDoubleValues);
+
+	// We allocate space for every blend even if there are gaps so we can do a straight index into each array
+	if (AllocatedBlendChannels.Num() == 0)
+	{
+		return;
+	}
+
+	ReinitializeAccumulationBuffers();
+	if (AccumulationBuffers.IsEmpty())
+	{
+		return;
+	}
+
+	FBuiltInComponentTypes* BuiltInComponents = FBuiltInComponentTypes::Get();
+
+	FTaskID ResetWeightsTask = TaskScheduler->AddMemberFunctionTask(
+		FTaskParams(TEXT("Reset Double Blender Weights")),
+		this, &UMovieScenePiecewiseDoubleBlenderSystem::ZeroAccumulationBuffers);
+
+	FTaskID SyncTask = TaskScheduler->AddNullTask();
+
+	if (AccumulationBuffers.Absolute.Num() != 0)
+	{
+		for (TPair<FComponentTypeID, TArray<FBlendResult>>& Pair : AccumulationBuffers.Absolute)
+		{
+			FTaskID AbsoluteTask = FEntityTaskBuilder()
+			.Read(Pair.Key.ReinterpretCast<double>())
+			.Read(BuiltInComponents->BlendChannelInput)
+			.ReadOptional(BuiltInComponents->WeightAndEasingResult)
+			.FilterAll({ BuiltInComponents->Tags.AbsoluteBlend, GetBlenderTypeTag() })
+			.FilterAny(BlendedResultMask)
+			.FilterNone({ BuiltInComponents->Tags.Ignored })
+			.AddDynamicReadDependency(BlendedResultMask)
+			.SetStat(BlendValuesStatId)
+			.Schedule_PerAllocation<FForkedAccumulationTask>(&EntityManager, TaskScheduler, &Pair.Value);
+
+			TaskScheduler->AddPrerequisite(ResetWeightsTask, AbsoluteTask);
+			TaskScheduler->AddPrerequisite(AbsoluteTask, SyncTask);
+		}
+	}
+
+	if (AccumulationBuffers.Relative.Num() != 0)
+	{
+		for (TPair<FComponentTypeID, TArray<FBlendResult>>& Pair : AccumulationBuffers.Relative)
+		{
+			FTaskID RelativeTask = FEntityTaskBuilder()
+			.Read(Pair.Key.ReinterpretCast<double>())
+			.Read(BuiltInComponents->BlendChannelInput)
+			.ReadOptional(BuiltInComponents->WeightAndEasingResult)
+			.FilterAll({ BuiltInComponents->Tags.RelativeBlend, GetBlenderTypeTag() })
+			.FilterAny(BlendedResultMask)
+			.FilterNone({ BuiltInComponents->Tags.Ignored })
+			.AddDynamicReadDependency(BlendedResultMask)
+			.SetStat(BlendValuesStatId)
+			.Schedule_PerAllocation<FForkedAccumulationTask>(&EntityManager, TaskScheduler, &Pair.Value);
+
+			TaskScheduler->AddPrerequisite(ResetWeightsTask, RelativeTask);
+			TaskScheduler->AddPrerequisite(RelativeTask, SyncTask);
+		}
+	}
+
+	if (AccumulationBuffers.Additive.Num() != 0)
+	{
+		for (TPair<FComponentTypeID, TArray<FBlendResult>>& Pair : AccumulationBuffers.Additive)
+		{
+			FTaskID AdditiveTask = FEntityTaskBuilder()
+			.Read(Pair.Key.ReinterpretCast<double>())
+			.Read(BuiltInComponents->BlendChannelInput)
+			.ReadOptional(BuiltInComponents->WeightAndEasingResult)
+			.FilterAll({ BuiltInComponents->Tags.AdditiveBlend, GetBlenderTypeTag() })
+			.FilterAny(BlendedResultMask)
+			.FilterNone({ BuiltInComponents->Tags.Ignored })
+			.AddDynamicReadDependency(BlendedResultMask)
+			.SetStat(BlendValuesStatId)
+			.Schedule_PerAllocation<FForkedAccumulationTask>(&EntityManager, TaskScheduler, &Pair.Value);
+
+			TaskScheduler->AddPrerequisite(ResetWeightsTask, AdditiveTask);
+			TaskScheduler->AddPrerequisite(AdditiveTask, SyncTask);
+		}
+	}
+
+	if (AccumulationBuffers.AdditiveFromBase.Num() != 0)
+	{
+		FTaskID AdditiveFromBaseTask = FEntityTaskBuilder()
+		.Read(BuiltInComponents->BlendChannelInput)
+		.ReadOptional(BuiltInComponents->WeightAndEasingResult)
+		.FilterAll({ BuiltInComponents->Tags.AdditiveFromBaseBlend, GetBlenderTypeTag() })
+		.FilterAny(BlendedResultMask)
+		.FilterNone({ BuiltInComponents->Tags.Ignored })
+		.AddDynamicReadDependency(BlendedResultMask)
+		.SetStat(BlendValuesStatId)
+		.Schedule_PerAllocation<FAdditiveFromBaseBlendTask>(&EntityManager, TaskScheduler, &AccumulationBuffers.AdditiveFromBase);
+
+		TaskScheduler->AddPrerequisite(ResetWeightsTask, AdditiveFromBaseTask);
+		TaskScheduler->AddPrerequisite(AdditiveFromBaseTask, SyncTask);
+	}
+
+	// Combine blends for all blend outputs on properties
+	if (BlendedPropertyMask.Find(true) != INDEX_NONE)
+	{
+		// Schedule tasks for every allocation that has a blended property
+		for (TConstSetBitIterator<> PropertyIndex(CachedRelevantProperties); PropertyIndex; ++PropertyIndex)
+		{
+			const FPropertyDefinition& PropertyDefinition = BuiltInComponents->PropertyRegistry.GetDefinition(FCompositePropertyTypeID::FromIndex(PropertyIndex.GetIndex()));
+			TArrayView<const FPropertyCompositeDefinition> Composites = BuiltInComponents->PropertyRegistry.GetComposites(PropertyDefinition);
+
+			FComponentMask WriteDependencies;
+			for (int32 CompositeIndex = 0; CompositeIndex < Composites.Num(); ++CompositeIndex)
+			{
+				const bool bIsCompositeSupported = ((PropertyDefinition.DoubleCompositeMask & (1 << CompositeIndex)) != 0);
+				if (bIsCompositeSupported)
+				{
+					WriteDependencies.Set(Composites[CompositeIndex].ComponentTypeID);
+				}
+			}
+
+			FTaskID CombineTask = FEntityTaskBuilder()
+			.Read(BuiltInComponents->BlendChannelOutput)
+			.ReadErasedOptional(PropertyDefinition.InitialValueType)
+			.FilterAll({ GetBlenderTypeTag() })
+			.FilterAny({ PropertyDefinition.PropertyType })
+			.AddDynamicWriteDependency(WriteDependencies)
+			.SetStat(CombineBlendsStatId)
+			.Fork_PerAllocation<FScheduleCombineBlendsForProperties>(
+				&EntityManager,
+				TaskScheduler,
+				&EntityManager,
+				&AccumulationBuffers,
+				&PropertyDefinition
+			);
+
+			TaskScheduler->AddPrerequisite(SyncTask, CombineTask);
+		}
+	}
+
+	if (bContainsNonPropertyBlends)
+	{
+		// Blend task that combines vanilla (non-property-based) components
+		FTaskID CombineTask = FEntityTaskBuilder()
+		.Read(BuiltInComponents->BlendChannelOutput)
+		.FilterAll({ GetBlenderTypeTag() })
+		.FilterAny(BlendedResultMask)
+		.FilterNone(BlendedPropertyMask)
+		.AddDynamicWriteDependency(MakeArrayView(FBuiltInComponentTypes::Get()->DoubleResult))
+		.SetStat(CombineBlendsStatId)
+		.Fork_PerAllocation<FCombineBlends>(&EntityManager, TaskScheduler, &AccumulationBuffers, FEntityAllocationWriteContext(EntityManager));
+		
+		TaskScheduler->AddPrerequisite(SyncTask, CombineTask);
+	}
+}
+
 void UMovieScenePiecewiseDoubleBlenderSystem::OnRun(FSystemTaskPrerequisites& InPrerequisites, FSystemSubsequentTasks& Subsequents)
 {
 	using namespace UE::MovieScene;
@@ -456,14 +694,16 @@ void UMovieScenePiecewiseDoubleBlenderSystem::OnRun(FSystemTaskPrerequisites& In
 	{
 		ReinitializeAccumulationBuffers();
 	}
+	else
+	{
+		ZeroAccumulationBuffers();
+	}
 
 	if (AccumulationBuffers.IsEmpty())
 	{
 		return;
 	}
 
-	ZeroAccumulationBuffers();
-
 	FBuiltInComponentTypes* BuiltInComponents = FBuiltInComponentTypes::Get();
 
 	FSystemTaskPrerequisites Prereqs;
@@ -476,7 +716,7 @@ void UMovieScenePiecewiseDoubleBlenderSystem::OnRun(FSystemTaskPrerequisites& In
 		.FilterAny(BlendedResultMask)
 		.FilterNone({ BuiltInComponents->Tags.Ignored })
 		.SetStat(BlendValuesStatId)
-		.template Dispatch_PerAllocation<FAccumulationTask>(&EntityManager, InPrerequisites, nullptr, &AccumulationBuffers.Absolute);
+		.Dispatch_PerAllocation<FAccumulationTask>(&EntityManager, InPrerequisites, nullptr, &AccumulationBuffers.Absolute);
 
 		if (Task)
 		{
@@ -493,7 +733,7 @@ void UMovieScenePiecewiseDoubleBlenderSystem::OnRun(FSystemTaskPrerequisites& In
 		.FilterAny(BlendedResultMask)
 		.FilterNone({ BuiltInComponents->Tags.Ignored })
 		.SetStat(BlendValuesStatId)
-		.template Dispatch_PerAllocation<FAccumulationTask>(&EntityManager, InPrerequisites, nullptr, &AccumulationBuffers.Relative);
+		.Dispatch_PerAllocation<FAccumulationTask>(&EntityManager, InPrerequisites, nullptr, &AccumulationBuffers.Relative);
 
 		if (Task)
 		{
@@ -510,7 +750,7 @@ void UMovieScenePiecewiseDoubleBlenderSystem::OnRun(FSystemTaskPrerequisites& In
 		.FilterAny(BlendedResultMask)
 		.FilterNone({ BuiltInComponents->Tags.Ignored })
 		.SetStat(BlendValuesStatId)
-		.template Dispatch_PerAllocation<FAccumulationTask>(&EntityManager, InPrerequisites, nullptr, &AccumulationBuffers.Additive);
+		.Dispatch_PerAllocation<FAccumulationTask>(&EntityManager, InPrerequisites, nullptr, &AccumulationBuffers.Additive);
 
 		if (Task)
 		{
@@ -527,7 +767,7 @@ void UMovieScenePiecewiseDoubleBlenderSystem::OnRun(FSystemTaskPrerequisites& In
 		.FilterAny(BlendedResultMask)
 		.FilterNone({ BuiltInComponents->Tags.Ignored })
 		.SetStat(BlendValuesStatId)
-		.template Dispatch_PerAllocation<FAdditiveFromBaseBlendTask>(&EntityManager, InPrerequisites, nullptr, &AccumulationBuffers.AdditiveFromBase);
+		.Dispatch_PerAllocation<FAdditiveFromBaseBlendTask>(&EntityManager, InPrerequisites, nullptr, &AccumulationBuffers.AdditiveFromBase);
 
 		if (Task)
 		{
@@ -543,7 +783,7 @@ void UMovieScenePiecewiseDoubleBlenderSystem::OnRun(FSystemTaskPrerequisites& In
 		.FilterAll({ GetBlenderTypeTag() })
 		.FilterAny(BlendedPropertyMask)
 		.SetStat(CombineBlendsStatId)
-		.template Dispatch_PerAllocation<FCombineBlendsForProperties>(&EntityManager, Prereqs, &Subsequents, CachedRelevantProperties, &AccumulationBuffers, FEntityAllocationWriteContext(EntityManager));
+		.Dispatch_PerAllocation<FLegacyCombineBlendsForProperties>(&EntityManager, Prereqs, &Subsequents, CachedRelevantProperties, &AccumulationBuffers, FEntityAllocationWriteContext(EntityManager));
 	}
 
 	if (bContainsNonPropertyBlends)
@@ -555,7 +795,7 @@ void UMovieScenePiecewiseDoubleBlenderSystem::OnRun(FSystemTaskPrerequisites& In
 		.FilterAny(BlendedResultMask)
 		.FilterNone(BlendedPropertyMask)
 		.SetStat(CombineBlendsStatId)
-		.template Dispatch_PerAllocation<FCombineBlends>(&EntityManager, Prereqs, &Subsequents, &AccumulationBuffers, FEntityAllocationWriteContext(EntityManager));
+		.Dispatch_PerAllocation<FCombineBlends>(&EntityManager, Prereqs, &Subsequents, &AccumulationBuffers, FEntityAllocationWriteContext(EntityManager));
 	}
 }
 
@@ -624,6 +864,8 @@ void UMovieScenePiecewiseDoubleBlenderSystem::ReinitializeAccumulationBuffers()
 		return;
 	}
 
+	ZeroAccumulationBuffers();
+
 	FComponentMask AllPropertyTypes;
 
 	// This code works on the assumption that properties can never be removed (which is safe)
@@ -774,7 +1016,7 @@ FGraphEventRef UMovieScenePiecewiseDoubleBlenderSystem::DispatchDecomposeTask(co
 			.ReadOptional(BaseValueComponentType)
 			.ReadOptional(BuiltInComponents->WeightAndEasingResult)
 			.FilterAll({ Params.PropertyTag, GetBlenderTypeTag() })
-			.template Dispatch_PerAllocation<FChannelResultTask>(&Linker->EntityManager, FSystemTaskPrerequisites(), nullptr, Params, Output);
+			.Dispatch_PerAllocation<FChannelResultTask>(&Linker->EntityManager, FSystemTaskPrerequisites(), nullptr, Params, Output);
 	}
 	else
 	{
@@ -785,7 +1027,7 @@ FGraphEventRef UMovieScenePiecewiseDoubleBlenderSystem::DispatchDecomposeTask(co
 			.ReadOptional(BaseValueComponentType)
 			.ReadOptional(BuiltInComponents->WeightAndEasingResult)
 			.FilterAll({ Params.PropertyTag, GetBlenderTypeTag() })
-			.template Dispatch_PerAllocation<FChannelResultTask>(&Linker->EntityManager, FSystemTaskPrerequisites(), nullptr, Params, Output);
+			.Dispatch_PerAllocation<FChannelResultTask>(&Linker->EntityManager, FSystemTaskPrerequisites(), nullptr, Params, Output);
 	}
 }
 
diff --git a/Engine/Source/Runtime/MovieSceneTracks/Private/Systems/MovieScenePiecewiseIntegerBlenderSystem.cpp b/Engine/Source/Runtime/MovieSceneTracks/Private/Systems/MovieScenePiecewiseIntegerBlenderSystem.cpp
--- a/Engine/Source/Runtime/MovieSceneTracks/Private/Systems/MovieScenePiecewiseIntegerBlenderSystem.cpp
+++ b/Engine/Source/Runtime/MovieSceneTracks/Private/Systems/MovieScenePiecewiseIntegerBlenderSystem.cpp
@@ -28,7 +28,7 @@ struct FIntegerAccumulationTask
 	{}
 
 	/** Task entry point - iterates the allocation's headers and accumulates int32 results for any required components */
-	void ForEachAllocation(const FEntityAllocation* Allocation, TRead<FMovieSceneBlendChannelID> BlendIDs, TRead<int32> Integers, TReadOptional<double> EasingAndWeights)
+	void ForEachAllocation(const FEntityAllocation* Allocation, TRead<FMovieSceneBlendChannelID> BlendIDs, TRead<int32> Integers, TReadOptional<double> EasingAndWeights) const
 	{
 		static const FMovieSceneBlenderSystemID IntegerBlenderSystemID = UMovieSceneBlenderSystem::GetBlenderSystemID<UMovieScenePiecewiseIntegerBlenderSystem>();
 
@@ -78,7 +78,15 @@ struct FIntegerAdditiveFromBaseBlendTask
 		: AccumulationBuffer(*InAccumulationBuffer)
 	{}
 
-	void ForEachAllocation(const FEntityAllocation* Allocation, TRead<FMovieSceneBlendChannelID> BlendIDs, TRead<int32> Integers, TRead<int32> BaseValues, TReadOptional<double> EasingAndWeights)
+	void PreTask()
+	{
+		if (AccumulationBuffer.Num() > 0)
+		{
+			FMemory::Memzero(AccumulationBuffer.GetData(), sizeof(FIntegerBlendResult) * AccumulationBuffer.Num());
+		}
+	}
+
+	void ForEachAllocation(const FEntityAllocation* Allocation, TRead<FMovieSceneBlendChannelID> BlendIDs, TRead<int32> Integers, TRead<int32> BaseValues, TReadOptional<double> EasingAndWeights) const
 	{
 		static const FMovieSceneBlenderSystemID IntegerBlenderSystemID = UMovieSceneBlenderSystem::GetBlenderSystemID<UMovieScenePiecewiseIntegerBlenderSystem>();
 
@@ -124,7 +132,7 @@ struct FIntegerCombineBlends
 		: AccumulationBuffers(*InAccumulationBuffers)
 	{}
 
-	void ForEachAllocation(FEntityAllocation* Allocation, TRead<FMovieSceneBlendChannelID> BlendIDs, TReadOptional<int32> InitialValues, TWrite<int32> IntegerResults)
+	void ForEachAllocation(FEntityAllocation* Allocation, TRead<FMovieSceneBlendChannelID> BlendIDs, TReadOptional<int32> InitialValues, TWrite<int32> IntegerResults) const
 	{
 		static const FMovieSceneBlenderSystemID IntegerBlenderSystemID = UMovieSceneBlenderSystem::GetBlenderSystemID<UMovieScenePiecewiseIntegerBlenderSystem>();
 
@@ -146,7 +154,7 @@ struct FIntegerCombineBlends
 		}
 	}
 
-	void BlendResultsWithInitial(uint16 BlendID, const int32 InitialValue, int32& OutFinalBlendResult)
+	void BlendResultsWithInitial(uint16 BlendID, const int32 InitialValue, int32& OutFinalBlendResult) const
 	{
 		FIntegerBlendResult AbsoluteResult = AccumulationBuffers.Absolute.Num() > 0 ? AccumulationBuffers.Absolute[BlendID] : FIntegerBlendResult();
 		FIntegerBlendResult RelativeResult = AccumulationBuffers.Relative.Num() > 0 ? AccumulationBuffers.Relative[BlendID] : FIntegerBlendResult();
@@ -188,7 +196,7 @@ struct FIntegerCombineBlends
 		}
 	}
 
-	void BlendResults(uint16 BlendID, int32& OutFinalBlendResult)
+	void BlendResults(uint16 BlendID, int32& OutFinalBlendResult) const
 	{
 		FIntegerBlendResult AbsoluteResult = AccumulationBuffers.Absolute.Num() > 0 ? AccumulationBuffers.Absolute[BlendID] : FIntegerBlendResult();
 		FIntegerBlendResult AdditiveResult = AccumulationBuffers.Additive.Num() > 0 ? AccumulationBuffers.Additive[BlendID] : FIntegerBlendResult();
@@ -232,6 +240,10 @@ void FIntegerAccumulationBuffers::Reset()
 UMovieScenePiecewiseIntegerBlenderSystem::UMovieScenePiecewiseIntegerBlenderSystem(const FObjectInitializer& ObjInit)
 	: Super(ObjInit)
 {
+	using namespace UE::MovieScene;
+
+	Phase = ESystemPhase::Scheduling;
+
 	if (HasAnyFlags(RF_ClassDefaultObject))
 	{
 		DefineImplicitPrerequisite(UIntegerChannelEvaluatorSystem::StaticClass(), GetClass());
@@ -243,6 +255,102 @@ void UMovieScenePiecewiseIntegerBlenderSystem::OnLink()
 {
 }
 
+void UMovieScenePiecewiseIntegerBlenderSystem::OnSchedulePersistentTasks(UE::MovieScene::IEntitySystemScheduler* TaskScheduler)
+{
+	using namespace UE::MovieScene;
+
+	CompactBlendChannels();
+
+	// We allocate space for every blend even if there are gaps so we can do a straight index into each array
+	const int32 MaximumNumBlends = AllocatedBlendChannels.Num();
+	if (MaximumNumBlends == 0)
+	{
+		return;
+	}
+
+	ReinitializeAccumulationBuffers();
+	if (AccumulationBuffers.IsEmpty())
+	{
+		return;
+	}
+
+	const FBuiltInComponentTypes* BuiltInComponents = FBuiltInComponentTypes::Get();
+	const FMovieSceneTracksComponentTypes* TracksComponents = FMovieSceneTracksComponentTypes::Get();
+
+	FTaskID ResetWeightsTask = TaskScheduler->AddMemberFunctionTask(
+		FTaskParams(TEXT("Reset Integer Blender Weights")),
+		this, &UMovieScenePiecewiseIntegerBlenderSystem::ZeroAccumulationBuffers);
+
+	FTaskID SyncTask = TaskScheduler->AddNullTask();
+
+	if (AccumulationBuffers.Absolute.Num() != 0)
+	{
+		FTaskID AbsoluteTask = FEntityTaskBuilder()
+		.Read(BuiltInComponents->BlendChannelInput)
+		.Read(BuiltInComponents->IntegerResult)
+		.ReadOptional(BuiltInComponents->WeightAndEasingResult)
+		.FilterAll({ BuiltInComponents->Tags.AbsoluteBlend, GetBlenderTypeTag() })
+		.FilterNone({ BuiltInComponents->Tags.Ignored })
+		.Schedule_PerAllocation<FIntegerAccumulationTask>(&Linker->EntityManager, TaskScheduler, &AccumulationBuffers.Absolute);
+
+		TaskScheduler->AddPrerequisite(ResetWeightsTask, AbsoluteTask);
+		TaskScheduler->AddPrerequisite(AbsoluteTask, SyncTask);
+	}
+
+	if (AccumulationBuffers.Relative.Num() != 0)
+	{
+		FTaskID RelativeTask = FEntityTaskBuilder()
+		.Read(BuiltInComponents->BlendChannelInput)
+		.Read(BuiltInComponents->IntegerResult)
+		.ReadOptional(BuiltInComponents->WeightAndEasingResult)
+		.FilterAll({ BuiltInComponents->Tags.RelativeBlend, GetBlenderTypeTag() })
+		.FilterNone({ BuiltInComponents->Tags.Ignored })
+		.Schedule_PerAllocation<FIntegerAccumulationTask>(&Linker->EntityManager, TaskScheduler, &AccumulationBuffers.Relative);
+
+		TaskScheduler->AddPrerequisite(ResetWeightsTask, RelativeTask);
+		TaskScheduler->AddPrerequisite(RelativeTask, SyncTask);
+	}
+
+	if (AccumulationBuffers.Additive.Num() != 0)
+	{
+		FTaskID AdditiveTask = FEntityTaskBuilder()
+		.Read(BuiltInComponents->BlendChannelInput)
+		.Read(BuiltInComponents->IntegerResult)
+		.ReadOptional(BuiltInComponents->WeightAndEasingResult)
+		.FilterAll({ BuiltInComponents->Tags.AdditiveBlend, GetBlenderTypeTag() })
+		.FilterNone({ BuiltInComponents->Tags.Ignored })
+		.Schedule_PerAllocation<FIntegerAccumulationTask>(&Linker->EntityManager, TaskScheduler, &AccumulationBuffers.Additive);
+
+		TaskScheduler->AddPrerequisite(ResetWeightsTask, AdditiveTask);
+		TaskScheduler->AddPrerequisite(AdditiveTask, SyncTask);
+	}
+
+	if (AccumulationBuffers.AdditiveFromBase.Num() != 0)
+	{
+		FTaskID AdditiveFromBaseTask = FEntityTaskBuilder()
+		.Read(BuiltInComponents->BlendChannelInput)
+		.Read(BuiltInComponents->IntegerResult)
+		.Read(BuiltInComponents->BaseInteger)
+		.ReadOptional(BuiltInComponents->WeightAndEasingResult)
+		.FilterAll({ BuiltInComponents->Tags.AdditiveFromBaseBlend, GetBlenderTypeTag() })
+		.FilterNone({ BuiltInComponents->Tags.Ignored })
+		.Schedule_PerAllocation<FIntegerAdditiveFromBaseBlendTask>(&Linker->EntityManager, TaskScheduler, &AccumulationBuffers.AdditiveFromBase);
+
+		TaskScheduler->AddPrerequisite(ResetWeightsTask, AdditiveFromBaseTask);
+		TaskScheduler->AddPrerequisite(AdditiveFromBaseTask, SyncTask);
+	}
+
+	// Root task that performs the actual blends
+	FTaskID CombineTask = FEntityTaskBuilder()
+	.Read(BuiltInComponents->BlendChannelOutput)
+	.ReadOptional(TracksComponents->Integer.InitialValue)
+	.Write(BuiltInComponents->IntegerResult)
+	.FilterAll({ GetBlenderTypeTag() })
+	.Fork_PerAllocation<FIntegerCombineBlends>(&Linker->EntityManager, TaskScheduler, &AccumulationBuffers);
+
+	TaskScheduler->AddPrerequisite(SyncTask, CombineTask);
+}
+
 void UMovieScenePiecewiseIntegerBlenderSystem::OnRun(FSystemTaskPrerequisites& InPrerequisites, FSystemSubsequentTasks& Subsequents)
 {
 	using namespace UE::MovieScene;
@@ -261,14 +369,16 @@ void UMovieScenePiecewiseIntegerBlenderSystem::OnRun(FSystemTaskPrerequisites& I
 	{
 		ReinitializeAccumulationBuffers();
 	}
+	else
+	{
+		ZeroAccumulationBuffers();
+	}
 
 	if (AccumulationBuffers.IsEmpty())
 	{
 		return;
 	}
 
-	ZeroAccumulationBuffers();
-
 	const FBuiltInComponentTypes* BuiltInComponents = FBuiltInComponentTypes::Get();
 	const FMovieSceneTracksComponentTypes* TracksComponents = FMovieSceneTracksComponentTypes::Get();
 
@@ -381,6 +491,8 @@ void UMovieScenePiecewiseIntegerBlenderSystem::ReinitializeAccumulationBuffers()
 	{
 		AccumulationBuffers.AdditiveFromBase.SetNum(MaximumNumBlends);
 	}
+
+	ZeroAccumulationBuffers();
 }
 
 void UMovieScenePiecewiseIntegerBlenderSystem::ZeroAccumulationBuffers()
diff --git a/Engine/Source/Runtime/MovieSceneTracks/Private/Systems/MovieScenePropertySystem.cpp b/Engine/Source/Runtime/MovieSceneTracks/Private/Systems/MovieScenePropertySystem.cpp
--- a/Engine/Source/Runtime/MovieSceneTracks/Private/Systems/MovieScenePropertySystem.cpp
+++ b/Engine/Source/Runtime/MovieSceneTracks/Private/Systems/MovieScenePropertySystem.cpp
@@ -14,6 +14,7 @@ UMovieScenePropertySystem::UMovieScenePropertySystem(const FObjectInitializer& O
 {
 	using namespace UE::MovieScene;
 
+	Phase = ESystemPhase::Scheduling;
 	SystemCategories = EEntitySystemCategory::PropertySystems | FSystemInterrogator::GetExcludedFromInterrogationCategory();
 }
 
@@ -29,6 +30,26 @@ void UMovieScenePropertySystem::OnLink()
 	}
 }
 
+void UMovieScenePropertySystem::OnSchedulePersistentTasks(UE::MovieScene::IEntitySystemScheduler* TaskScheduler)
+{
+	using namespace UE::MovieScene;
+
+	// Never apply properties during evaluation. This code is necessary if derived types do support interrogation.
+	if (!InstantiatorSystem)
+	{
+		return;
+	}
+
+	FPropertyStats Stats = InstantiatorSystem->GetStatsForProperty(CompositePropertyID);
+	if (Stats.NumProperties > 0)
+	{
+		const FPropertyRegistry&   PropertyRegistry = FBuiltInComponentTypes::Get()->PropertyRegistry;
+		const FPropertyDefinition& Definition       = PropertyRegistry.GetDefinition(CompositePropertyID);
+
+		Definition.Handler->ScheduleSetterTasks(Definition, PropertyRegistry.GetComposites(Definition), Stats, TaskScheduler, Linker);
+	}
+}
+
 void UMovieScenePropertySystem::OnRun(FSystemTaskPrerequisites& InPrerequisites, FSystemSubsequentTasks& Subsequents)
 {
 	using namespace UE::MovieScene;
diff --git a/Engine/Source/Runtime/MovieSceneTracks/Private/Systems/MovieSceneQuaternionBlenderSystem.cpp b/Engine/Source/Runtime/MovieSceneTracks/Private/Systems/MovieSceneQuaternionBlenderSystem.cpp
--- a/Engine/Source/Runtime/MovieSceneTracks/Private/Systems/MovieSceneQuaternionBlenderSystem.cpp
+++ b/Engine/Source/Runtime/MovieSceneTracks/Private/Systems/MovieSceneQuaternionBlenderSystem.cpp
@@ -19,52 +19,51 @@
 namespace UE::MovieScene
 {
 
-struct FAccumulateAbsoluteWeights
+struct FResetBuffers
 {
 	FQuaternionBlenderAccumulationBuffers* Buffers;
 
-	FAccumulateAbsoluteWeights(FQuaternionBlenderAccumulationBuffers* InBuffers)
-		: Buffers(InBuffers)
-	{}
-
-	void PreTask()
+	void Run(FEntityAllocationWriteContext WriteContext) const
 	{
 		Buffers->AbsoluteWeights.Reset();
-	}
 
-	void ForEachEntity(FMovieSceneBlendChannelID BlendChannel, double WeightResult)
-	{
-		if (!Buffers->AbsoluteWeights.IsValidIndex(BlendChannel.ChannelID))
+		for (int32 Index = 0; Index < Buffers->Absolutes.Num(); ++Index)
 		{
-			Buffers->AbsoluteWeights.Insert(BlendChannel.ChannelID, WeightResult);
+			Buffers->Absolutes[Index] = FQuatTransform{
+				FVector::ZeroVector,
+				FQuat::Identity,
+				FVector::ZeroVector
+			};
 		}
-		else
+		for (int32 Index = 0; Index < Buffers->Additives.Num(); ++Index)
 		{
-			Buffers->AbsoluteWeights[BlendChannel.ChannelID] += WeightResult;
+			Buffers->Additives[Index] = FQuatTransform{ 
+				FVector::ZeroVector,
+				FQuat::Identity,
+				FVector::ZeroVector
+			};
 		}
 	}
 };
 
-struct FCollectInitialValues
+struct FAccumulateAbsoluteWeights
 {
 	FQuaternionBlenderAccumulationBuffers* Buffers;
 
-	FCollectInitialValues(FQuaternionBlenderAccumulationBuffers* InBuffers)
+	FAccumulateAbsoluteWeights(FQuaternionBlenderAccumulationBuffers* InBuffers)
 		: Buffers(InBuffers)
 	{}
 
-	void PreTask()
+	void ForEachEntity(FMovieSceneBlendChannelID BlendChannel, double WeightResult) const
 	{
-		FMemory::Memzero(Buffers->InitialValues.GetData(), sizeof(FQuatTransform) * Buffers->InitialValues.Num());
-	}
-
-	void ForEachEntity(FMovieSceneBlendChannelID BlendChannel, const FIntermediate3DTransform& InitialValue)
-	{
-		Buffers->InitialValues[BlendChannel.ChannelID] = FQuatTransform{
-			InitialValue.GetTranslation(),
-			InitialValue.GetRotation().Quaternion(),
-			InitialValue.GetScale()
-		};
+		if (!Buffers->AbsoluteWeights.IsValidIndex(BlendChannel.ChannelID))
+		{
+			Buffers->AbsoluteWeights.Insert(BlendChannel.ChannelID, WeightResult);
+		}
+		else
+		{
+			Buffers->AbsoluteWeights[BlendChannel.ChannelID] += WeightResult;
+		}
 	}
 };
 
@@ -78,25 +77,13 @@ struct FAbsoluteAccumulationTask
 		: Buffers(InBuffers)
 	{}
 
-	void PreTask()
-	{
-		for (int32 Index = 0; Index < Buffers->Absolutes.Num(); ++Index)
-		{
-			Buffers->Absolutes[Index] = FQuatTransform{ 
-				FVector::ZeroVector,
-				FQuat::Identity,
-				FVector::ZeroVector
-			};
-		}
-	}
-
 	void ForEachEntity(
 		FMovieSceneBlendChannelID BlendChannel,
 		double LocationX, double LocationY, double LocationZ,
 		double RotationRoll, double RotationPitch, double RotationYaw,
 		double ScaleX, double ScaleY, double ScaleZ,
 		const double* WeightResult
-		)
+		) const
 	{
 		double Weight = WeightResult ? *WeightResult : 1.0;
 		if (Weight == 0.0)
@@ -130,25 +117,13 @@ struct FAdditiveAccumulationTask
 		: Buffers(InBuffers)
 	{}
 
-	void PreTask()
-	{
-		for (int32 Index = 0; Index < Buffers->Additives.Num(); ++Index)
-		{
-			Buffers->Additives[Index] = FQuatTransform{ 
-				FVector::ZeroVector,
-				FQuat::Identity,
-				FVector::ZeroVector
-			};
-		}
-	}
-
 	void ForEachEntity(
 		FMovieSceneBlendChannelID BlendChannel,
 		double LocationX, double LocationY, double LocationZ,
 		double RotationRoll, double RotationPitch, double RotationYaw,
 		double ScaleX, double ScaleY, double ScaleZ,
 		const double* WeightResult
-		)
+		) const
 	{
 		FQuatTransform& Transform = Buffers->Additives[BlendChannel.ChannelID];
 
@@ -174,7 +149,7 @@ struct FCombineTask
 		double& RotationRoll, double& RotationPitch, double& RotationYaw,
 		double& ScaleX, double& ScaleY, double& ScaleZ,
 		const FIntermediate3DTransform* InitialValue
-		)
+		) const
 	{
 		const uint16 ChannelID = BlendChannel.ChannelID;
 		const double TotalWeight = Buffers->AbsoluteWeights.IsValidIndex(ChannelID) ? Buffers->AbsoluteWeights[ChannelID] : 1.0;
@@ -217,6 +192,11 @@ struct FCombineTask
 UMovieSceneQuaternionBlenderSystem::UMovieSceneQuaternionBlenderSystem(const FObjectInitializer& ObjInit)
 	: Super(ObjInit)
 {
+	using namespace UE::MovieScene;
+
+	SelectionPriority = 1000;
+	Phase = ESystemPhase::Scheduling;
+
 	if (HasAnyFlags(RF_ClassDefaultObject))
 	{
 		DefineImplicitPrerequisite(UFloatChannelEvaluatorSystem::StaticClass(), GetClass());
@@ -227,6 +207,72 @@ UMovieSceneQuaternionBlenderSystem::UMovieSceneQuaternionBlenderSystem(const FOb
 	}
 }
 
+void UMovieSceneQuaternionBlenderSystem::OnSchedulePersistentTasks(UE::MovieScene::IEntitySystemScheduler* TaskScheduler)
+{
+	using namespace UE::MovieScene;
+
+	FBuiltInComponentTypes* BuiltInComponents = FBuiltInComponentTypes::Get();
+	FMovieSceneTracksComponentTypes* TracksComponents = FMovieSceneTracksComponentTypes::Get();
+
+	Buffers.Absolutes.SetNumZeroed(AllocatedBlendChannels.Num());
+	Buffers.Additives.SetNumZeroed(AllocatedBlendChannels.Num());
+
+	FTaskID ResetBufferTask = TaskScheduler->AddTask<FResetBuffers>(FTaskParams(TEXT("Reset Quaternion Blender Buffers")), &Buffers);
+
+	// Kick off additives first because they don't need normalized weighting
+	// Not handling additive from base yet because ideally we could do those in the AdditivesTask without caring about the base
+	FTaskID AdditivesTask = FEntityTaskBuilder()
+	.Read(BuiltInComponents->BlendChannelInput)
+	.ReadAllOf(BuiltInComponents->DoubleResult[0], BuiltInComponents->DoubleResult[1], BuiltInComponents->DoubleResult[2]) // Translation components
+	.ReadAllOf(BuiltInComponents->DoubleResult[3], BuiltInComponents->DoubleResult[4], BuiltInComponents->DoubleResult[5]) // Rotation components
+	.ReadAllOf(BuiltInComponents->DoubleResult[6], BuiltInComponents->DoubleResult[7], BuiltInComponents->DoubleResult[8]) // Scale components
+	.ReadOptional(BuiltInComponents->WeightAndEasingResult)
+	.FilterAny({ BuiltInComponents->Tags.AdditiveBlend, BuiltInComponents->Tags.AdditiveFromBaseBlend })
+	.FilterAll({ GetBlenderTypeTag() })
+	.FilterNone({ BuiltInComponents->Tags.Ignored })
+	.Schedule_PerEntity<FAdditiveAccumulationTask>(&Linker->EntityManager, TaskScheduler, &Buffers);
+
+	TaskScheduler->AddPrerequisite(ResetBufferTask, AdditivesTask);
+
+	// Accumulate all weights so we can do a normalized lerp.
+	// Everything downstream depends on this task
+	FTaskID CollectWeightsTask = FEntityTaskBuilder()
+	.Read(BuiltInComponents->BlendChannelInput)
+	.Read(BuiltInComponents->WeightAndEasingResult)
+	.FilterAll({ BuiltInComponents->Tags.AbsoluteBlend, GetBlenderTypeTag() })
+	.FilterNone({ BuiltInComponents->Tags.Ignored })
+	.Schedule_PerEntity<FAccumulateAbsoluteWeights>(&Linker->EntityManager, TaskScheduler, &Buffers);
+
+	TaskScheduler->AddPrerequisite(ResetBufferTask, AdditivesTask);
+
+	// Next we accumulate normalized transforms
+	FTaskID AbsolutesTask = FEntityTaskBuilder()
+	.Read(BuiltInComponents->BlendChannelInput)
+	.ReadAllOf(BuiltInComponents->DoubleResult[0], BuiltInComponents->DoubleResult[1], BuiltInComponents->DoubleResult[2]) // Translation components
+	.ReadAllOf(BuiltInComponents->DoubleResult[3], BuiltInComponents->DoubleResult[4], BuiltInComponents->DoubleResult[5]) // Rotation components
+	.ReadAllOf(BuiltInComponents->DoubleResult[6], BuiltInComponents->DoubleResult[7], BuiltInComponents->DoubleResult[8]) // Scale components
+	.ReadOptional(BuiltInComponents->WeightAndEasingResult)
+	.FilterAll({ BuiltInComponents->Tags.AbsoluteBlend, GetBlenderTypeTag() })
+	.FilterNone({ BuiltInComponents->Tags.Ignored })
+	.Schedule_PerEntity<FAbsoluteAccumulationTask>(&Linker->EntityManager, TaskScheduler, &Buffers);
+
+	TaskScheduler->AddPrerequisite(CollectWeightsTask, AbsolutesTask);
+
+	// Now blend the results
+	FTaskID BlendTask = FEntityTaskBuilder()
+	.Read(BuiltInComponents->BlendChannelOutput)
+	.WriteAllOf(BuiltInComponents->DoubleResult[0], BuiltInComponents->DoubleResult[1], BuiltInComponents->DoubleResult[2]) // Translation components
+	.WriteAllOf(BuiltInComponents->DoubleResult[3], BuiltInComponents->DoubleResult[4], BuiltInComponents->DoubleResult[5]) // Rotation components
+	.WriteAllOf(BuiltInComponents->DoubleResult[6], BuiltInComponents->DoubleResult[7], BuiltInComponents->DoubleResult[8]) // Scale components
+	.ReadOptional(TracksComponents->ComponentTransform.InitialValue)
+	.FilterAll({ GetBlenderTypeTag() })
+	.Fork_PerEntity<FCombineTask>(&Linker->EntityManager, TaskScheduler, &Buffers);
+
+	TaskScheduler->AddPrerequisite(AdditivesTask, BlendTask);
+	TaskScheduler->AddPrerequisite(CollectWeightsTask, BlendTask);
+	TaskScheduler->AddPrerequisite(AbsolutesTask, BlendTask);
+}
+
 void UMovieSceneQuaternionBlenderSystem::OnRun(FSystemTaskPrerequisites& InPrerequisites, FSystemSubsequentTasks& Subsequents)
 {
 	using namespace UE::MovieScene;
@@ -261,21 +307,10 @@ void UMovieSceneQuaternionBlenderSystem::OnRun(FSystemTaskPrerequisites& InPrere
 	.FilterNone({ BuiltInComponents->Tags.Ignored })
 	.template Dispatch_PerEntity<FAccumulateAbsoluteWeights>(&Linker->EntityManager, InPrerequisites, nullptr, &Buffers);
 
-	// Accumulate initial values
-	//FGraphEventRef CollectInitialValuesTask = FEntityTaskBuilder()
-	//.Read(BuiltInComponents->BlendChannelOutput)
-	//.Read(TracksComponents->ComponentTransform.InitialValue)
-	//.FilterAll({ GetBlenderTypeTag() })
-	//.template Dispatch_PerEntity<FCollectInitialValues>(&Linker->EntityManager, InPrerequisites, &Subsequents, &Buffers);
-
 	if (CollectWeightsTask)
 	{
 		InPrerequisites.AddRootTask(CollectWeightsTask);
 	}
-	//if (CollectInitialValuesTask)
-	//{
-	//	InPrerequisites.AddRootTask(CollectInitialValuesTask);
-	//}
 
 	// Next we accumulate normalized transforms
 	FGraphEventRef AbsolutesTask = FEntityTaskBuilder()
diff --git a/Engine/Source/Runtime/MovieSceneTracks/Private/Systems/MovieSceneQuaternionInterpolationRotationSystem.cpp b/Engine/Source/Runtime/MovieSceneTracks/Private/Systems/MovieSceneQuaternionInterpolationRotationSystem.cpp
--- a/Engine/Source/Runtime/MovieSceneTracks/Private/Systems/MovieSceneQuaternionInterpolationRotationSystem.cpp
+++ b/Engine/Source/Runtime/MovieSceneTracks/Private/Systems/MovieSceneQuaternionInterpolationRotationSystem.cpp
@@ -10,7 +10,6 @@
 
 #include UE_INLINE_GENERATED_CPP_BY_NAME(MovieSceneQuaternionInterpolationRotationSystem)
 
-DECLARE_CYCLE_STAT(TEXT("MovieScene: Quat-interp-rot channel system"), MovieSceneEval_QuatInterpRotChannelSystem, STATGROUP_MovieSceneECS);
 DECLARE_CYCLE_STAT(TEXT("MovieScene: Evaluate quat-interp-rot channels"), MovieSceneEval_EvaluateQuatInterpRotChannelTask, STATGROUP_MovieSceneECS);
 
 namespace UE
@@ -46,7 +45,7 @@ struct FEvaluateQuaternionInterpolationRotationChannels
 			const FEntityAllocation* Allocation,
 			TRead<FFrameTime> FrameTimes, 
 			TMultiReadOptional<FSourceDoubleChannel, FSourceDoubleChannel, FSourceDoubleChannel> RotChannelAccessors,
-			double* OutResultXs, double* OutResultYs, double* OutResultZs)
+			double* OutResultXs, double* OutResultYs, double* OutResultZs) const
 	{
 		const FSourceDoubleChannel* RotationXs = RotChannelAccessors.Get<0>();
 		const FSourceDoubleChannel* RotationYs = RotChannelAccessors.Get<1>();
@@ -188,6 +187,7 @@ UMovieSceneQuaternionInterpolationRotationSystem::UMovieSceneQuaternionInterpola
 	using namespace UE::MovieScene;
 
 	SystemCategories = EEntitySystemCategory::ChannelEvaluators;
+	Phase = ESystemPhase::Scheduling;
 
 	if (HasAnyFlags(RF_ClassDefaultObject))
 	{
@@ -213,11 +213,29 @@ bool UMovieSceneQuaternionInterpolationRotationSystem::IsRelevantImpl(UMovieScen
 		});
 }
 
-void UMovieSceneQuaternionInterpolationRotationSystem::OnRun(FSystemTaskPrerequisites& InPrerequisites, FSystemSubsequentTasks& Subsequents)
+void UMovieSceneQuaternionInterpolationRotationSystem::OnSchedulePersistentTasks(UE::MovieScene::IEntitySystemScheduler* TaskScheduler)
 {
 	using namespace UE::MovieScene;
 
-	SCOPE_CYCLE_COUNTER(MovieSceneEval_QuatInterpRotChannelSystem);
+	FBuiltInComponentTypes* BuiltInComponents = FBuiltInComponentTypes::Get();
+	FMovieSceneTracksComponentTypes* TrackComponents = FMovieSceneTracksComponentTypes::Get();
+
+	FEntityTaskBuilder()
+	.Read(BuiltInComponents->EvalTime)
+	.ReadOneOrMoreOf(
+			TrackComponents->QuaternionRotationChannel[0],
+			TrackComponents->QuaternionRotationChannel[1],
+			TrackComponents->QuaternionRotationChannel[2])
+	.WriteOptional(BuiltInComponents->DoubleResult[3])
+	.WriteOptional(BuiltInComponents->DoubleResult[4])
+	.WriteOptional(BuiltInComponents->DoubleResult[5])
+	.SetStat(GET_STATID(MovieSceneEval_EvaluateQuatInterpRotChannelTask))
+	.Fork_PerAllocation<FEvaluateQuaternionInterpolationRotationChannels>(&Linker->EntityManager, TaskScheduler);
+}
+
+void UMovieSceneQuaternionInterpolationRotationSystem::OnRun(FSystemTaskPrerequisites& InPrerequisites, FSystemSubsequentTasks& Subsequents)
+{
+	using namespace UE::MovieScene;
 
 	FBuiltInComponentTypes* BuiltInComponents = FBuiltInComponentTypes::Get();
 	FMovieSceneTracksComponentTypes* TrackComponents = FMovieSceneTracksComponentTypes::Get();
diff --git a/Engine/Source/Runtime/MovieSceneTracks/Private/Systems/MovieSceneSkeletalAnimationSystem.cpp b/Engine/Source/Runtime/MovieSceneTracks/Private/Systems/MovieSceneSkeletalAnimationSystem.cpp
--- a/Engine/Source/Runtime/MovieSceneTracks/Private/Systems/MovieSceneSkeletalAnimationSystem.cpp
+++ b/Engine/Source/Runtime/MovieSceneTracks/Private/Systems/MovieSceneSkeletalAnimationSystem.cpp
@@ -312,14 +312,19 @@ void FSkeletalAnimationSystemData::ResetSkeletalAnimations()
 struct FGatherSkeletalAnimations
 {
 	const FInstanceRegistry* InstanceRegistry;
-	FSkeletalAnimationSystemData& SystemData;
+	FSkeletalAnimationSystemData* SystemData;
 
-	FGatherSkeletalAnimations(const FInstanceRegistry* InInstanceRegistry, FSkeletalAnimationSystemData& InSystemData)
+	FGatherSkeletalAnimations(const FInstanceRegistry* InInstanceRegistry, FSkeletalAnimationSystemData* InSystemData)
 		: InstanceRegistry(InInstanceRegistry)
 		, SystemData(InSystemData)
 	{
 	}
 
+	void PreTask() const
+	{
+		SystemData->SkeletalAnimations.Reset();
+	}
+
 	void ForEachAllocation(
 			const FEntityAllocationProxy AllocationProxy, 
 			TRead<FMovieSceneEntityID> EntityIDs,
@@ -327,7 +332,7 @@ struct FGatherSkeletalAnimations
 			TRead<FInstanceHandle> InstanceHandles,
 			TRead<UObject*> BoundObjects,
 			TRead<FMovieSceneSkeletalAnimationComponentData> SkeletalAnimations,
-			TReadOptional<double> WeightAndEasings)
+			TReadOptional<double> WeightAndEasings) const
 	{
 		// Gather all the skeletal animations currently active in all sequences.
 		// We map these animations to their bound object, which means we might blend animations from different sequences
@@ -361,7 +366,7 @@ struct FGatherSkeletalAnimations
 			const float EvalTime = AnimParams.MapTimeToAnimation(AnimSection, Context.GetTime(), Context.GetFrameRate());
 			const float PreviousEvalTime = AnimParams.MapTimeToAnimation(AnimSection, Context.GetPreviousTime(), Context.GetFrameRate());
 
-			FBoundObjectActiveSkeletalAnimations& BoundObjectAnimations = SystemData.SkeletalAnimations.FindOrAdd(BoundObject);
+			FBoundObjectActiveSkeletalAnimations& BoundObjectAnimations = SystemData->SkeletalAnimations.FindOrAdd(BoundObject);
 			BoundObjectAnimations.Animations.Add(FActiveSkeletalAnimation{ Player, AnimSection, Context, EntityID, RootInstanceHandle, PreviousEvalTime, EvalTime, Weight, bWantsRestoreState });
 
 			if (FMotionVectorSimulation::IsEnabled())
@@ -390,7 +395,7 @@ struct FEvaluateSkeletalAnimations
 private:
 
 	UMovieSceneEntitySystemLinker* Linker;
-	FSkeletalAnimationSystemData& SystemData;
+	FSkeletalAnimationSystemData* SystemData;
 
 	TSharedPtr<FPreAnimatedSkeletalAnimationStorage> PreAnimatedStorage;
 	TSharedPtr<FPreAnimatedSkeletalAnimationMontageStorage> PreAnimatedMontageStorage;
@@ -398,7 +403,7 @@ private:
 
 public:
 
-	FEvaluateSkeletalAnimations(UMovieSceneEntitySystemLinker* InLinker, FSkeletalAnimationSystemData& InSystemData)
+	FEvaluateSkeletalAnimations(UMovieSceneEntitySystemLinker* InLinker, FSkeletalAnimationSystemData* InSystemData)
 		: Linker(InLinker)
 		, SystemData(InSystemData)
 	{
@@ -407,9 +412,13 @@ public:
 		PreAnimatedAnimInstanceStorage = InLinker->PreAnimatedState.GetOrCreateStorage<FPreAnimatedSkeletalAnimationAnimInstanceStorage>();
 	}
 
-	void Run()
+	void Run(FEntityAllocationWriteContext WriteContext) const
 	{
-		for (const TTuple<UObject*, FBoundObjectActiveSkeletalAnimations>& Pair : SystemData.SkeletalAnimations)
+		Run();
+	}
+	void Run() const
+	{
+		for (const TTuple<UObject*, FBoundObjectActiveSkeletalAnimations>& Pair : SystemData->SkeletalAnimations)
 		{
 			EvaluateSkeletalAnimations(Pair.Key, Pair.Value);
 		}
@@ -417,7 +426,7 @@ public:
 
 private:
 
-	void EvaluateSkeletalAnimations(UObject* InObject, const FBoundObjectActiveSkeletalAnimations& InSkeletalAnimations)
+	void EvaluateSkeletalAnimations(UObject* InObject, const FBoundObjectActiveSkeletalAnimations& InSkeletalAnimations) const
 	{
 		ensureMsgf(InObject, TEXT("Attempting to evaluate an Animation track with a null object."));
 
@@ -481,7 +490,7 @@ private:
 			USkeletalMeshComponent* InSkeletalMeshComponent, 
 			UAnimInstance* InExistingAnimInstance, 
 			ISequencerAnimationSupport* InSequencerInstance, 
-			const FBoundObjectActiveSkeletalAnimations& InSkeletalAnimations)
+			const FBoundObjectActiveSkeletalAnimations& InSkeletalAnimations) const
 	{
 		const bool bPreviewPlayback = ShouldUsePreviewPlayback(InPlayer, *InSkeletalMeshComponent);
 
@@ -511,7 +520,7 @@ private:
 		}
 		else if (InExistingAnimInstance)
 		{
-			for (const TPair<FObjectKey, FMontagePlayerPerSectionData >& Pair : SystemData.MontageData)
+			for (const TPair<FObjectKey, FMontagePlayerPerSectionData >& Pair : SystemData->MontageData)
 			{
 				int32 InstanceId = Pair.Value.MontageInstanceId;
 				FAnimMontageInstance* MontageInstanceToUpdate = InExistingAnimInstance->GetMontageInstanceForID(InstanceId);
@@ -630,7 +639,7 @@ private:
 		bool bResetDynamics;
 	};
 
-	void SimulateMotionVectors(USkeletalMeshComponent* SkeletalMeshComponent, UMovieSceneMotionVectorSimulationSystem* MotionVectorSim)
+	void SimulateMotionVectors(USkeletalMeshComponent* SkeletalMeshComponent, UMovieSceneMotionVectorSimulationSystem* MotionVectorSim) const
 	{
 		for (USceneComponent* Child : SkeletalMeshComponent->GetAttachChildren())
 		{
@@ -655,7 +664,7 @@ private:
 		TArrayView<const FActiveSkeletalAnimation> SkeletalAnimations,
 		bool bPreviewPlayback,
 		bool bFireNotifies,
-		bool bResetDynamics)
+		bool bResetDynamics) const
 	{
 		const EMovieScenePlayerStatus::Type PlayerStatus = Player.GetPlaybackStatus();
 
@@ -670,7 +679,7 @@ private:
 			{
 				FTransform Transform = RootMotionParams.Transform.IsSet() ? RootMotionParams.Transform.GetValue() : FTransform::Identity;
 				TOptional<FTransform> InitialTransform = GetCurrentTransform(AnimSection->Params.SwapRootBone, SkeletalMeshComponent);
-				if (FBoundObjectActiveSkeletalAnimations* BoundObjectAnimations = SystemData.SkeletalAnimations.Find(InObject)) //InObject may be an AActor
+				if (FBoundObjectActiveSkeletalAnimations* BoundObjectAnimations = SystemData->SkeletalAnimations.Find(InObject)) //InObject may be an AActor
 				{
 					BoundObjectAnimations->BoneTransformFinalizeData.Register(SkeletalMeshComponent, AnimSection->Params.SwapRootBone,Transform, InitialTransform);
 				}
@@ -772,7 +781,7 @@ private:
 		return CurrentTransform;
 	}
 
-	void SetAnimPosition(const FSetAnimPositionParams& Params)
+	void SetAnimPosition(const FSetAnimPositionParams& Params) const
 	{
 		static const bool bLooping = false;
 
@@ -826,7 +835,7 @@ private:
 		}
 		else if (UAnimInstance* AnimInst = GetSourceAnimInstance(Params.SkeletalMeshComponent))
 		{
-			FMontagePlayerPerSectionData* SectionData = SystemData.MontageData.Find(Params.Section);
+			FMontagePlayerPerSectionData* SectionData = SystemData->MontageData.Find(Params.Section);
 
 			int32 InstanceId = (SectionData) ? SectionData->MontageInstanceId : INDEX_NONE;
 
@@ -846,8 +855,7 @@ private:
 			UAnimMontage* Montage = WeakMontage.Get();
 			if (Montage)
 			{
-
-				FMontagePlayerPerSectionData& DataContainer = SystemData.MontageData.FindOrAdd(Params.Section);
+				FMontagePlayerPerSectionData& DataContainer = SystemData->MontageData.FindOrAdd(Params.Section);
 				DataContainer.Montage = WeakMontage;
 				DataContainer.MontageInstanceId = InstanceId;
 
@@ -866,7 +874,7 @@ private:
 		}
 	}
 
-	void PreviewSetAnimPosition(const FSetAnimPositionParams& Params)
+	void PreviewSetAnimPosition(const FSetAnimPositionParams& Params) const
 	{
 		using namespace UE::MovieScene;
 
@@ -920,7 +928,7 @@ private:
 		}
 		else if (UAnimInstance* AnimInst = GetSourceAnimInstance(Params.SkeletalMeshComponent))
 		{
-			FMontagePlayerPerSectionData* SectionData = SystemData.MontageData.Find(Params.Section);
+			FMontagePlayerPerSectionData* SectionData = SystemData->MontageData.Find(Params.Section);
 
 			int32 InstanceId = SectionData ? SectionData->MontageInstanceId : INDEX_NONE;
 		
@@ -941,8 +949,7 @@ private:
 			UAnimMontage* Montage = WeakMontage.Get();
 			if (Montage)
 			{
-
-				FMontagePlayerPerSectionData& DataContainer = SystemData.MontageData.FindOrAdd(Params.Section);
+				FMontagePlayerPerSectionData& DataContainer = SystemData->MontageData.FindOrAdd(Params.Section);
 				DataContainer.Montage = WeakMontage;
 				DataContainer.MontageInstanceId = InstanceId;
 
@@ -976,7 +983,7 @@ UMovieSceneSkeletalAnimationSystem::UMovieSceneSkeletalAnimationSystem(const FOb
 
 	FMovieSceneTracksComponentTypes* TrackComponents = FMovieSceneTracksComponentTypes::Get();
 	RelevantComponent = TrackComponents->SkeletalAnimation;
-	Phase = ESystemPhase::Instantiation | ESystemPhase::Evaluation;
+	Phase = ESystemPhase::Instantiation | ESystemPhase::Scheduling;
 
 	SystemCategories |= FSystemInterrogator::GetExcludedFromInterrogationCategory();
 
@@ -988,6 +995,34 @@ UMovieSceneSkeletalAnimationSystem::UMovieSceneSkeletalAnimationSystem(const FOb
 	}
 }
 
+void UMovieSceneSkeletalAnimationSystem::OnSchedulePersistentTasks(UE::MovieScene::IEntitySystemScheduler* TaskScheduler)
+{
+	using namespace UE::MovieScene;
+
+	FBuiltInComponentTypes* BuiltInComponents = FBuiltInComponentTypes::Get();
+	FMovieSceneTracksComponentTypes* TrackComponents = FMovieSceneTracksComponentTypes::Get();
+
+	FTaskID GatherTask = FEntityTaskBuilder()
+	.ReadEntityIDs()
+	.Read(BuiltInComponents->RootInstanceHandle)
+	.Read(BuiltInComponents->InstanceHandle)
+	.Read(BuiltInComponents->BoundObject)
+	.Read(TrackComponents->SkeletalAnimation)
+	.ReadOptional(BuiltInComponents->WeightAndEasingResult)
+	.FilterNone({ BuiltInComponents->Tags.Ignored })
+	.SetStat(GET_STATID(MovieSceneEval_GatherSkeletalAnimations))
+	.Schedule_PerAllocation<FGatherSkeletalAnimations>(&Linker->EntityManager, TaskScheduler, 
+			Linker->GetInstanceRegistry(), &SystemData);
+
+	// Now evaluate gathered animations. We need to do this on the game thread (when in multi-threaded mode)
+	// because this task will call into a lot of animation system code that needs to be called there.
+	FTaskParams Params(GET_STATID(MovieSceneEval_EvaluateSkeletalAnimations));
+	Params.ForceGameThread();
+	FTaskID EvaluateTask = TaskScheduler->AddTask<FEvaluateSkeletalAnimations>(Params, Linker, &SystemData);
+
+	TaskScheduler->AddPrerequisite(GatherTask, EvaluateTask);
+}
+
 void UMovieSceneSkeletalAnimationSystem::OnRun(FSystemTaskPrerequisites& InPrerequisites, FSystemSubsequentTasks& Subsequents)
 {
 	using namespace UE::MovieScene;
@@ -1017,7 +1052,7 @@ void UMovieSceneSkeletalAnimationSystem::OnRun(FSystemTaskPrerequisites& InPrere
 	.FilterNone({ BuiltInComponents->Tags.Ignored })
 	.SetStat(GatherStatId)
 	.Dispatch_PerAllocation<FGatherSkeletalAnimations>(&Linker->EntityManager, InPrerequisites, nullptr, 
-			Linker->GetInstanceRegistry(), SystemData);
+			Linker->GetInstanceRegistry(), &SystemData);
 
 	FSystemTaskPrerequisites EvalPrereqs;
 	if (GatherTask)
@@ -1030,7 +1065,7 @@ void UMovieSceneSkeletalAnimationSystem::OnRun(FSystemTaskPrerequisites& InPrere
 	FEntityTaskBuilder()
 	.SetStat(GET_STATID(MovieSceneEval_EvaluateSkeletalAnimations))
 	.SetDesiredThread(Linker->EntityManager.GetGatherThread())
-	.Dispatch<FEvaluateSkeletalAnimations>(&Linker->EntityManager, EvalPrereqs, &Subsequents, Linker, SystemData);
+	.Dispatch<FEvaluateSkeletalAnimations>(&Linker->EntityManager, EvalPrereqs, &Subsequents, Linker, &SystemData);
 }
 
 void UMovieSceneSkeletalAnimationSystem::CleanSystemData()
diff --git a/Engine/Source/Runtime/MovieSceneTracks/Private/Systems/MovieSceneTransformOriginSystem.cpp b/Engine/Source/Runtime/MovieSceneTracks/Private/Systems/MovieSceneTransformOriginSystem.cpp
--- a/Engine/Source/Runtime/MovieSceneTracks/Private/Systems/MovieSceneTransformOriginSystem.cpp
+++ b/Engine/Source/Runtime/MovieSceneTracks/Private/Systems/MovieSceneTransformOriginSystem.cpp
@@ -23,13 +23,48 @@ namespace UE
 namespace MovieScene
 {
 
+struct FGatherTransformOrigin
+{
+	TSparseArray<FTransform>* TransformOriginsByInstanceID;
+	const FInstanceRegistry* InstanceRegistry;
+
+	void Run(FEntityAllocationWriteContext WriteContext) const
+	{
+		TransformOriginsByInstanceID->Empty(InstanceRegistry->GetSparseInstances().Num());
+
+		const TSparseArray<FSequenceInstance>& SparseInstances = InstanceRegistry->GetSparseInstances();
+		for (int32 Index = 0; Index < SparseInstances.GetMaxIndex(); ++Index)
+		{
+			if (!SparseInstances.IsValidIndex(Index))
+			{
+				continue;
+			}
+
+			const FSequenceInstance& Instance = SparseInstances[Index];
+
+			const IMovieScenePlaybackClient*  Client       = Instance.GetPlayer()->GetPlaybackClient();
+			const UObject*                    InstanceData = Client ? Client->GetInstanceData() : nullptr;
+			const IMovieSceneTransformOrigin* RawInterface = Cast<const IMovieSceneTransformOrigin>(InstanceData);
+
+			const bool bHasInterface = RawInterface || (InstanceData && InstanceData->GetClass()->ImplementsInterface(UMovieSceneTransformOrigin::StaticClass()));
+			if (bHasInterface)
+			{
+				// Retrieve the current origin
+				FTransform TransformOrigin = RawInterface ? RawInterface->GetTransformOrigin() : IMovieSceneTransformOrigin::Execute_BP_GetTransformOrigin(InstanceData);
+
+				TransformOriginsByInstanceID->Insert(Index, TransformOrigin);
+			}
+		}
+	}
+};
+
 struct FAssignTransformOrigin
 {
 	const TSparseArray<FTransform>* TransformOriginsByInstanceID;
 
 	void ForEachAllocation(const FEntityAllocation* Allocation, TRead<FInstanceHandle> Instances, TRead<UObject*> BoundObjects,
 		TWriteOptional<double> LocationX, TWriteOptional<double> LocationY, TWriteOptional<double> LocationZ,
-		TWriteOptional<double> RotationX, TWriteOptional<double> RotationY, TWriteOptional<double> RotationZ)
+		TWriteOptional<double> RotationX, TWriteOptional<double> RotationY, TWriteOptional<double> RotationZ) const
 	{
 		TransformLocation(Instances, BoundObjects, LocationX, LocationY, LocationZ, RotationX, RotationY, RotationZ, Allocation->Num());
 	}
@@ -37,7 +72,7 @@ struct FAssignTransformOrigin
 	void TransformLocation(const FInstanceHandle* Instances, const UObject* const * BoundObjects,
 		double* OutLocationX, double* OutLocationY, double* OutLocationZ,
 		double* OutRotationX, double* OutRotationY, double* OutRotationZ,
-		int32 Num)
+		int32 Num) const
 	{
 		for (int32 Index = 0; Index < Num; ++Index)
 		{
@@ -125,6 +160,8 @@ UMovieSceneTransformOriginSystem::UMovieSceneTransformOriginSystem(const FObject
 {
 	using namespace UE::MovieScene;
 
+	Phase = ESystemPhase::Scheduling;
+
 	if (HasAnyFlags(RF_ClassDefaultObject))
 	{
 		// This system relies upon anything that creates entities
@@ -166,6 +203,43 @@ void UMovieSceneTransformOriginSystem::OnLink()
 	Linker->SystemGraph.AddReference(this, Instantiator);
 }
 
+void UMovieSceneTransformOriginSystem::OnSchedulePersistentTasks(UE::MovieScene::IEntitySystemScheduler* TaskScheduler)
+{
+	using namespace UE::MovieScene;
+
+	FInstanceRegistry* InstanceRegistry = Linker->GetInstanceRegistry();
+
+	FBuiltInComponentTypes* BuiltInComponents = FBuiltInComponentTypes::Get();
+	FMovieSceneTracksComponentTypes* TracksComponents = FMovieSceneTracksComponentTypes::Get();
+
+	FEntityComponentFilter Filter;
+	Filter.All({ TracksComponents->ComponentTransform.PropertyTag, BuiltInComponents->Tags.AbsoluteBlend });
+	Filter.None({ BuiltInComponents->BlendChannelOutput });
+
+	FTaskID GatherTask = TaskScheduler->AddTask<FGatherTransformOrigin>(
+		FTaskParams(TEXT("Gather Transform Origins")).ForceGameThread(),
+		&TransformOriginsByInstanceID,
+		InstanceRegistry
+	);
+
+	FTaskID AssignTask = FEntityTaskBuilder()
+	.Read(BuiltInComponents->InstanceHandle)
+	.Read(BuiltInComponents->BoundObject)
+	.WriteOptional(BuiltInComponents->DoubleResult[0])
+	.WriteOptional(BuiltInComponents->DoubleResult[1])
+	.WriteOptional(BuiltInComponents->DoubleResult[2])
+	.WriteOptional(BuiltInComponents->DoubleResult[3])
+	.WriteOptional(BuiltInComponents->DoubleResult[4])
+	.WriteOptional(BuiltInComponents->DoubleResult[5])
+	.CombineFilter(Filter)
+	// Must contain at least one double result
+	.FilterAny({ BuiltInComponents->DoubleResult[0], BuiltInComponents->DoubleResult[1], BuiltInComponents->DoubleResult[2],
+		BuiltInComponents->DoubleResult[3], BuiltInComponents->DoubleResult[4], BuiltInComponents->DoubleResult[5] })
+	.Fork_PerAllocation<FAssignTransformOrigin>(&Linker->EntityManager, TaskScheduler, &TransformOriginsByInstanceID);
+
+	TaskScheduler->AddPrerequisite(GatherTask, AssignTask);
+}
+
 void UMovieSceneTransformOriginSystem::OnRun(FSystemTaskPrerequisites& InPrerequisites, FSystemSubsequentTasks& Subsequents)
 {
 	using namespace UE::MovieScene;
diff --git a/Engine/Source/Runtime/MovieSceneTracks/Private/Systems/ObjectPathChannelEvaluatorSystem.cpp b/Engine/Source/Runtime/MovieSceneTracks/Private/Systems/ObjectPathChannelEvaluatorSystem.cpp
--- a/Engine/Source/Runtime/MovieSceneTracks/Private/Systems/ObjectPathChannelEvaluatorSystem.cpp
+++ b/Engine/Source/Runtime/MovieSceneTracks/Private/Systems/ObjectPathChannelEvaluatorSystem.cpp
@@ -9,7 +9,22 @@
 
 #include UE_INLINE_GENERATED_CPP_BY_NAME(ObjectPathChannelEvaluatorSystem)
 
-DECLARE_CYCLE_STAT(TEXT("MovieScene: Evaluate object path channels"), MovieSceneEval_EvaluateObjectPathChannelTask, STATGROUP_MovieSceneECS);
+namespace UE::MovieScene
+{
+	DECLARE_CYCLE_STAT(TEXT("MovieScene: Evaluate object path channels"), MovieSceneEval_EvaluateObjectPathChannelTask, STATGROUP_MovieSceneECS);
+
+	struct FEvaluateObjectPathChannels
+	{
+		static void ForEachEntity(FSourceObjectPathChannel ObjectPathChannel, FFrameTime FrameTime, UObject*& OutResult)
+		{
+			if (!ObjectPathChannel.Source->Evaluate(FrameTime, OutResult))
+			{
+				OutResult = nullptr;
+			}
+		}
+	};
+}
+
 
 UObjectPathChannelEvaluatorSystem::UObjectPathChannelEvaluatorSystem(const FObjectInitializer& ObjInit)
 	: Super(ObjInit)
@@ -26,22 +41,26 @@ UObjectPathChannelEvaluatorSystem::UObjectPathChannelEvaluatorSystem(const FObje
 	}
 }
 
-void UObjectPathChannelEvaluatorSystem::OnRun(FSystemTaskPrerequisites& InPrerequisites, FSystemSubsequentTasks& Subsequents)
+void UObjectPathChannelEvaluatorSystem::OnSchedulePersistentTasks(UE::MovieScene::IEntitySystemScheduler* TaskScheduler)
 {
 	using namespace UE::MovieScene;
 
 	FBuiltInComponentTypes* BuiltInComponents = FBuiltInComponentTypes::Get();
 
-	struct FEvaluateObjectPathChannels
-	{
-		static void ForEachEntity(FSourceObjectPathChannel ObjectPathChannel, FFrameTime FrameTime, UObject*& OutResult)
-		{
-			if (!ObjectPathChannel.Source->Evaluate(FrameTime, OutResult))
-			{
-				OutResult = nullptr;
-			}
-		}
-	};
+	FEntityTaskBuilder()
+	.Read(BuiltInComponents->ObjectPathChannel)
+	.Read(BuiltInComponents->EvalTime)
+	.Write(BuiltInComponents->ObjectResult)
+	.FilterNone({ BuiltInComponents->Tags.Ignored })
+	.SetStat(GET_STATID(MovieSceneEval_EvaluateObjectPathChannelTask))
+	.Fork_PerEntity<FEvaluateObjectPathChannels>(&Linker->EntityManager, TaskScheduler);
+}
+
+void UObjectPathChannelEvaluatorSystem::OnRun(FSystemTaskPrerequisites& InPrerequisites, FSystemSubsequentTasks& Subsequents)
+{
+	using namespace UE::MovieScene;
+
+	FBuiltInComponentTypes* BuiltInComponents = FBuiltInComponentTypes::Get();
 
 	FEntityTaskBuilder()
 	.Read(BuiltInComponents->ObjectPathChannel)
diff --git a/Engine/Source/Runtime/MovieSceneTracks/Private/Systems/WeightAndEasingEvaluatorSystem.cpp b/Engine/Source/Runtime/MovieSceneTracks/Private/Systems/WeightAndEasingEvaluatorSystem.cpp
--- a/Engine/Source/Runtime/MovieSceneTracks/Private/Systems/WeightAndEasingEvaluatorSystem.cpp
+++ b/Engine/Source/Runtime/MovieSceneTracks/Private/Systems/WeightAndEasingEvaluatorSystem.cpp
@@ -25,8 +25,11 @@
 
 #include UE_INLINE_GENERATED_CPP_BY_NAME(WeightAndEasingEvaluatorSystem)
 
-DECLARE_CYCLE_STAT(TEXT("MovieScene: Evaluate easing"), MovieSceneEval_EvaluateEasingTask, STATGROUP_MovieSceneECS);
-DECLARE_CYCLE_STAT(TEXT("MovieScene: Harvest easing"), MovieSceneEval_HarvestEasingTask, STATGROUP_MovieSceneECS);
+DECLARE_CYCLE_STAT(TEXT("Weights: Reset"), MovieSceneEval_ResetWeightsTask, STATGROUP_MovieSceneECS);
+DECLARE_CYCLE_STAT(TEXT("Weights: Evaluate easing"), MovieSceneEval_EvaluateEasingTask, STATGROUP_MovieSceneECS);
+DECLARE_CYCLE_STAT(TEXT("Weights: Accumulate manual weights"), MovieSceneEval_AccumulateManualWeights, STATGROUP_MovieSceneECS);
+DECLARE_CYCLE_STAT(TEXT("Weights: Harvest hierarchical easing"), MovieSceneEval_HarvestEasingTask, STATGROUP_MovieSceneECS);
+DECLARE_CYCLE_STAT(TEXT("Weights: Propagate hierarchical easing"), MovieSceneEval_PropagateEasing, STATGROUP_MovieSceneECS);
 
 
 namespace UE::MovieScene
@@ -604,7 +607,7 @@ struct FHarvestHierarchicalEasings
 	}
 
 	// Accumulate all entities that contribute to the channel
-	void ForEachEntity(double Result, uint16 EasingChannel)
+	void ForEachEntity(double Result, uint16 EasingChannel) const
 	{
 		ComputationData[EasingChannel].FinalResult *= Result;
 		ComputationData[EasingChannel].UnaccumulatedResult *= Result;
@@ -682,7 +685,7 @@ struct FPropagateHierarchicalEasings
 		: ComputationData(InComputationData)
 	{}
 
-	void ForEachAllocation(const FEntityAllocation* Allocation, TRead<uint16> HierarchicalEasingChannels, TWrite<double> WeightAndEasingResults)
+	void ForEachAllocation(const FEntityAllocation* Allocation, TRead<uint16> HierarchicalEasingChannels, TWrite<double> WeightAndEasingResults) const
 	{
 		const int32 Num = Allocation->Num();
 		for (int32 Index = 0; Index < Num; ++Index)
@@ -1058,6 +1061,7 @@ UWeightAndEasingEvaluatorSystem::UWeightAndEasingEvaluatorSystem(const FObjectIn
 	using namespace UE::MovieScene;
 
 	SystemCategories = UE::MovieScene::EEntitySystemCategory::ChannelEvaluators;
+	Phase = ESystemPhase::Scheduling;
 
 	if (HasAnyFlags(RF_ClassDefaultObject))
 	{
@@ -1094,6 +1098,65 @@ UE::MovieScene::FHierarchicalEasingChannelBuffer& UWeightAndEasingEvaluatorSyste
 	return PreAllocatedComputationData;
 }
 
+void UWeightAndEasingEvaluatorSystem::OnSchedulePersistentTasks(UE::MovieScene::IEntitySystemScheduler* TaskScheduler)
+{
+	using namespace UE::MovieScene;
+
+	FBuiltInComponentTypes* Components = FBuiltInComponentTypes::Get();
+	FInstanceRegistry* InstanceRegistry = Linker->GetInstanceRegistry();
+
+	// Reset all weights back to one. This can happen immediately.
+	//
+	FTaskID ResetWeights = FEntityTaskBuilder()
+	.Write(Components->WeightAndEasingResult)
+	.SetStat(GET_STATID(MovieSceneEval_ResetWeightsTask))
+	.Fork_PerEntity<FResetFinalWeightResults>(&Linker->EntityManager, TaskScheduler);
+
+	// Evaluate easing - this can be scheduled as soon as eval times are populated
+	//
+	FTaskID EvaluateEasing = FEntityTaskBuilder()
+	.Read(Components->EvalTime)
+	.Read(Components->Easing)
+	.Write(Components->EasingResult)
+	.SetStat(GET_STATID(MovieSceneEval_EvaluateEasingTask))
+	.Fork_PerEntity<FEvaluateEasings>(&Linker->EntityManager, TaskScheduler);
+
+	// Accumulate weights - must be done after ResetWeights and EvaluateEasing
+	//
+	FTaskID AccumulateWeights = FEntityTaskBuilder()
+	.ReadOneOrMoreOf(Components->WeightResult, Components->EasingResult)
+	.Write(Components->WeightAndEasingResult)
+	.SetStat(GET_STATID(MovieSceneEval_AccumulateManualWeights))
+	.Schedule_PerAllocation<FAccumulateManualWeights>(&Linker->EntityManager, TaskScheduler);
+
+	TaskScheduler->AddPrerequisite(ResetWeights, AccumulateWeights);
+	TaskScheduler->AddPrerequisite(EvaluateEasing, AccumulateWeights);
+
+	// If we have hierarchical easing, we initialize all the weights to their hierarchical defaults
+	if (!PreAllocatedComputationData.IsEmpty())
+	{
+		// Step 2: Harvest any hierarchical results from providers
+		//
+		FTaskID HarvestTask = FEntityTaskBuilder()
+		.Read(Components->WeightAndEasingResult)
+		.Read(Components->HierarchicalEasingChannel)
+		.FilterAll({ Components->HierarchicalEasingProvider })  // Only harvest results from entities that are providing results
+		.SetStat(GET_STATID(MovieSceneEval_HarvestEasingTask))
+		.Schedule_PerEntity<FHarvestHierarchicalEasings>(&Linker->EntityManager, TaskScheduler, &PreAllocatedComputationData.Channels);
+
+		// Step 3: Apply hierarchical easing results to all entities inside affected sub-sequences.
+		//
+		FTaskID PropagateTask = FEntityTaskBuilder()
+		.Read(Components->HierarchicalEasingChannel)
+		.Write(Components->WeightAndEasingResult)
+		.FilterNone({ Components->HierarchicalEasingProvider }) // Do not propagate hierarchical weights onto providers!
+		.SetStat(GET_STATID(MovieSceneEval_PropagateEasing))
+		.Fork_PerAllocation<FPropagateHierarchicalEasings>(&Linker->EntityManager, TaskScheduler, PreAllocatedComputationData.Channels);
+
+		TaskScheduler->AddPrerequisite(HarvestTask, PropagateTask);
+	}
+}
+
 void UWeightAndEasingEvaluatorSystem::OnRun(FSystemTaskPrerequisites& InPrerequisites, FSystemSubsequentTasks& Subsequents)
 {
 	using namespace UE::MovieScene;
diff --git a/Engine/Source/Runtime/MovieSceneTracks/Public/Systems/ByteChannelEvaluatorSystem.h b/Engine/Source/Runtime/MovieSceneTracks/Public/Systems/ByteChannelEvaluatorSystem.h
--- a/Engine/Source/Runtime/MovieSceneTracks/Public/Systems/ByteChannelEvaluatorSystem.h
+++ b/Engine/Source/Runtime/MovieSceneTracks/Public/Systems/ByteChannelEvaluatorSystem.h
@@ -10,16 +10,6 @@
 
 class UObject;
 
-namespace UE
-{
-namespace MovieScene
-{
-
-	struct FSourceByteChannel;
-
-} // namespace MovieScene
-} // namespace UE
-
 /**
  * System that is responsible for evaluating byte channels.
  */
@@ -32,5 +22,6 @@ public:
 
 	UByteChannelEvaluatorSystem(const FObjectInitializer& ObjInit);
 
+	virtual void OnSchedulePersistentTasks(UE::MovieScene::IEntitySystemScheduler* TaskScheduler) override;
 	virtual void OnRun(FSystemTaskPrerequisites& InPrerequisites, FSystemSubsequentTasks& Subsequents) override;
 };
diff --git a/Engine/Source/Runtime/MovieSceneTracks/Public/Systems/FloatChannelEvaluatorSystem.h b/Engine/Source/Runtime/MovieSceneTracks/Public/Systems/FloatChannelEvaluatorSystem.h
--- a/Engine/Source/Runtime/MovieSceneTracks/Public/Systems/FloatChannelEvaluatorSystem.h
+++ b/Engine/Source/Runtime/MovieSceneTracks/Public/Systems/FloatChannelEvaluatorSystem.h
@@ -13,18 +13,14 @@
 
 #include "FloatChannelEvaluatorSystem.generated.h"
 
-class UMovieSceneEntitySystemLinker;
-class UObject;
-
 namespace UE::MovieScene
 {
 	struct FSourceFloatChannel;
-
 	namespace Interpolation
 	{
 		struct FCachedInterpolation;
 	}
-} // namespace UE::MovieScene
+}
 
 /**
  * System that is responsible for evaluating double channels.
@@ -39,6 +35,7 @@ public:
 
 	UFloatChannelEvaluatorSystem(const FObjectInitializer& ObjInit);
 
+	virtual void OnSchedulePersistentTasks(UE::MovieScene::IEntitySystemScheduler* TaskScheduler) override;
 	virtual void OnRun(FSystemTaskPrerequisites& InPrerequisites, FSystemSubsequentTasks& Subsequents) override;
 	virtual bool IsRelevantImpl(UMovieSceneEntitySystemLinker* InLinker) const override;
 
diff --git a/Engine/Source/Runtime/MovieSceneTracks/Public/Systems/FloatPerlinNoiseChannelEvaluatorSystem.h b/Engine/Source/Runtime/MovieSceneTracks/Public/Systems/FloatPerlinNoiseChannelEvaluatorSystem.h
--- a/Engine/Source/Runtime/MovieSceneTracks/Public/Systems/FloatPerlinNoiseChannelEvaluatorSystem.h
+++ b/Engine/Source/Runtime/MovieSceneTracks/Public/Systems/FloatPerlinNoiseChannelEvaluatorSystem.h
@@ -15,5 +15,7 @@ public:
 	GENERATED_BODY()
 
 	UFloatPerlinNoiseChannelEvaluatorSystem(const FObjectInitializer& ObjInit);
+
+	virtual void OnSchedulePersistentTasks(UE::MovieScene::IEntitySystemScheduler* TaskScheduler) override;
 	virtual void OnRun(FSystemTaskPrerequisites& InPrerequisites, FSystemSubsequentTasks& Subsequents) override;
 };
diff --git a/Engine/Source/Runtime/MovieSceneTracks/Public/Systems/IntegerChannelEvaluatorSystem.h b/Engine/Source/Runtime/MovieSceneTracks/Public/Systems/IntegerChannelEvaluatorSystem.h
--- a/Engine/Source/Runtime/MovieSceneTracks/Public/Systems/IntegerChannelEvaluatorSystem.h
+++ b/Engine/Source/Runtime/MovieSceneTracks/Public/Systems/IntegerChannelEvaluatorSystem.h
@@ -8,18 +8,6 @@
 
 #include "IntegerChannelEvaluatorSystem.generated.h"
 
-class UObject;
-
-namespace UE
-{
-namespace MovieScene
-{
-
-	struct FSourceIntegerChannel;
-
-} // namespace MovieScene
-} // namespace UE
-
 /**
  * System that is responsible for evaluating integer channels.
  */
@@ -32,5 +20,6 @@ public:
 
 	UIntegerChannelEvaluatorSystem(const FObjectInitializer& ObjInit);
 
+	virtual void OnSchedulePersistentTasks(UE::MovieScene::IEntitySystemScheduler* TaskScheduler) override;
 	virtual void OnRun(FSystemTaskPrerequisites& InPrerequisites, FSystemSubsequentTasks& Subsequents) override;
 };
diff --git a/Engine/Source/Runtime/MovieSceneTracks/Public/Systems/MovieSceneQuaternionBlenderSystem.h b/Engine/Source/Runtime/MovieSceneTracks/Public/Systems/MovieSceneQuaternionBlenderSystem.h
--- a/Engine/Source/Runtime/MovieSceneTracks/Public/Systems/MovieSceneQuaternionBlenderSystem.h
+++ b/Engine/Source/Runtime/MovieSceneTracks/Public/Systems/MovieSceneQuaternionBlenderSystem.h
@@ -22,7 +22,6 @@ namespace UE::MovieScene
 	struct FQuaternionBlenderAccumulationBuffers
 	{
 		TSparseArray<double> AbsoluteWeights;
-		TArray<FQuatTransform> InitialValues;
 		TArray<FQuatTransform> Absolutes;
 		TArray<FQuatTransform> Additives;
 	};
@@ -36,6 +35,7 @@ public:
 
 	UMovieSceneQuaternionBlenderSystem(const FObjectInitializer& ObjInit);
 
+	virtual void OnSchedulePersistentTasks(UE::MovieScene::IEntitySystemScheduler* TaskScheduler) override;
 	virtual void OnRun(FSystemTaskPrerequisites& InPrerequisites, FSystemSubsequentTasks& Subsequents) override;
 
 	virtual FGraphEventRef DispatchDecomposeTask(const UE::MovieScene::FValueDecompositionParams& Params, UE::MovieScene::FAlignedDecomposedValue* Output) override;
