Sequencer: Updated a few missing systems to the task scheduler API

diff --git a/Engine/Plugins/MovieScene/TemplateSequence/Source/TemplateSequence/Private/Systems/TemplateSequenceSystem.cpp b/Engine/Plugins/MovieScene/TemplateSequence/Source/TemplateSequence/Private/Systems/TemplateSequenceSystem.cpp
--- a/Engine/Plugins/MovieScene/TemplateSequence/Source/TemplateSequence/Private/Systems/TemplateSequenceSystem.cpp
+++ b/Engine/Plugins/MovieScene/TemplateSequence/Source/TemplateSequence/Private/Systems/TemplateSequenceSystem.cpp
@@ -261,7 +261,7 @@ UTemplateSequencePropertyScalingEvaluatorSystem::UTemplateSequencePropertyScalin
 	const FBuiltInComponentTypes* BuiltInComponents = FBuiltInComponentTypes::Get();
 	const FTemplateSequenceComponentTypes* TemplateSequenceComponents = FTemplateSequenceComponentTypes::Get();
 
-	Phase = ESystemPhase::Evaluation;
+	Phase = ESystemPhase::Scheduling;
 	RelevantComponent = TemplateSequenceComponents->PropertyScale;
 
 	if (HasAnyFlags(RF_ClassDefaultObject))
@@ -283,6 +283,11 @@ UTemplateSequencePropertyScalingEvaluatorSystem::UTemplateSequencePropertyScalin
 	}
 }
 
+void UTemplateSequencePropertyScalingEvaluatorSystem::ResetPropertyScales()
+{
+	PropertyScales.Reset();
+}
+
 void UTemplateSequencePropertyScalingEvaluatorSystem::AddPropertyScale(const FPropertyScaleKey& Key, const FPropertyScaleValue& Value)
 {
 	FMultiPropertyScaleValue& MultiValue = PropertyScales.FindOrAdd(Key);
@@ -340,10 +345,16 @@ struct FGatherPropertyScales
 		InstanceRegistry = InEvaluatorSystem->GetLinker()->GetInstanceRegistry();
 	}
 
+	void PreTask()
+	{
+		// Clear our map of property scales from last frame.
+		EvaluatorSystem->ResetPropertyScales();
+	}
+
 	void ForEachEntity(
 			const FInstanceHandle& InstanceHandle, 
 			const FTemplateSequencePropertyScaleComponentData& PropertyScale,
-			float ScaleFactor)
+			float ScaleFactor) const
 	{
 		// Find the instance handle of the sub-sequence we need to scale.
 		const FInstanceHandle SubSequenceHandle = InstanceRegistry->FindRelatedInstanceHandle(InstanceHandle, PropertyScale.SubSequenceID);
@@ -360,9 +371,9 @@ struct FScaleTransformProperties
 {
 	using UEvaluatorSystem = UTemplateSequencePropertyScalingEvaluatorSystem;
 
-	UEvaluatorSystem* EvaluatorSystem;
+	const UEvaluatorSystem* EvaluatorSystem;
 
-	FScaleTransformProperties(UEvaluatorSystem* InEvaluatorSystem)
+	FScaleTransformProperties(const UEvaluatorSystem* InEvaluatorSystem)
 		: EvaluatorSystem(InEvaluatorSystem)
 	{
 		check(InEvaluatorSystem);
@@ -380,7 +391,7 @@ struct FScaleTransformProperties
 			TWriteOptional<double> LocationZs,
 			TWriteOptional<double> RotationXs,
 			TWriteOptional<double> RotationYs,
-			TWriteOptional<double> RotationZs)
+			TWriteOptional<double> RotationZs) const
 	{
 		for (int32 Index = 0; Index < Allocation->Num(); ++Index)
 		{
@@ -424,9 +435,9 @@ struct FScaleFloatProperties
 {
 	using UEvaluatorSystem = UTemplateSequencePropertyScalingEvaluatorSystem;
 
-	UEvaluatorSystem* EvaluatorSystem;
+	const UEvaluatorSystem* EvaluatorSystem;
 
-	FScaleFloatProperties(UEvaluatorSystem* InEvaluatorSystem)
+	FScaleFloatProperties(const UEvaluatorSystem* InEvaluatorSystem)
 		: EvaluatorSystem(InEvaluatorSystem)
 	{
 		check(InEvaluatorSystem);
@@ -438,7 +449,7 @@ struct FScaleFloatProperties
 			TRead<FMovieScenePropertyBinding> PropertyBindings,
 			TRead<FGuid> ReverseBindingLookups,
 			TReadOptional<double> BasePropertyValues,
-			TWrite<double> PropertyValues)
+			TWrite<double> PropertyValues) const
 	{
 		for (int32 Index = 0; Index < Allocation->Num(); ++Index)
 		{
@@ -471,6 +482,71 @@ struct FScaleFloatProperties
 } // namespace MovieScene
 } // namespace UE
 
+void UTemplateSequencePropertyScalingEvaluatorSystem::OnSchedulePersistentTasks(UE::MovieScene::IEntitySystemScheduler* TaskScheduler)
+{
+	using namespace UE::MovieScene;
+
+	const FBuiltInComponentTypes* BuiltInComponents = FBuiltInComponentTypes::Get();
+	const FMovieSceneTracksComponentTypes* TrackComponents = FMovieSceneTracksComponentTypes::Get();
+	const FTemplateSequenceComponentTypes* TemplateSequenceComponents = FTemplateSequenceComponentTypes::Get();
+
+	const FInstanceRegistry* InstanceRegistry = Linker->GetInstanceRegistry();
+	const UTemplateSequencePropertyScalingInstantiatorSystem* InstantiatorSystem = Linker->FindSystem<UTemplateSequencePropertyScalingInstantiatorSystem>();
+
+	// Step 1: We are going to look for all entities that describe an active property scale, pick up its up-to-date (evaluated)
+	// scale value (which was evaluated by the float channel evaluator), and build a map that tells us:
+	// - What properties (sub-sequence instance, object/component binding, property path)...
+	// - ...are scaled by what factor (scale value), using what type of scale (scale type enum).
+	//
+	FTaskID GatherTask = FEntityTaskBuilder()
+	.Read(BuiltInComponents->InstanceHandle)
+	.Read(TemplateSequenceComponents->PropertyScale)
+	.Read(BuiltInComponents->DoubleResult[0])
+	.Schedule_PerEntity<FGatherPropertyScales>(&Linker->EntityManager, TaskScheduler, this);
+
+
+	// Step 2: If we have any transform property scaling, iterate specifically on them and apply the scale factor
+	// to the appropriate components.
+	if (ensure(InstantiatorSystem) && InstantiatorSystem->HasAnyTransformScales())
+	{
+		FTaskID ScaleTransformTask = FEntityTaskBuilder()
+			.Read(BuiltInComponents->InstanceHandle)
+			.Read(BuiltInComponents->PropertyBinding)
+			.Read(TemplateSequenceComponents->PropertyScaleReverseBindingLookup)
+			.ReadOptional(BuiltInComponents->BaseDouble[0])
+			.ReadOptional(BuiltInComponents->BaseDouble[1])
+			.ReadOptional(BuiltInComponents->BaseDouble[2])
+			.ReadOptional(BuiltInComponents->BaseDouble[3])
+			.ReadOptional(BuiltInComponents->BaseDouble[4])
+			.ReadOptional(BuiltInComponents->BaseDouble[5])
+			.WriteOptional(BuiltInComponents->DoubleResult[0])
+			.WriteOptional(BuiltInComponents->DoubleResult[1])
+			.WriteOptional(BuiltInComponents->DoubleResult[2])
+			.WriteOptional(BuiltInComponents->DoubleResult[3])
+			.WriteOptional(BuiltInComponents->DoubleResult[4])
+			.WriteOptional(BuiltInComponents->DoubleResult[5])
+			.FilterAny({ TrackComponents->ComponentTransform.PropertyTag, TrackComponents->Transform.PropertyTag })
+			.Fork_PerAllocation<FScaleTransformProperties>(&Linker->EntityManager, TaskScheduler, this);
+
+		TaskScheduler->AddPrerequisite(GatherTask, ScaleTransformTask);
+	}
+
+	// Step 3: Do the same as step 2 for float properties.
+	if (ensure(InstantiatorSystem) && InstantiatorSystem->HasAnyFloatScales())
+	{
+		FTaskID ScaleFloatTask = FEntityTaskBuilder()
+			.Read(BuiltInComponents->InstanceHandle)
+			.Read(BuiltInComponents->PropertyBinding)
+			.Read(TemplateSequenceComponents->PropertyScaleReverseBindingLookup)
+			.ReadOptional(BuiltInComponents->BaseDouble[0])
+			.Write(BuiltInComponents->DoubleResult[0])
+			.FilterAll({ TrackComponents->Float.PropertyTag })
+			.Fork_PerAllocation<FScaleFloatProperties>(&Linker->EntityManager, TaskScheduler, this);
+
+		TaskScheduler->AddPrerequisite(GatherTask, ScaleFloatTask);
+	}
+}
+
 void UTemplateSequencePropertyScalingEvaluatorSystem::OnRun(FSystemTaskPrerequisites& InPrerequisites, FSystemSubsequentTasks& Subsequents)
 {
 	using namespace UE::MovieScene;
diff --git a/Engine/Plugins/MovieScene/TemplateSequence/Source/TemplateSequence/Private/Systems/TemplateSequenceSystem.h b/Engine/Plugins/MovieScene/TemplateSequence/Source/TemplateSequence/Private/Systems/TemplateSequenceSystem.h
--- a/Engine/Plugins/MovieScene/TemplateSequence/Source/TemplateSequence/Private/Systems/TemplateSequenceSystem.h
+++ b/Engine/Plugins/MovieScene/TemplateSequence/Source/TemplateSequence/Private/Systems/TemplateSequenceSystem.h
@@ -63,11 +63,13 @@ public:
 
 	using FPropertyScaleValueArray = TArray<FPropertyScaleValue, TInlineAllocator<2>>;
 
+	void ResetPropertyScales();
 	void AddPropertyScale(const FPropertyScaleKey& Key, const FPropertyScaleValue& Value);
 	void FindPropertyScales(const FPropertyScaleKey& Key, FPropertyScaleValueArray& OutValues) const;
 
 private:
 
+	virtual void OnSchedulePersistentTasks(UE::MovieScene::IEntitySystemScheduler* TaskScheduler) override final;
 	virtual void OnRun(FSystemTaskPrerequisites& InPrerequisites, FSystemSubsequentTasks& Subsequents) override final;
 	
 private:
diff --git a/Engine/Source/Runtime/MovieScene/Private/EntitySystem/TrackInstance/MovieSceneTrackInstanceSystem.cpp b/Engine/Source/Runtime/MovieScene/Private/EntitySystem/TrackInstance/MovieSceneTrackInstanceSystem.cpp
--- a/Engine/Source/Runtime/MovieScene/Private/EntitySystem/TrackInstance/MovieSceneTrackInstanceSystem.cpp
+++ b/Engine/Source/Runtime/MovieScene/Private/EntitySystem/TrackInstance/MovieSceneTrackInstanceSystem.cpp
@@ -260,6 +260,7 @@ void UMovieSceneTrackInstanceInstantiator::OnRun(FSystemTaskPrerequisites& InPre
 UMovieSceneTrackInstanceSystem::UMovieSceneTrackInstanceSystem(const FObjectInitializer& ObjInit)
 	: Super(ObjInit)
 {
+	Phase = UE::MovieScene::ESystemPhase::Scheduling;
 	RelevantComponent = UE::MovieScene::FBuiltInComponentTypes::Get()->TrackInstance;
 }
 
@@ -269,6 +270,15 @@ void UMovieSceneTrackInstanceSystem::OnLink()
 	Linker->SystemGraph.AddReference(this, Instantiator);
 }
 
+void UMovieSceneTrackInstanceSystem::OnSchedulePersistentTasks(UE::MovieScene::IEntitySystemScheduler* TaskScheduler)
+{
+	using namespace UE::MovieScene;
+	if (this->Instantiator->GetTrackInstances().Num() != 0)
+	{
+		TaskScheduler->AddMemberFunctionTask(FTaskParams(TEXT("Evaluate Track Instances")).ForceGameThread(), this, &UMovieSceneTrackInstanceSystem::EvaluateAllInstances);
+	}
+}
+
 void UMovieSceneTrackInstanceSystem::OnRun(FSystemTaskPrerequisites& InPrerequisites, FSystemSubsequentTasks& Subsequents)
 {
 	using namespace UE::MovieScene;
@@ -277,25 +287,25 @@ void UMovieSceneTrackInstanceSystem::OnRun(FSystemTaskPrerequisites& InPrerequis
 
 	if (this->Instantiator->GetTrackInstances().Num() != 0)
 	{
-		auto Run = [this]
-		{
-			for (const FMovieSceneTrackInstanceEntry& Entry : this->Instantiator->GetTrackInstances())
-			{
-				if (ensure(Entry.TrackInstance))
-				{
-					Entry.TrackInstance->Animate();
-				}
-			}
-		};
-
 		if (Linker->EntityManager.GetThreadingModel() == EEntityThreadingModel::NoThreading)
 		{
-			Run();
+			this->EvaluateAllInstances();
 		}
 		else
 		{
-			FGraphEventRef Task = FFunctionGraphTask::CreateAndDispatchWhenReady(MoveTemp(Run), GET_STATID(MovieSceneEval_GenericTrackInstanceTask), InPrerequisites.All(), Linker->EntityManager.GetGatherThread());
+			FGraphEventRef Task = FFunctionGraphTask::CreateAndDispatchWhenReady([this]{ this->EvaluateAllInstances(); }, GET_STATID(MovieSceneEval_GenericTrackInstanceTask), InPrerequisites.All(), Linker->EntityManager.GetGatherThread());
 			Subsequents.AddRootTask(Task);
 		}
 	}
 }
+
+void UMovieSceneTrackInstanceSystem::EvaluateAllInstances()
+{
+	for (const FMovieSceneTrackInstanceEntry& Entry : this->Instantiator->GetTrackInstances())
+	{
+		if (ensure(Entry.TrackInstance))
+		{
+			Entry.TrackInstance->Animate();
+		}
+	}
+}
\ No newline at end of file
diff --git a/Engine/Source/Runtime/MovieScene/Public/EntitySystem/MovieSceneEntitySystemTask.h b/Engine/Source/Runtime/MovieScene/Public/EntitySystem/MovieSceneEntitySystemTask.h
--- a/Engine/Source/Runtime/MovieScene/Public/EntitySystem/MovieSceneEntitySystemTask.h
+++ b/Engine/Source/Runtime/MovieScene/Public/EntitySystem/MovieSceneEntitySystemTask.h
@@ -1133,9 +1133,9 @@ struct TFilteredEntityTask
 	/**
 	 * Assign the scheduled task parameters for this task
 	 */
-	TEntityTaskComponents< T... >& SetParams(const FTaskParams& InOtherParams)
+	TFilteredEntityTask< T... >& SetParams(const FTaskParams& InOtherParams)
 	{
-		this->CommonParams.TaskParams = InOtherParams;
+		Components.CommonParams.TaskParams = InOtherParams;
 		return *this;
 	}
 
diff --git a/Engine/Source/Runtime/MovieScene/Public/EntitySystem/TrackInstance/MovieSceneTrackInstanceSystem.h b/Engine/Source/Runtime/MovieScene/Public/EntitySystem/TrackInstance/MovieSceneTrackInstanceSystem.h
--- a/Engine/Source/Runtime/MovieScene/Public/EntitySystem/TrackInstance/MovieSceneTrackInstanceSystem.h
+++ b/Engine/Source/Runtime/MovieScene/Public/EntitySystem/TrackInstance/MovieSceneTrackInstanceSystem.h
@@ -85,8 +85,11 @@ class MOVIESCENE_API UMovieSceneTrackInstanceSystem : public UMovieSceneEntitySy
 private:
 
 	virtual void OnLink() override final;
+	virtual void OnSchedulePersistentTasks(UE::MovieScene::IEntitySystemScheduler* TaskScheduler) override final;
 	virtual void OnRun(FSystemTaskPrerequisites& InPrerequisites, FSystemSubsequentTasks& Subsequents) override final;
 
+	void EvaluateAllInstances();
+
 	UPROPERTY()
 	TObjectPtr<UMovieSceneTrackInstanceInstantiator> Instantiator;
 };
diff --git a/Engine/Source/Runtime/MovieSceneTracks/Private/Systems/MovieSceneComponentMaterialSystem.cpp b/Engine/Source/Runtime/MovieSceneTracks/Private/Systems/MovieSceneComponentMaterialSystem.cpp
--- a/Engine/Source/Runtime/MovieSceneTracks/Private/Systems/MovieSceneComponentMaterialSystem.cpp
+++ b/Engine/Source/Runtime/MovieSceneTracks/Private/Systems/MovieSceneComponentMaterialSystem.cpp
@@ -100,7 +100,7 @@ UMovieSceneComponentMaterialSystem::UMovieSceneComponentMaterialSystem(const FOb
 	FMovieSceneTracksComponentTypes* TracksComponents = FMovieSceneTracksComponentTypes::Get();
 
 	RelevantComponent = TracksComponents->ComponentMaterialIndex;
-	Phase = ESystemPhase::Instantiation | ESystemPhase::Evaluation;
+	Phase = ESystemPhase::Instantiation;
 
 	if (HasAnyFlags(RF_ClassDefaultObject))
 	{
diff --git a/Engine/Source/Runtime/MovieSceneTracks/Private/Systems/MovieSceneConstraintSystem.cpp b/Engine/Source/Runtime/MovieSceneTracks/Private/Systems/MovieSceneConstraintSystem.cpp
--- a/Engine/Source/Runtime/MovieSceneTracks/Private/Systems/MovieSceneConstraintSystem.cpp
+++ b/Engine/Source/Runtime/MovieSceneTracks/Private/Systems/MovieSceneConstraintSystem.cpp
@@ -154,6 +154,64 @@ struct FPreAnimatedConstraintTraits : FBoundObjectPreAnimatedStateTraits
 	}
 };
 
+
+struct FEvaluateConstraintChannels
+{
+	FEvaluateConstraintChannels(UWorld* InWorld, FInstanceRegistry* InInstanceRegistry, UMovieSceneConstraintSystem* InSystem)
+		: World(InWorld), InstanceRegistry(InInstanceRegistry), Controller(nullptr), System(InSystem)
+	{
+		check(World && InstanceRegistry && System);
+	}
+
+	void PreTask()
+	{
+		// Retrieve the controller at the start of the task since this task is now stored persistently
+		Controller = &FConstraintsManagerController::Get(World);
+		check(Controller);
+
+		System->DynamicOffsets.Reset();
+	}
+
+	void ForEachEntity(UObject* BoundObject, FInstanceHandle InstanceHandle, const FConstraintComponentData& ConstraintChannel, FFrameTime FrameTime) const
+	{
+		FConstraintAndActiveChannel* ConstraintAndActiveChannel = ConstraintChannel.Section->GetConstraintChannel(ConstraintChannel.ConstraintName);
+		if (!ConstraintAndActiveChannel)
+		{
+			return;
+		}
+
+		UTickableConstraint* Constraint = CreateConstraintIfNeeded(*Controller, ConstraintAndActiveChannel, ConstraintChannel);
+		if (!Constraint)
+		{
+			return;
+		}
+
+		const FSequenceInstance& TargetInstance = InstanceRegistry->GetInstance(InstanceHandle);
+
+		bool Result = false;
+		ConstraintAndActiveChannel->ActiveChannel.Evaluate(FrameTime, Result);
+		Constraint->SetActive(Result);
+		Constraint->ResolveBoundObjects(TargetInstance.GetSequenceID(), *TargetInstance.GetPlayer());
+		if (UTickableTransformConstraint* TransformConstraint = Cast< UTickableTransformConstraint>(Constraint))
+		{
+			if (UTransformableComponentHandle* ComponentHandle = Cast<UTransformableComponentHandle>(TransformConstraint->ChildTRSHandle))
+			{
+				//bound component may change so need to update constraint
+				ComponentHandle->Component = Cast<USceneComponent>(BoundObject);
+				UMovieSceneConstraintSystem::FUpdateHandleForConstraint UpdateHandle;
+				UpdateHandle.Constraint = TransformConstraint;
+				UpdateHandle.TransformHandle = ComponentHandle;
+				System->DynamicOffsets.Add(UpdateHandle);
+			}
+		}
+	}
+
+	UWorld* World;
+	FInstanceRegistry* InstanceRegistry;
+	FConstraintsManagerController* Controller;
+	UMovieSceneConstraintSystem* System;
+};
+
 struct FPreAnimatedConstraintStorage
 	: public TPreAnimatedStateStorage<FPreAnimatedConstraintTraits>
 {
@@ -174,7 +232,7 @@ UMovieSceneConstraintSystem::UMovieSceneConstraintSystem(const FObjectInitialize
 	RelevantComponent = TracksComponents->ConstraintChannel;
 
 	// Run constraints during instantiation or evaluation
-	Phase = ESystemPhase::Instantiation | ESystemPhase::Evaluation | ESystemPhase::Finalization;
+	Phase = ESystemPhase::Instantiation | ESystemPhase::Scheduling | ESystemPhase::Finalization;
 
 	if (HasAnyFlags(RF_ClassDefaultObject) ) 
 	{
@@ -182,10 +240,27 @@ UMovieSceneConstraintSystem::UMovieSceneConstraintSystem(const FObjectInitialize
 		// This is only really necessary if they are in the same phase (which they are not), but I've
 		// defined the prerequisite for safety if its phase changes in future
 		DefineImplicitPrerequisite(GetClass(), UMovieSceneComponentTransformSystem::StaticClass());
-		DefineComponentConsumer(GetClass(), UE::MovieScene::FBuiltInComponentTypes::Get()->EvalTime);
+		DefineComponentConsumer(GetClass(), FBuiltInComponentTypes::Get()->EvalTime);
 	}
 }
 
+void UMovieSceneConstraintSystem::OnSchedulePersistentTasks(UE::MovieScene::IEntitySystemScheduler* TaskScheduler)
+{
+	using namespace UE::MovieScene;
+
+	FMovieSceneTracksComponentTypes* TracksComponents = FMovieSceneTracksComponentTypes::Get();
+	FBuiltInComponentTypes* BuiltInComponents = FBuiltInComponentTypes::Get();
+
+	// Set up new constraints
+	FEntityTaskBuilder()
+	.Read(BuiltInComponents->BoundObject)
+	.Read(BuiltInComponents->InstanceHandle)
+	.Read(TracksComponents->ConstraintChannel)
+	.Read(BuiltInComponents->EvalTime)
+	.SetDesiredThread(Linker->EntityManager.GetGatherThread())
+	.Schedule_PerEntity<FEvaluateConstraintChannels>(&Linker->EntityManager, TaskScheduler, GetWorld(), Linker->GetInstanceRegistry(), this);
+}
+
 void UMovieSceneConstraintSystem::OnRun(FSystemTaskPrerequisites& InPrerequisites, FSystemSubsequentTasks& Subsequents)
 {
 	using namespace UE::MovieScene;
@@ -205,55 +280,18 @@ void UMovieSceneConstraintSystem::OnRun(FSystemTaskPrerequisites& InPrerequisite
 	}
 	else if (CurrentPhase == ESystemPhase::Evaluation)
 	{
-		DynamicOffsets.Reset();
-		FConstraintsManagerController& Controller = FConstraintsManagerController::Get(GetWorld());
+		// Backwards compat with legacy non-persistent tasks
 		FMovieSceneTracksComponentTypes* TracksComponents = FMovieSceneTracksComponentTypes::Get();
 		FBuiltInComponentTypes* BuiltInComponents = FBuiltInComponentTypes::Get();
 
-		struct FEvaluateConstraintChannels
-		{
-			FEvaluateConstraintChannels(FInstanceRegistry* InInstanceRegistry, FConstraintsManagerController* InController, UMovieSceneConstraintSystem* InSystem) : InstanceRegistry(InInstanceRegistry), Controller(InController), System(InSystem) { check(InstanceRegistry); check(Controller); check(System) };
-			void ForEachEntity(UObject* BoundObject, UE::MovieScene::FInstanceHandle InstanceHandle, const FConstraintComponentData& ConstraintChannel, FFrameTime FrameTime)
-			{
-				const FSequenceInstance& TargetInstance = InstanceRegistry->GetInstance(InstanceHandle);
-				if (FConstraintAndActiveChannel* ConstraintAndActiveChannel = ConstraintChannel.Section->GetConstraintChannel(ConstraintChannel.ConstraintName))
-				{
-					UTickableConstraint* Constraint = CreateConstraintIfNeeded(*Controller, ConstraintAndActiveChannel, ConstraintChannel);
-
-					if (Constraint)
-					{
-						bool Result = false;
-						ConstraintAndActiveChannel->ActiveChannel.Evaluate(FrameTime, Result);
-						Constraint->SetActive(Result);
-						Constraint->ResolveBoundObjects(TargetInstance.GetSequenceID(), *TargetInstance.GetPlayer());
-						if (UTickableTransformConstraint* TransformConstraint = Cast< UTickableTransformConstraint>(Constraint))
-						{
-							if (UTransformableComponentHandle* ComponentHandle = Cast<UTransformableComponentHandle>(TransformConstraint->ChildTRSHandle))
-							{
-								//bound component may change so need to update constraint
-								ComponentHandle->Component = Cast<USceneComponent>(BoundObject);
-								FUpdateHandleForConstraint UpdateHandle;
-								UpdateHandle.Constraint = TransformConstraint;
-								UpdateHandle.TransformHandle = ComponentHandle;
-								System->DynamicOffsets.Add(UpdateHandle);
-							}
-						}
-					}
-				}
-			}
-			FInstanceRegistry* InstanceRegistry;
-			FConstraintsManagerController* Controller;
-			UMovieSceneConstraintSystem* System;
-		};
-
 		// Set up new constraints
 		FEntityTaskBuilder()
-			.Read(BuiltInComponents->BoundObject)
-			.Read(BuiltInComponents->InstanceHandle)
-			.Read(TracksComponents->ConstraintChannel)
-			.Read(BuiltInComponents->EvalTime)
-			.SetDesiredThread(Linker->EntityManager.GetGatherThread())
-			.Dispatch_PerEntity<FEvaluateConstraintChannels>(&Linker->EntityManager, InPrerequisites, &Subsequents, Linker->GetInstanceRegistry(), &Controller, this);
+		.Read(BuiltInComponents->BoundObject)
+		.Read(BuiltInComponents->InstanceHandle)
+		.Read(TracksComponents->ConstraintChannel)
+		.Read(BuiltInComponents->EvalTime)
+		.SetDesiredThread(Linker->EntityManager.GetGatherThread())
+		.Dispatch_PerEntity<FEvaluateConstraintChannels>(&Linker->EntityManager, InPrerequisites, &Subsequents, GetWorld(), Linker->GetInstanceRegistry(), this);
 	}
 	else if (CurrentPhase == ESystemPhase::Finalization)
 	{
diff --git a/Engine/Source/Runtime/MovieSceneTracks/Private/Systems/MovieSceneDeferredComponentMovementSystem.cpp b/Engine/Source/Runtime/MovieSceneTracks/Private/Systems/MovieSceneDeferredComponentMovementSystem.cpp
--- a/Engine/Source/Runtime/MovieSceneTracks/Private/Systems/MovieSceneDeferredComponentMovementSystem.cpp
+++ b/Engine/Source/Runtime/MovieSceneTracks/Private/Systems/MovieSceneDeferredComponentMovementSystem.cpp
@@ -47,7 +47,7 @@ UMovieSceneDeferredComponentMovementSystem::UMovieSceneDeferredComponentMovement
 {
 	using namespace UE::MovieScene;
 
-	Phase = ESystemPhase::Instantiation | ESystemPhase::Evaluation | ESystemPhase::Finalization;
+	Phase = ESystemPhase::Instantiation | ESystemPhase::Scheduling | ESystemPhase::Finalization;
 
 	if (HasAnyFlags(RF_ClassDefaultObject))
 	{
@@ -111,6 +111,36 @@ bool UMovieSceneDeferredComponentMovementSystem::IsRelevantImpl(UMovieSceneEntit
 		InLinker->EntityManager.ContainsComponent(Components->AttachParent);
 }
 
+void UMovieSceneDeferredComponentMovementSystem::OnSchedulePersistentTasks(UE::MovieScene::IEntitySystemScheduler* TaskScheduler)
+{
+	using namespace UE::MovieScene;
+
+	FBuiltInComponentTypes*          BuiltInComponents = FBuiltInComponentTypes::Get();
+	FMovieSceneTracksComponentTypes* Components        = FMovieSceneTracksComponentTypes::Get();
+
+	FTaskParams TaskParams = FTaskParams(TEXT("Defer Movement")).ForceGameThread();
+	TaskParams.bForcePropagateDownstream = true;
+
+	struct FCacheDeferredUpdates
+	{
+		UMovieSceneDeferredComponentMovementSystem* System;
+
+		void ForEachEntity(UObject* BoundObject) const
+		{
+			if (USceneComponent* SceneComponent = Cast<USceneComponent>(BoundObject))
+			{
+				System->DeferMovementUpdates(SceneComponent);
+			}
+		}
+	};
+
+	FEntityTaskBuilder()
+	.Read(BuiltInComponents->BoundObject)
+	.FilterAny({ Components->ComponentTransform.PropertyTag })
+	.SetParams(TaskParams)
+	.Schedule_PerEntity<FCacheDeferredUpdates>(&Linker->EntityManager, TaskScheduler, this);
+}
+
 void UMovieSceneDeferredComponentMovementSystem::OnRun(FSystemTaskPrerequisites& InPrerequisites, FSystemSubsequentTasks& Subsequents)
 {
 	using namespace UE::MovieScene;
@@ -139,6 +169,7 @@ void UMovieSceneDeferredComponentMovementSystem::OnRun(FSystemTaskPrerequisites&
 	}
 	else if (CurrentPhase == ESystemPhase::Evaluation)
 	{
+		// Legacy back compat
 		FEntityTaskBuilder()
 		.Read(BuiltInComponents->BoundObject)
 		.FilterAny({ Components->ComponentTransform.PropertyTag })
diff --git a/Engine/Source/Runtime/MovieSceneTracks/Public/Systems/MovieSceneConstraintSystem.h b/Engine/Source/Runtime/MovieSceneTracks/Public/Systems/MovieSceneConstraintSystem.h
--- a/Engine/Source/Runtime/MovieSceneTracks/Public/Systems/MovieSceneConstraintSystem.h
+++ b/Engine/Source/Runtime/MovieSceneTracks/Public/Systems/MovieSceneConstraintSystem.h
@@ -8,6 +8,11 @@
 class UTransformableComponentHandle;
 class UTickableTransformConstraint;
 
+namespace UE::MovieScene
+{
+	struct FEvaluateConstraintChannels;
+}
+
 /**
  * System that is responsible for propagating constraints to a bound object's FConstraintsManagerController.
  */
@@ -25,9 +30,11 @@ public:
 	};
 
 	UMovieSceneConstraintSystem(const FObjectInitializer& ObjInit);
+	virtual void OnSchedulePersistentTasks(UE::MovieScene::IEntitySystemScheduler* TaskScheduler) override;
 	virtual void OnRun(FSystemTaskPrerequisites& InPrerequisites, FSystemSubsequentTasks& Subsequents) override;
 
 protected:
-	TArray<FUpdateHandleForConstraint>  DynamicOffsets;
+	friend UE::MovieScene::FEvaluateConstraintChannels;
 
+	TArray<FUpdateHandleForConstraint>  DynamicOffsets;
 };
diff --git a/Engine/Source/Runtime/MovieSceneTracks/Public/Systems/MovieSceneDeferredComponentMovementSystem.h b/Engine/Source/Runtime/MovieSceneTracks/Public/Systems/MovieSceneDeferredComponentMovementSystem.h
--- a/Engine/Source/Runtime/MovieSceneTracks/Public/Systems/MovieSceneDeferredComponentMovementSystem.h
+++ b/Engine/Source/Runtime/MovieSceneTracks/Public/Systems/MovieSceneDeferredComponentMovementSystem.h
@@ -45,6 +45,7 @@ private:
 	virtual void BeginDestroy() override;
 
 	virtual bool IsRelevantImpl(UMovieSceneEntitySystemLinker* InLinker) const override;
+	virtual void OnSchedulePersistentTasks(UE::MovieScene::IEntitySystemScheduler* TaskScheduler) override;
 	virtual void OnRun(FSystemTaskPrerequisites& InPrerequisites, FSystemSubsequentTasks& Subsequents) override;
 	virtual void OnUnlink() override final;
 
diff --git a/Engine/Source/Runtime/UMG/Private/Animation/MovieSceneWidgetMaterialSystem.cpp b/Engine/Source/Runtime/UMG/Private/Animation/MovieSceneWidgetMaterialSystem.cpp
--- a/Engine/Source/Runtime/UMG/Private/Animation/MovieSceneWidgetMaterialSystem.cpp
+++ b/Engine/Source/Runtime/UMG/Private/Animation/MovieSceneWidgetMaterialSystem.cpp
@@ -103,7 +103,7 @@ UMovieSceneWidgetMaterialSystem::UMovieSceneWidgetMaterialSystem(const FObjectIn
 	FMovieSceneTracksComponentTypes* TracksComponents  = FMovieSceneTracksComponentTypes::Get();
 
 	RelevantComponent = WidgetComponents->WidgetMaterialHandle;
-	Phase = ESystemPhase::Instantiation | ESystemPhase::Evaluation;
+	Phase = ESystemPhase::Instantiation;
 
 	if (HasAnyFlags(RF_ClassDefaultObject))
 	{
