Sequencer: Add or update an existing parameter key instead of always adding additional keys

diff --git a/Engine/Source/Runtime/MovieSceneTracks/Private/Sections/MovieSceneParameterSection.cpp b/Engine/Source/Runtime/MovieSceneTracks/Private/Sections/MovieSceneParameterSection.cpp
--- a/Engine/Source/Runtime/MovieSceneTracks/Private/Sections/MovieSceneParameterSection.cpp
+++ b/Engine/Source/Runtime/MovieSceneTracks/Private/Sections/MovieSceneParameterSection.cpp
@@ -274,7 +274,7 @@ void UMovieSceneParameterSection::AddScalarParameterKey( FName InParameterName,
 		ReconstructChannelProxy();
 	}
 
-	ExistingChannel->AddCubicKey(InTime, InValue);
+	ExistingChannel->GetData().UpdateOrAddKey(InTime, FMovieSceneFloatValue(InValue));
 
 	if (TryModify())
 	{
@@ -329,8 +329,8 @@ void UMovieSceneParameterSection::AddVector2DParameterKey(FName InParameterName,
 		ReconstructChannelProxy();
 	}
 
-	ExistingCurves->XCurve.AddCubicKey(InTime, InValue.X);
-	ExistingCurves->YCurve.AddCubicKey(InTime, InValue.Y);
+	ExistingCurves->XCurve.GetData().UpdateOrAddKey(InTime, FMovieSceneFloatValue(InValue.X));
+	ExistingCurves->YCurve.GetData().UpdateOrAddKey(InTime, FMovieSceneFloatValue(InValue.Y));
 
 	if (TryModify())
 	{
@@ -357,9 +357,9 @@ void UMovieSceneParameterSection::AddVectorParameterKey( FName InParameterName,
 		ReconstructChannelProxy();
 	}
 
-	ExistingCurves->XCurve.AddCubicKey(InTime, InValue.X);
-	ExistingCurves->YCurve.AddCubicKey(InTime, InValue.Y);
-	ExistingCurves->ZCurve.AddCubicKey(InTime, InValue.Z);
+	ExistingCurves->XCurve.GetData().UpdateOrAddKey(InTime, FMovieSceneFloatValue(InValue.X));
+	ExistingCurves->YCurve.GetData().UpdateOrAddKey(InTime, FMovieSceneFloatValue(InValue.Y));
+	ExistingCurves->ZCurve.GetData().UpdateOrAddKey(InTime, FMovieSceneFloatValue(InValue.Z));
 
 	if (TryModify())
 	{
@@ -385,11 +385,11 @@ void UMovieSceneParameterSection::AddColorParameterKey( FName InParameterName, F
 
 		ReconstructChannelProxy();
 	}
-
-	ExistingCurves->RedCurve.AddCubicKey(   InTime, InValue.R );
-	ExistingCurves->GreenCurve.AddCubicKey( InTime, InValue.G );
-	ExistingCurves->BlueCurve.AddCubicKey(  InTime, InValue.B );
-	ExistingCurves->AlphaCurve.AddCubicKey( InTime, InValue.A );
+	
+	ExistingCurves->RedCurve.GetData().UpdateOrAddKey(InTime, FMovieSceneFloatValue(InValue.R));
+	ExistingCurves->GreenCurve.GetData().UpdateOrAddKey(InTime, FMovieSceneFloatValue(InValue.G));
+	ExistingCurves->BlueCurve.GetData().UpdateOrAddKey(InTime, FMovieSceneFloatValue(InValue.B));
+	ExistingCurves->AlphaCurve.GetData().UpdateOrAddKey(InTime, FMovieSceneFloatValue(InValue.A));
 
 	if (TryModify())
 	{
@@ -418,17 +418,18 @@ void UMovieSceneParameterSection::AddTransformParameterKey(FName InParameterName
 	FVector Translation = InValue.GetTranslation();
 	FRotator Rotator = InValue.GetRotation().Rotator();
 	FVector Scale = InValue.GetScale3D();
-	ExistingCurves->Translation[0].AddCubicKey(InTime, Translation[0]);
-	ExistingCurves->Translation[1].AddCubicKey(InTime, Translation[1]);
-	ExistingCurves->Translation[2].AddCubicKey(InTime, Translation[2]);
 
-	ExistingCurves->Rotation[0].AddCubicKey(InTime, Rotator.Roll);
-	ExistingCurves->Rotation[1].AddCubicKey(InTime, Rotator.Pitch);
-	ExistingCurves->Rotation[2].AddCubicKey(InTime, Rotator.Yaw);
+	ExistingCurves->Translation[0].GetData().UpdateOrAddKey(InTime, FMovieSceneFloatValue(Translation[0]));
+	ExistingCurves->Translation[1].GetData().UpdateOrAddKey(InTime, FMovieSceneFloatValue(Translation[1]));
+	ExistingCurves->Translation[2].GetData().UpdateOrAddKey(InTime, FMovieSceneFloatValue(Translation[2]));
+
+	ExistingCurves->Rotation[0].GetData().UpdateOrAddKey(InTime, FMovieSceneFloatValue(Rotator.Roll));
+	ExistingCurves->Rotation[1].GetData().UpdateOrAddKey(InTime, FMovieSceneFloatValue(Rotator.Pitch));
+	ExistingCurves->Rotation[2].GetData().UpdateOrAddKey(InTime, FMovieSceneFloatValue(Rotator.Yaw));
 
-	ExistingCurves->Scale[0].AddCubicKey(InTime, Scale[0]);
-	ExistingCurves->Scale[1].AddCubicKey(InTime, Scale[1]);
-	ExistingCurves->Scale[2].AddCubicKey(InTime, Scale[2]);
+	ExistingCurves->Scale[0].GetData().UpdateOrAddKey(InTime, FMovieSceneFloatValue(Scale[0]));
+	ExistingCurves->Scale[1].GetData().UpdateOrAddKey(InTime, FMovieSceneFloatValue(Scale[1]));
+	ExistingCurves->Scale[2].GetData().UpdateOrAddKey(InTime, FMovieSceneFloatValue(Scale[2]));
 
 	if (TryModify())
 	{
