Sequencer: Ported reminaing systems to new task scheduling API

These were added on UE5/Main and so missed the original check in.

diff --git a/Engine/Source/Runtime/MovieSceneTracks/Private/Systems/MovieSceneAudioSystem.cpp b/Engine/Source/Runtime/MovieSceneTracks/Private/Systems/MovieSceneAudioSystem.cpp
--- a/Engine/Source/Runtime/MovieSceneTracks/Private/Systems/MovieSceneAudioSystem.cpp
+++ b/Engine/Source/Runtime/MovieSceneTracks/Private/Systems/MovieSceneAudioSystem.cpp
@@ -131,7 +131,7 @@ struct FGatherAudioInputs
 				double, double, double,
 				double, double, double,
 				FString, int32, bool>
-			AudioInputResults)
+			AudioInputResults) const
 	{
 		FAudioInputsBySectionKey& AudioInputsBySectionKey = AudioSystem->AudioInputsBySectionKey;
 
@@ -204,7 +204,7 @@ struct FGatherAudioTriggers
 	void ForEachAllocation(
 		const FEntityAllocation* Allocation,
 		TRead<FMovieSceneAudioComponentData> AudioDatas,
-		TRead<FName> AudioTriggerNames)
+		TRead<FName> AudioTriggerNames) const
 	{
 		FAudioInputsBySectionKey& AudioInputsBySectionKey = AudioSystem->AudioInputsBySectionKey;
 
@@ -269,7 +269,7 @@ struct FEvaluateAudio
 			TRead<FMovieSceneAudioComponentData> AudioDatas,
 			TRead<double> VolumeMultipliers,
 			TRead<double> PitchMultipliers,
-			TReadOptional<UObject*> BoundObjects)
+			TReadOptional<UObject*> BoundObjects) const
 	{
 		const FBuiltInComponentTypes* BuiltInComponents = FBuiltInComponentTypes::Get();
 
@@ -303,7 +303,7 @@ private:
 			double VolumeMultiplier,
 			double PitchMultiplier,
 			UObject* BoundObject,
-			bool bWantsRestoreState)
+			bool bWantsRestoreState) const
 	{
 		const FMovieSceneContext& Context = Instance.GetContext();
 		IMovieScenePlayer* Player = Instance.GetPlayer();
@@ -680,6 +680,7 @@ UMovieSceneAudioSystem::UMovieSceneAudioSystem(const FObjectInitializer& ObjInit
 	using namespace UE::MovieScene;
 
 	RelevantComponent = FMovieSceneTracksComponentTypes::Get()->Audio;
+	Phase = ESystemPhase::Scheduling;
 
 	if (HasAnyFlags(RF_ClassDefaultObject))
 	{
@@ -731,10 +732,20 @@ void UMovieSceneAudioSystem::OnUnlink()
 
 }
 
-void UMovieSceneAudioSystem::OnRun(FSystemTaskPrerequisites& InPrerequisites, FSystemSubsequentTasks& Subsequents)
+void UMovieSceneAudioSystem::ResetSharedData()
 {
-	MOVIESCENE_DETAILED_SCOPE_CYCLE_COUNTER(MovieSceneEval_AudioTrack_Evaluate)
+	AudioInputsBySectionKey.Reset();
+	for (TPair<FObjectKey, FAudioComponentBySectionKey>& AudioComponentsForActor : AudioComponentsByActorKey)
+	{
+		for (TPair<FObjectKey, FAudioComponentEvaluationData>& AudioComponentForSection : AudioComponentsForActor.Value)
+		{
+			AudioComponentForSection.Value.bEvaluatedThisFrame = false;
+		}
+	}
+}
 
+void UMovieSceneAudioSystem::OnSchedulePersistentTasks(UE::MovieScene::IEntitySystemScheduler* TaskScheduler)
+{
 	if (!GEngine || !GEngine->UseSound())
 	{
 		return;
@@ -746,15 +757,64 @@ void UMovieSceneAudioSystem::OnRun(FSystemTaskPrerequisites& InPrerequisites, FS
 	const FMovieSceneTracksComponentTypes* TrackComponents = FMovieSceneTracksComponentTypes::Get();
 
 	// Reset shared data.
-	AudioInputsBySectionKey.Reset();
-	for (TPair<FObjectKey, FAudioComponentBySectionKey>& AudioComponentsForActor : AudioComponentsByActorKey)
+	FTaskID ResetSharedDataTask = TaskScheduler->AddMemberFunctionTask(FTaskParams(TEXT("Reset Audio Data")), this, &UMovieSceneAudioSystem::ResetSharedData);
+
+	// Gather audio input values computed by the channel evaluators.
+	FTaskID GatherInputsTask = FEntityTaskBuilder()
+	.Read(TrackComponents->Audio)
+	.Read(TrackComponents->AudioInputs)
+	.ReadOneOrMoreOf(
+			BuiltInComponents->DoubleResult[0], BuiltInComponents->DoubleResult[1], BuiltInComponents->DoubleResult[2], 
+			BuiltInComponents->DoubleResult[3], BuiltInComponents->DoubleResult[4], BuiltInComponents->DoubleResult[5], 
+			BuiltInComponents->DoubleResult[6], BuiltInComponents->DoubleResult[7], BuiltInComponents->DoubleResult[8], 
+			BuiltInComponents->StringResult,
+			BuiltInComponents->IntegerResult,
+			BuiltInComponents->BoolResult)
+	.Schedule_PerAllocation<FGatherAudioInputs>(&Linker->EntityManager, TaskScheduler, this);
+
+	TaskScheduler->AddPrerequisite(ResetSharedDataTask, GatherInputsTask);
+
+	// Gather up audio triggers
+	FTaskID GatherTriggersTask = FEntityTaskBuilder()
+	.Read(TrackComponents->Audio)
+	.Read(TrackComponents->AudioTriggerName)
+	.Schedule_PerAllocation<FGatherAudioTriggers>(&Linker->EntityManager, TaskScheduler, this);
+
+	TaskScheduler->AddPrerequisite(ResetSharedDataTask, GatherTriggersTask);
+
+	// Next, evaluate audio to play and use the gathered audio input values to set on the audio components.
+	FTaskID EvaluateAudioTask = FEntityTaskBuilder()
+	.ReadEntityIDs()
+	.Read(BuiltInComponents->RootInstanceHandle)
+	.Read(BuiltInComponents->InstanceHandle)
+	.Read(TrackComponents->Audio)
+	.Read(BuiltInComponents->DoubleResult[0]) // Volume
+	.Read(BuiltInComponents->DoubleResult[1]) // Pitch multiplier
+	.ReadOptional(BuiltInComponents->BoundObject)
+	.SetDesiredThread(Linker->EntityManager.GetGatherThread())
+	.Schedule_PerAllocation<FEvaluateAudio>(&Linker->EntityManager, TaskScheduler, this);
+
+	TaskScheduler->AddPrerequisite(GatherInputsTask, EvaluateAudioTask);
+	TaskScheduler->AddPrerequisite(GatherTriggersTask, EvaluateAudioTask);
+}
+
+void UMovieSceneAudioSystem::OnRun(FSystemTaskPrerequisites& InPrerequisites, FSystemSubsequentTasks& Subsequents)
+{
+	MOVIESCENE_DETAILED_SCOPE_CYCLE_COUNTER(MovieSceneEval_AudioTrack_Evaluate)
+
+	if (!GEngine || !GEngine->UseSound())
 	{
-		for (TPair<FObjectKey, FAudioComponentEvaluationData>& AudioComponentForSection : AudioComponentsForActor.Value)
-		{
-			AudioComponentForSection.Value.bEvaluatedThisFrame = false;
-		}
+		return;
 	}
 
+	using namespace UE::MovieScene;
+
+	const FBuiltInComponentTypes* BuiltInComponents = FBuiltInComponentTypes::Get();
+	const FMovieSceneTracksComponentTypes* TrackComponents = FMovieSceneTracksComponentTypes::Get();
+
+	// Reset shared data.
+	ResetSharedData();
+
 	// Gather audio input values computed by the channel evaluators.
 	FSystemTaskPrerequisites Prereqs;
 
diff --git a/Engine/Source/Runtime/MovieSceneTracks/Private/Systems/MovieSceneAudioSystem.h b/Engine/Source/Runtime/MovieSceneTracks/Private/Systems/MovieSceneAudioSystem.h
--- a/Engine/Source/Runtime/MovieSceneTracks/Private/Systems/MovieSceneAudioSystem.h
+++ b/Engine/Source/Runtime/MovieSceneTracks/Private/Systems/MovieSceneAudioSystem.h
@@ -74,6 +74,7 @@ public:
 	//~ UMovieSceneEntitySystem members
 	virtual void OnLink() override;
 	virtual void OnUnlink() override;
+	virtual void OnSchedulePersistentTasks(UE::MovieScene::IEntitySystemScheduler* TaskScheduler) override;
 	virtual void OnRun(FSystemTaskPrerequisites& InPrerequisites, FSystemSubsequentTasks& Subsequents) override;
 
 	/**
@@ -98,6 +99,11 @@ public:
 	 */
 	void StopSound(FObjectKey ActorKey, FObjectKey SectionKey);
 
+	/**
+	 * Reset shared accumulation data required every evaluation frame
+	 */
+	void ResetSharedData();
+
 private:
 
 	/** Map of all created audio components */
diff --git a/Engine/Source/Runtime/MovieSceneTracks/Private/Systems/StringChannelEvaluatorSystem.cpp b/Engine/Source/Runtime/MovieSceneTracks/Private/Systems/StringChannelEvaluatorSystem.cpp
--- a/Engine/Source/Runtime/MovieSceneTracks/Private/Systems/StringChannelEvaluatorSystem.cpp
+++ b/Engine/Source/Runtime/MovieSceneTracks/Private/Systems/StringChannelEvaluatorSystem.cpp
@@ -12,6 +12,27 @@
 
 DECLARE_CYCLE_STAT(TEXT("MovieScene: Evaluate string channels"), MovieSceneEval_EvaluateStringChannelTask, STATGROUP_MovieSceneECS);
 
+namespace UE::MovieScene
+{
+
+struct FEvaluateStringChannels
+{
+	static void ForEachEntity(FSourceStringChannel StringChannel, FFrameTime FrameTime, FString& OutResult)
+	{
+		if (const FString* Value = StringChannel.Source->Evaluate(FrameTime))
+		{
+			OutResult = *Value;
+		}
+		else
+		{
+			OutResult = FString();
+		}
+	}
+};
+
+
+} // namespace UE::MovieScene
+
 UStringChannelEvaluatorSystem::UStringChannelEvaluatorSystem(const FObjectInitializer& ObjInit)
 	: Super(ObjInit)
 {
@@ -19,6 +40,7 @@ UStringChannelEvaluatorSystem::UStringChannelEvaluatorSystem(const FObjectInitia
 
 	SystemCategories = EEntitySystemCategory::ChannelEvaluators;
 	RelevantComponent = FBuiltInComponentTypes::Get()->StringChannel;
+	Phase = ESystemPhase::Scheduling;
 
 	if (HasAnyFlags(RF_ClassDefaultObject))
 	{
@@ -26,26 +48,27 @@ UStringChannelEvaluatorSystem::UStringChannelEvaluatorSystem(const FObjectInitia
 	}
 }
 
-void UStringChannelEvaluatorSystem::OnRun(FSystemTaskPrerequisites& InPrerequisites, FSystemSubsequentTasks& Subsequents)
+
+void UStringChannelEvaluatorSystem::OnSchedulePersistentTasks(UE::MovieScene::IEntitySystemScheduler* TaskScheduler)
 {
 	using namespace UE::MovieScene;
 
 	FBuiltInComponentTypes* BuiltInComponents = FBuiltInComponentTypes::Get();
 
-	struct FEvaluateStringChannels
-	{
-		static void ForEachEntity(FSourceStringChannel StringChannel, FFrameTime FrameTime, FString& OutResult)
-		{
-			if (const FString* Value = StringChannel.Source->Evaluate(FrameTime))
-			{
-				OutResult = *Value;
-			}
-			else
-			{
-				OutResult = FString();
-			}
-		}
-	};
+	FEntityTaskBuilder()
+	.Read(BuiltInComponents->StringChannel)
+	.Read(BuiltInComponents->EvalTime)
+	.Write(BuiltInComponents->StringResult)
+	.FilterNone({ BuiltInComponents->Tags.Ignored })
+	.SetStat(GET_STATID(MovieSceneEval_EvaluateStringChannelTask))
+	.Fork_PerEntity<FEvaluateStringChannels>(&Linker->EntityManager, TaskScheduler);
+}
+
+void UStringChannelEvaluatorSystem::OnRun(FSystemTaskPrerequisites& InPrerequisites, FSystemSubsequentTasks& Subsequents)
+{
+	using namespace UE::MovieScene;
+
+	FBuiltInComponentTypes* BuiltInComponents = FBuiltInComponentTypes::Get();
 
 	FEntityTaskBuilder()
 	.Read(BuiltInComponents->StringChannel)
diff --git a/Engine/Source/Runtime/MovieSceneTracks/Public/Systems/StringChannelEvaluatorSystem.h b/Engine/Source/Runtime/MovieSceneTracks/Public/Systems/StringChannelEvaluatorSystem.h
--- a/Engine/Source/Runtime/MovieSceneTracks/Public/Systems/StringChannelEvaluatorSystem.h
+++ b/Engine/Source/Runtime/MovieSceneTracks/Public/Systems/StringChannelEvaluatorSystem.h
@@ -22,5 +22,6 @@ public:
 
 	UStringChannelEvaluatorSystem(const FObjectInitializer& ObjInit);
 
+	virtual void OnSchedulePersistentTasks(UE::MovieScene::IEntitySystemScheduler* TaskScheduler) override;
 	virtual void OnRun(FSystemTaskPrerequisites& InPrerequisites, FSystemSubsequentTasks& Subsequents) override;
 };
