Sequencer: Bitfield boolean properties are now able to use fast property offsets

Bitfields were previously not supported because there was no exposed way of determining how an FBoolProperty is mapped in memory. Using fast pointer offsets would have previously caused us to always stomp a full byte over the address. Since we now support meta-data on property components (which is already used by colors and vectors), we can use the FBoolProperty when resolving the property to detect the bitoffset if FBoolProperty::IsNativeBool is false. (see usage of CountTrailingZeros64 inside FBoolHandler)

For my test this pulls down the cost of assigning boolean flags for PostProcessingOverrides from 15us to 1.7us.

Other miscellaneous changes included:
  - Added a way for traits to tell the component handlers that they are not composite property types. This causes them to use direct property setters rather than requiring a composite temporary, and prevents the handler from instantiating partial property setter code.
  - Moved custom accessor functions into their own file. Added one for ULightComponent::VolumetricScatteringIntensity.

diff --git a/Engine/Source/Runtime/CoreUObject/Public/UObject/UnrealType.h b/Engine/Source/Runtime/CoreUObject/Public/UObject/UnrealType.h
--- a/Engine/Source/Runtime/CoreUObject/Public/UObject/UnrealType.h
+++ b/Engine/Source/Runtime/CoreUObject/Public/UObject/UnrealType.h
@@ -2465,6 +2465,18 @@ public:
 		return FieldMask == 0xff;
 	}
 
+	/** Return the the mask that defines the relevant bit for this boolean, or 0xFF if IsNativeBool() is true */
+	FORCEINLINE uint8 GetFieldMask() const
+	{
+		return FieldMask;
+	}
+
+	/** Return the byte offset from this property's storage type to the byte that FieldMask applies to. Only valid if IsNativeBool() is true. */
+	FORCEINLINE uint8 GetByteOffset() const
+	{
+		return ByteOffset;
+	}
+
 	uint32 GetValueTypeHashInternal(const void* Src) const override;
 };
 
diff --git a/Engine/Source/Runtime/MovieScene/Private/EntitySystem/MovieScenePropertyRegistry.cpp b/Engine/Source/Runtime/MovieScene/Private/EntitySystem/MovieScenePropertyRegistry.cpp
--- a/Engine/Source/Runtime/MovieScene/Private/EntitySystem/MovieScenePropertyRegistry.cpp
+++ b/Engine/Source/Runtime/MovieScene/Private/EntitySystem/MovieScenePropertyRegistry.cpp
@@ -13,33 +13,131 @@ namespace UE
 namespace MovieScene
 {
 
+struct FPropertyAndAddress
+{
+	FProperty* Property = nullptr;
+	uint8* PropertyAddress = nullptr;
+};
+
+/**
+ * Find a nested property and address from the specified path of the form StructProp.Struct.Inner.
+ * Any symbols other than identifiers and dots are not supported (or even checked for).
+ */
+template<typename CharType>
+FPropertyAndAddress FindPropertyFromNestedPath(UStruct* Struct, void* Container, TStringView<CharType> InPath)
+{
+	// Find the next dot in the path
+	int32 DotIndex = InPath.Len();
+	const bool bHasTail = InPath.FindChar('.', DotIndex);
+
+	// Look up our property name using FNAME_Find. If no name exists with this name, a property cannot exist for it.
+	TStringView<CharType> Head = bHasTail ? InPath.Left(DotIndex) : InPath;
+	FName PropertyName(Head.Len(), Head.GetData(), FNAME_Find);
+	if (PropertyName.IsNone())
+	{
+		return FPropertyAndAddress{};
+	}
+
+	// Find a property for the head
+	FProperty* Property = Struct->FindPropertyByName(PropertyName);
+	if (!Property)
+	{
+		return FPropertyAndAddress{};
+	}
+
+	uint8* PropertyAddress = Property->ContainerPtrToValuePtr<uint8>(Container);
+
+	// If we have more properties to look up, start again from the tail path and the current property address
+	if (bHasTail)
+	{
+		FStructProperty* StructProperty = CastField<FStructProperty>(Property);
+		if (!StructProperty || !StructProperty->Struct)
+		{
+			return FPropertyAndAddress{};
+		}
+
+		TStringView<CharType> Tail = InPath.RightChop(DotIndex+1);
+		return FindPropertyFromNestedPath(StructProperty->Struct, PropertyAddress, Tail);
+	}
+
+	return FPropertyAndAddress{ Property, PropertyAddress };
+}
+
+FPropertyAndAddress FindPropertyFromPath(UClass* ObjectClass, const FMovieScenePropertyBinding& PropertyBinding)
+{
+	if (PropertyBinding.PropertyName == PropertyBinding.PropertyPath)
+	{
+		// Simple case - just a property name
+		FProperty* Property = ObjectClass->FindPropertyByName(PropertyBinding.PropertyName);
+		if (Property)
+		{
+			UObject* DefaultObject   = ObjectClass->GetDefaultObject();
+			uint8*   PropertyAddress = Property->ContainerPtrToValuePtr<uint8>(DefaultObject);
+			return FPropertyAndAddress{ Property, PropertyAddress };
+		}
+	}
+	else
+	{
+		// Lookup the property name using a string view to avoid allocation.
+		// Most properties will be ansi strings
+		const FNameEntry* NameEntry = PropertyBinding.PropertyPath.GetComparisonNameEntry();
+		if (NameEntry)
+		{
+			if (NameEntry->IsWide())
+			{
+				TWideStringBuilder<128> Path;
+				NameEntry->AppendNameToString(Path);
+				return FindPropertyFromNestedPath(ObjectClass, ObjectClass->GetDefaultObject(), Path.ToView());
+			}
+			else
+			{
+				TAnsiStringBuilder<128> Path;
+				NameEntry->AppendAnsiNameToString(Path);
+				return FindPropertyFromNestedPath(ObjectClass, ObjectClass->GetDefaultObject(), Path.ToView());
+			}
+		}
+	}
+
+	return FPropertyAndAddress{};
+}
+
+
+
 TOptional<uint16> ComputeFastPropertyPtrOffset(UClass* ObjectClass, const FMovieScenePropertyBinding& PropertyBinding)
 {
 	using namespace UE::MovieScene;
 
-	FProperty* Property = ObjectClass->FindPropertyByName(PropertyBinding.PropertyName);
-	
-	const bool bFoundSetter = (!Property || Property->HasSetter());
+	FPropertyAndAddress PropertyAndAddress = FindPropertyFromPath(ObjectClass, PropertyBinding);
+
+	const bool bFoundSetter = (!PropertyAndAddress.Property || PropertyAndAddress.Property->HasSetter());
 	if (bFoundSetter)
 	{
 		return TOptional<uint16>();
 	}
 
+	if (!PropertyAndAddress.Property)
+	{
+		return TOptional<uint16>();
+	}
+
 	// @todo: Constructing FNames from strings is _very_ costly and we really shouldn't be doing this at runtime.
-	UFunction* Setter   = ObjectClass->FindFunctionByName(*(FString("Set") + PropertyBinding.PropertyName.ToString()));
-	if (Property && !Setter)
+	//        This is a little better now we use a string builder and an FNAME_Find, but it's still not ideal
+	TStringBuilder<128> SetterName;
+	SetterName.Append(TEXT("Set"));
+	PropertyBinding.PropertyName.ToString(SetterName);
+
+	FName SetterFunctionName(SetterName.ToString(), FNAME_Find);
+	if (SetterFunctionName.IsNone() || ObjectClass->FindFunctionByName(SetterFunctionName) == nullptr)
 	{
-		// @todo: We disable fast offsets for all bools atm becasue there is no way of knowing whether the property
-		// represents a raw bool, or a bool : 1 or uint8 : 1;
-		FBoolProperty* BoolProperty = CastField<FBoolProperty>(Property);
-		if (BoolProperty)
+		if (FBoolProperty* BoolProperty = CastField<FBoolProperty>(PropertyAndAddress.Property))
 		{
-			return TOptional<uint16>();
+			// bitfield booleans potentially have an additional byte offset.
+			// In practice this is always 0 because the property internal offset itself is incremented, but this is here for completeness
+			PropertyAndAddress.PropertyAddress += BoolProperty->GetByteOffset();
 		}
 
 		UObject* DefaultObject   = ObjectClass->GetDefaultObject();
-		uint8*   PropertyAddress = Property->ContainerPtrToValuePtr<uint8>(DefaultObject);
-		int32    PropertyOffset  = PropertyAddress - reinterpret_cast<uint8*>(DefaultObject);
+		int32    PropertyOffset  = PropertyAndAddress.PropertyAddress - reinterpret_cast<uint8*>(DefaultObject);
 
 		if (PropertyOffset >= 0 && PropertyOffset < int32(uint16(0xFFFF)))
 		{
diff --git a/Engine/Source/Runtime/MovieScene/Private/MovieScenePropertyBinding.cpp b/Engine/Source/Runtime/MovieScene/Private/MovieScenePropertyBinding.cpp
new file mode 100644
--- /dev/null
+++ b/Engine/Source/Runtime/MovieScene/Private/MovieScenePropertyBinding.cpp
@@ -0,0 +1,59 @@
+// Copyright Epic Games, Inc. All Rights Reserved.
+
+#include "EntitySystem/MovieScenePropertyBinding.h"
+#include "UObject/NameTypes.h"
+#include "Misc/StringBuilder.h"
+#include "Containers/StringView.h"
+
+FMovieScenePropertyBinding::FMovieScenePropertyBinding(FName InPropertyName, const FString& InPropertyPath)
+	: PropertyName(InPropertyName), PropertyPath(*InPropertyPath)
+{
+	bCanUseClassLookup = !(InPropertyPath.Contains(TEXT("/")) || InPropertyPath.Contains(TEXT("\\")) || InPropertyPath.Contains(TEXT("[")));
+}
+
+FMovieScenePropertyBinding FMovieScenePropertyBinding::FromPath(const FString& InPropertyPath)
+{
+	FName PropertyName;
+
+	int32 NamePos = INDEX_NONE;
+	if (InPropertyPath.FindLastChar('.', NamePos) || InPropertyPath.FindLastChar('/', NamePos) || InPropertyPath.FindLastChar('\\', NamePos))
+	{
+		PropertyName = FName(FStringView(*InPropertyPath + NamePos, InPropertyPath.Len() - NamePos));
+	}
+	else
+	{
+		PropertyName = *InPropertyPath;
+	}
+	return FMovieScenePropertyBinding(PropertyName, InPropertyPath);
+}
+
+#if WITH_EDITORONLY_DATA
+void FMovieScenePropertyBinding::PostSerialize(const FArchive& Ar)
+{
+	if (Ar.IsLoading())
+	{
+		const FNameEntry* NameEntry = PropertyPath.GetComparisonNameEntry();
+		if (NameEntry)
+		{
+			if (NameEntry->IsWide())
+			{
+				TWideStringBuilder<128> Path;
+				NameEntry->AppendNameToString(Path);
+				TStringView<WIDECHAR> PathView = Path.ToView();
+
+				int Unused = 0;
+				bCanUseClassLookup = !(PathView.FindChar('/', Unused) || PathView.FindChar('\\', Unused) || PathView.FindChar('[', Unused));
+			}
+			else
+			{
+				TAnsiStringBuilder<128> Path;
+				NameEntry->AppendAnsiNameToString(Path);
+				TStringView<ANSICHAR> PathView = Path.ToView();
+
+				int Unused  = 0;
+				bCanUseClassLookup = !(PathView.FindChar('/', Unused) || PathView.FindChar('\\', Unused) || PathView.FindChar('[', Unused));
+			}
+		}
+	}
+}
+#endif
\ No newline at end of file
diff --git a/Engine/Source/Runtime/MovieScene/Public/EntitySystem/MovieScenePropertyBinding.h b/Engine/Source/Runtime/MovieScene/Public/EntitySystem/MovieScenePropertyBinding.h
--- a/Engine/Source/Runtime/MovieScene/Public/EntitySystem/MovieScenePropertyBinding.h
+++ b/Engine/Source/Runtime/MovieScene/Public/EntitySystem/MovieScenePropertyBinding.h
@@ -21,27 +21,9 @@ struct FMovieScenePropertyBinding
 		: bCanUseClassLookup(false)
 	{}
 
-	FMovieScenePropertyBinding(FName InPropertyName, const FString& InPropertyPath)
-		: PropertyName(InPropertyName), PropertyPath(*InPropertyPath)
-	{
-		bCanUseClassLookup = !(InPropertyPath.Contains(TEXT(".")) || InPropertyPath.Contains(TEXT("/")) || InPropertyPath.Contains(TEXT("\\")) || InPropertyPath.Contains(TEXT("[")));
-	}
+	MOVIESCENE_API FMovieScenePropertyBinding(FName InPropertyName, const FString& InPropertyPath);
 
-	static FMovieScenePropertyBinding FromPath(const FString& InPropertyPath)
-	{
-		FName PropertyName;
-
-		int32 NamePos = INDEX_NONE;
-		if (InPropertyPath.FindLastChar('.', NamePos) || InPropertyPath.FindLastChar('/', NamePos) || InPropertyPath.FindLastChar('\\', NamePos))
-		{
-			PropertyName = FName(FStringView(*InPropertyPath + NamePos, InPropertyPath.Len() - NamePos));
-		}
-		else
-		{
-			PropertyName = *InPropertyPath;
-		}
-		return FMovieScenePropertyBinding(PropertyName, InPropertyPath);
-	}
+	MOVIESCENE_API static FMovieScenePropertyBinding FromPath(const FString& InPropertyPath);
 
 	friend bool operator==(FMovieScenePropertyBinding A, FMovieScenePropertyBinding B)
 	{
@@ -53,6 +35,10 @@ struct FMovieScenePropertyBinding
 		return bCanUseClassLookup;
 	}
 
+#if WITH_EDITORONLY_DATA
+	void PostSerialize(const FArchive& Ar);
+#endif
+
 	/** Leaf name of the property to animate */
 	UPROPERTY()
 	FName PropertyName;
@@ -66,6 +52,13 @@ struct FMovieScenePropertyBinding
 	bool bCanUseClassLookup;
 };
 
+#if WITH_EDITORONLY_DATA
+template<>
+struct TStructOpsTypeTraits<FMovieScenePropertyBinding> : public TStructOpsTypeTraitsBase2<FMovieScenePropertyBinding>
+{
+	enum { WithPostSerialize = true };
+};
+#endif // WITH_EDITORONLY_DATA
 
 
 #if UE_MOVIESCENE_ENTITY_DEBUG
diff --git a/Engine/Source/Runtime/MovieScene/Public/EntitySystem/MovieScenePropertyComponentHandler.h b/Engine/Source/Runtime/MovieScene/Public/EntitySystem/MovieScenePropertyComponentHandler.h
--- a/Engine/Source/Runtime/MovieScene/Public/EntitySystem/MovieScenePropertyComponentHandler.h
+++ b/Engine/Source/Runtime/MovieScene/Public/EntitySystem/MovieScenePropertyComponentHandler.h
@@ -303,21 +303,30 @@ template<typename PropertyTraits>
 struct TInitialValueProcessor : TInitialValueProcessorImpl<PropertyTraits, typename PropertyTraits::MetaDataType, TMakeIntegerSequence<int, PropertyTraits::MetaDataType::Num>>
 {};
 
+template<typename T, typename U = decltype(T::bIsComposite)>
+constexpr bool IsCompositePropertyTraits(T*)
+{
+	return T::bIsComposite;
+}
+constexpr bool IsCompositePropertyTraits(...)
+{
+	return true;
+}
 
 template<typename PropertyTraits, typename ...MetaDataTypes, int ...MetaDataIndices, typename ...CompositeTypes, int ...CompositeIndices>
 struct TPropertyComponentHandlerImpl<PropertyTraits, TPropertyMetaData<MetaDataTypes...>, TIntegerSequence<int, MetaDataIndices...>, TIntegerSequence<int, CompositeIndices...>, CompositeTypes...> : IPropertyComponentHandler
 {
+	static constexpr bool bIsComposite = IsCompositePropertyTraits((PropertyTraits*)nullptr);
+
 	using StorageType        = typename PropertyTraits::StorageType;
-	using CompleteSetterTask = TSetCompositePropertyValues<PropertyTraits, CompositeTypes...>;
-	using PartialSetterTask  = TSetPartialPropertyValues<PropertyTraits, CompositeTypes...>;
+	using CompleteSetterTask = std::conditional_t<bIsComposite, TSetCompositePropertyValues<PropertyTraits, CompositeTypes...>, TSetPropertyValues<PropertyTraits>>;
 
 	using PreAnimatedStorageType = TPreAnimatedPropertyStorage<PropertyTraits>;
 
-	TPreAnimatedStorageID<PreAnimatedStorageType> StorageID;
+	TAutoRegisterPreAnimatedStorageID<PreAnimatedStorageType> StorageID;
 
 	TPropertyComponentHandlerImpl()
 	{
-		StorageID = FPreAnimatedStateExtension::RegisterStorage<PreAnimatedStorageType>();
 	}
 
 	virtual TSharedPtr<IPreAnimatedStorage> GetPreAnimatedStateStorage(const FPropertyDefinition& Definition, FPreAnimatedStateExtension* Container) override
@@ -347,25 +356,30 @@ struct TPropertyComponentHandlerImpl<PropertyTraits, TPropertyMetaData<MetaDataT
 		.SetDesiredThread(Linker->EntityManager.GetGatherThread())
 		.template Fork_PerAllocation<CompleteSetterTask>(&Linker->EntityManager, TaskScheduler, Definition.CustomPropertyRegistration);
 
-		if (Stats.NumPartialProperties > 0)
+		if constexpr (bIsComposite)
 		{
-			FComponentMask CompletePropertyMask;
-			for (const FPropertyCompositeDefinition& Composite : Composites)
+			if (Stats.NumPartialProperties > 0)
 			{
-				CompletePropertyMask.Set(Composite.ComponentTypeID);
-			}
+				using PartialSetterTask  = TSetPartialPropertyValues<PropertyTraits, CompositeTypes...>;
 
-			FEntityTaskBuilder()
-			.Read(BuiltInComponents->BoundObject)
-			.ReadOneOf(BuiltInComponents->CustomPropertyIndex, BuiltInComponents->FastPropertyOffset, BuiltInComponents->SlowProperty)
-			.ReadAllOf(Definition.GetMetaDataComponent<MetaDataTypes>(MetaDataIndices)...)
-			.ReadAnyOf(Composites[CompositeIndices].ComponentTypeID.ReinterpretCast<CompositeTypes>()...)
-			.FilterAny({ CompletePropertyMask })
-			.FilterAll({ Definition.PropertyType })
-			.FilterOut(CompletePropertyMask)
-			.SetStat(Definition.StatID)
-			.SetDesiredThread(Linker->EntityManager.GetGatherThread())
-			.template Fork_PerAllocation<PartialSetterTask>(&Linker->EntityManager, TaskScheduler, Definition.CustomPropertyRegistration, Composites);
+				FComponentMask CompletePropertyMask;
+				for (const FPropertyCompositeDefinition& Composite : Composites)
+				{
+					CompletePropertyMask.Set(Composite.ComponentTypeID);
+				}
+
+				FEntityTaskBuilder()
+				.Read(BuiltInComponents->BoundObject)
+				.ReadOneOf(BuiltInComponents->CustomPropertyIndex, BuiltInComponents->FastPropertyOffset, BuiltInComponents->SlowProperty)
+				.ReadAllOf(Definition.GetMetaDataComponent<MetaDataTypes>(MetaDataIndices)...)
+				.ReadAnyOf(Composites[CompositeIndices].ComponentTypeID.ReinterpretCast<CompositeTypes>()...)
+				.FilterAny({ CompletePropertyMask })
+				.FilterAll({ Definition.PropertyType })
+				.FilterOut(CompletePropertyMask)
+				.SetStat(Definition.StatID)
+				.SetDesiredThread(Linker->EntityManager.GetGatherThread())
+				.template Fork_PerAllocation<PartialSetterTask>(&Linker->EntityManager, TaskScheduler, Definition.CustomPropertyRegistration, Composites);
+			}
 		}
 	}
 
@@ -383,25 +397,30 @@ struct TPropertyComponentHandlerImpl<PropertyTraits, TPropertyMetaData<MetaDataT
 		.SetDesiredThread(Linker->EntityManager.GetGatherThread())
 		.template Dispatch_PerAllocation<CompleteSetterTask>(&Linker->EntityManager, InPrerequisites, &Subsequents, Definition.CustomPropertyRegistration);
 
-		if (Stats.NumPartialProperties > 0)
+		if constexpr (bIsComposite)
 		{
-			FComponentMask CompletePropertyMask;
-			for (const FPropertyCompositeDefinition& Composite : Composites)
+			if (Stats.NumPartialProperties > 0)
 			{
-				CompletePropertyMask.Set(Composite.ComponentTypeID);
-			}
+				using PartialSetterTask  = TSetPartialPropertyValues<PropertyTraits, CompositeTypes...>;
 
-			FEntityTaskBuilder()
-			.Read(BuiltInComponents->BoundObject)
-			.ReadOneOf(BuiltInComponents->CustomPropertyIndex, BuiltInComponents->FastPropertyOffset, BuiltInComponents->SlowProperty)
-			.ReadAllOf(Definition.GetMetaDataComponent<MetaDataTypes>(MetaDataIndices)...)
-			.ReadAnyOf(Composites[CompositeIndices].ComponentTypeID.ReinterpretCast<CompositeTypes>()...)
-			.FilterAny({ CompletePropertyMask })
-			.FilterAll({ Definition.PropertyType })
-			.FilterOut(CompletePropertyMask)
-			.SetStat(Definition.StatID)
-			.SetDesiredThread(Linker->EntityManager.GetGatherThread())
-			.template Dispatch_PerAllocation<PartialSetterTask>(&Linker->EntityManager, InPrerequisites, &Subsequents, Definition.CustomPropertyRegistration, Composites);
+				FComponentMask CompletePropertyMask;
+				for (const FPropertyCompositeDefinition& Composite : Composites)
+				{
+					CompletePropertyMask.Set(Composite.ComponentTypeID);
+				}
+
+				FEntityTaskBuilder()
+				.Read(BuiltInComponents->BoundObject)
+				.ReadOneOf(BuiltInComponents->CustomPropertyIndex, BuiltInComponents->FastPropertyOffset, BuiltInComponents->SlowProperty)
+				.ReadAllOf(Definition.GetMetaDataComponent<MetaDataTypes>(MetaDataIndices)...)
+				.ReadAnyOf(Composites[CompositeIndices].ComponentTypeID.ReinterpretCast<CompositeTypes>()...)
+				.FilterAny({ CompletePropertyMask })
+				.FilterAll({ Definition.PropertyType })
+				.FilterOut(CompletePropertyMask)
+				.SetStat(Definition.StatID)
+				.SetDesiredThread(Linker->EntityManager.GetGatherThread())
+				.template Dispatch_PerAllocation<PartialSetterTask>(&Linker->EntityManager, InPrerequisites, &Subsequents, Definition.CustomPropertyRegistration, Composites);
+			}
 		}
 	}
 
diff --git a/Engine/Source/Runtime/MovieScene/Public/EntitySystem/MovieScenePropertyTraits.h b/Engine/Source/Runtime/MovieScene/Public/EntitySystem/MovieScenePropertyTraits.h
--- a/Engine/Source/Runtime/MovieScene/Public/EntitySystem/MovieScenePropertyTraits.h
+++ b/Engine/Source/Runtime/MovieScene/Public/EntitySystem/MovieScenePropertyTraits.h
@@ -25,9 +25,11 @@ struct TPropertyMetaData
 /**
  * Property accessor traits that talk directly to the reflected UObject property type
  */
-template<typename UObjectPropertyType>
+template<typename UObjectPropertyType, bool bInIsComposite = true>
 struct TDirectPropertyTraits
 {
+	static constexpr bool bIsComposite = bInIsComposite;
+
 	using StorageType  = UObjectPropertyType;
 	using MetaDataType = TPropertyMetaData<>;
 	using TraitsType   = TDirectPropertyTraits<UObjectPropertyType>;
@@ -85,9 +87,11 @@ struct TDirectPropertyTraits
 /**
  * Property accessor traits that talk directly to the reflected UObject property type
  */
-template<typename UObjectPropertyType, typename InMemoryType>
+template<typename UObjectPropertyType, typename InMemoryType, bool bInIsComposite = true>
 struct TIndirectPropertyTraits
 {
+	static constexpr bool bIsComposite = bInIsComposite;
+
 	using StorageType  = InMemoryType;
 	using MetaDataType = TPropertyMetaData<>;
 	using TraitsType   = TIndirectPropertyTraits<UObjectPropertyType, InMemoryType>;
diff --git a/Engine/Source/Runtime/MovieSceneTracks/Private/MovieSceneTracksComponentTypes.cpp b/Engine/Source/Runtime/MovieSceneTracks/Private/MovieSceneTracksComponentTypes.cpp
--- a/Engine/Source/Runtime/MovieSceneTracks/Private/MovieSceneTracksComponentTypes.cpp
+++ b/Engine/Source/Runtime/MovieSceneTracks/Private/MovieSceneTracksComponentTypes.cpp
@@ -1,12 +1,8 @@
 // Copyright Epic Games, Inc. All Rights Reserved.
 
 #include "MovieSceneTracksComponentTypes.h"
-#include "Components/ExponentialHeightFogComponent.h"
-#include "Components/LightComponent.h"
-#include "Components/PrimitiveComponent.h"
 #include "Components/SkeletalMeshComponent.h"
-#include "Components/SkyAtmosphereComponent.h"
-#include "Components/SkyLightComponent.h"
+#include "MovieSceneTracksCustomAccessors.h"
 #include "EntitySystem/BuiltInComponentTypes.h"
 #include "EntitySystem/MovieSceneEntityManager.h"
 #include "EntitySystem/MovieSceneEntitySystemLinker.h"
@@ -171,375 +167,6 @@ void ConvertOperationalProperty(double In, float& Out)
 	Out = static_cast<float>(In);
 }
 
-FIntermediate3DTransform GetComponentTransform(const UObject* Object)
-{
-	const USceneComponent* SceneComponent = CastChecked<const USceneComponent>(Object);
-	FIntermediate3DTransform Result(SceneComponent->GetRelativeLocation(), SceneComponent->GetRelativeRotation(), SceneComponent->GetRelativeScale3D());
-	return Result;
-}
-
-void SetComponentTranslationAndRotation(USceneComponent* SceneComponent, const FIntermediate3DTransform& Transform)
-{
-	// If this is a simulating component, teleport since sequencer takes over. 
-	// Teleport will not have no velocity, but it's computed later by sequencer so that it will be correct for physics.
-	// @todo: We would really rather not 
-	AActor* Actor = SceneComponent->GetOwner();
-	USceneComponent* RootComponent = Actor ? Actor->GetRootComponent() : nullptr;
-	bool bIsSimulatingPhysics = RootComponent ? RootComponent->IsSimulatingPhysics() : false;
-
-	FVector Translation = Transform.GetTranslation();
-	FRotator Rotation = Transform.GetRotation();
-	SceneComponent->SetRelativeLocationAndRotation(Translation, Rotation, false, nullptr, bIsSimulatingPhysics ? ETeleportType::ResetPhysics : ETeleportType::None);
-
-	// Force the location and rotation values to avoid Rot->Quat->Rot conversions
-	SceneComponent->SetRelativeLocation_Direct(Translation);
-	SceneComponent->SetRelativeRotation_Direct(Rotation);
-}
-
-void SetComponentTransform(USceneComponent* SceneComponent, const FIntermediate3DTransform& Transform)
-{
-	// If this is a simulating component, teleport since sequencer takes over. 
-	// Teleport will not have no velocity, but it's computed later by sequencer so that it will be correct for physics.
-	// @todo: We would really rather not 
-	AActor* Actor = SceneComponent->GetOwner();
-	USceneComponent* RootComponent = Actor ? Actor->GetRootComponent() : nullptr;
-	bool bIsSimulatingPhysics = RootComponent ? RootComponent->IsSimulatingPhysics() : false;
-
-	FVector Translation = Transform.GetTranslation();
-	FRotator Rotation = Transform.GetRotation();
-	SceneComponent->SetRelativeTransform(FTransform(Rotation, Translation, Transform.GetScale()), false, nullptr, bIsSimulatingPhysics ? ETeleportType::ResetPhysics : ETeleportType::None);
-
-	// Force the location and rotation values to avoid Rot->Quat->Rot conversions
-	SceneComponent->SetRelativeLocation_Direct(Translation);
-	SceneComponent->SetRelativeRotation_Direct(Rotation);
-}
-
-void SetComponentTransformAndVelocity(UObject* Object, const FIntermediate3DTransform& InTransform)
-{
-	UE::MovieScene::FIntermediate3DTransform::ApplyTransformTo(CastChecked<USceneComponent>(Object), InTransform);
-}
-
-FIntermediateColor GetLightComponentLightColor(const UObject* Object, EColorPropertyType InColorType)
-{
-	ensure(InColorType == EColorPropertyType::Color);
-
-	const ULightComponent* LightComponent = CastChecked<const ULightComponent>(Object);
-	return FIntermediateColor(LightComponent->GetLightColor());
-}
-
-void SetLightComponentLightColor(UObject* Object, EColorPropertyType InColorType, const FIntermediateColor& InColor)
-{
-	// This is a little esoteric - ULightComponentBase::LightColor is the UPROPERTY that generates the meta-data
-	// for this custom callback, but it is an FColor, even though the public get/set functions expose it as an
-	// FLinearColor. FIntermediateColor is always blended and dealt with in linear space, so it's fine to 
-	// simply reinterpret the color
-	ensure(InColorType == EColorPropertyType::Color);
-
-	const bool bConvertBackToSRgb = true;
-	ULightComponent* LightComponent = CastChecked<ULightComponent>(Object);
-	LightComponent->SetLightColor(InColor.GetLinearColor(), bConvertBackToSRgb);
-}
-
-float GetLightComponentIntensity(const UObject* Object)
-{
-	const ULightComponent* LightComponent = CastChecked<const ULightComponent>(Object);
-	return LightComponent->Intensity;
-}
-
-void SetLightComponentIntensity(UObject* Object, float InIntensity)
-{
-	ULightComponent* LightComponent = CastChecked<ULightComponent>(Object);
-	LightComponent->SetIntensity(InIntensity);
-}
-
-FIntermediateColor GetSkyLightComponentLightColor(const UObject* Object, EColorPropertyType InColorType)
-{
-	ensure(InColorType == EColorPropertyType::Color);
-
-	const USkyLightComponent* SkyLightComponent = CastChecked<const USkyLightComponent>(Object);
-	return FIntermediateColor(SkyLightComponent->GetLightColor());
-}
-
-void SetSkyLightComponentLightColor(UObject* Object, EColorPropertyType InColorType, const FIntermediateColor& InColor)
-{
-	// This is a little esoteric - ULightComponentBase::LightColor is the UPROPERTY that generates the meta-data
-	// for this custom callback, but it is an FColor, even though the public get/set functions expose it as an
-	// FLinearColor. FIntermediateColor is always blended and dealt with in linear space, so it's fine to 
-	// simply reinterpret the color
-	ensure(InColorType == EColorPropertyType::Color);
-
-	USkyLightComponent* SkyLightComponent = CastChecked<USkyLightComponent>(Object);
-	SkyLightComponent->SetLightColor(InColor.GetLinearColor());
-}
-
-float GetSkyLightComponentIntensity(const UObject* Object)
-{
-	const USkyLightComponent* SkyLightComponent = CastChecked<const USkyLightComponent>(Object);
-	return SkyLightComponent->Intensity;
-}
-
-void SetSkyLightComponentIntensity(UObject* Object, float InIntensity)
-{
-	USkyLightComponent* SkyLightComponent = CastChecked<USkyLightComponent>(Object);
-	SkyLightComponent->SetIntensity(InIntensity);
-}
-
-float GetSkyAtmosphereComponentMieScatteringScale(const UObject* Object)
-{
-	const USkyAtmosphereComponent* SkyAtmosphereComponent = CastChecked<const USkyAtmosphereComponent>(Object);
-	return SkyAtmosphereComponent->MieScatteringScale;
-}
-
-void SetSkyAtmosphereComponentMieScatteringScale(UObject* Object, float InMieScatteringScale)
-{
-	USkyAtmosphereComponent* SkyAtmosphereComponent = CastChecked<USkyAtmosphereComponent>(Object);
-	SkyAtmosphereComponent->SetMieScatteringScale(InMieScatteringScale);
-}
-
-float GetSkyAtmosphereComponentOtherAbsorptionScale(const UObject* Object)
-{
-	const USkyAtmosphereComponent* SkyAtmosphereComponent = CastChecked<const USkyAtmosphereComponent>(Object);
-	return SkyAtmosphereComponent->OtherAbsorptionScale;
-}
-
-void SetSkyAtmosphereComponentOtherAbsorptionScale(UObject* Object, float InOtherAbsorptionScale)
-{
-	USkyAtmosphereComponent* SkyAtmosphereComponent = CastChecked<USkyAtmosphereComponent>(Object);
-	SkyAtmosphereComponent->SetOtherAbsorptionScale(InOtherAbsorptionScale);
-}
-
-float GetSkyAtmosphereComponentRayleighScatteringScale(const UObject* Object)
-{
-	const USkyAtmosphereComponent* SkyAtmosphereComponent = CastChecked<const USkyAtmosphereComponent>(Object);
-	return SkyAtmosphereComponent->RayleighScatteringScale;
-}
-
-void SetSkyAtmosphereComponentRayleighScatteringScale(UObject* Object, float InRayleighScatteringScale)
-{
-	USkyAtmosphereComponent* SkyAtmosphereComponent = CastChecked<USkyAtmosphereComponent>(Object);
-	SkyAtmosphereComponent->SetRayleighScatteringScale(InRayleighScatteringScale);
-}
-
-float GetSkyAtmosphereComponentRayleighExponentialDistribution(const UObject* Object)
-{
-	const USkyAtmosphereComponent* SkyAtmosphereComponent = CastChecked<const USkyAtmosphereComponent>(Object);
-	return SkyAtmosphereComponent->RayleighExponentialDistribution;
-}
-
-void SetSkyAtmosphereComponentRayleighExponentialDistribution(UObject* Object, float InRayleighExponentialDistribution)
-{
-	USkyAtmosphereComponent* SkyAtmosphereComponent = CastChecked<USkyAtmosphereComponent>(Object);
-	SkyAtmosphereComponent->SetRayleighExponentialDistribution(InRayleighExponentialDistribution);
-}
-
-float GetSkyAtmosphereComponentAerialPerspectiveViewDistanceScale(const UObject* Object)
-{
-	const USkyAtmosphereComponent* SkyAtmosphereComponent = CastChecked<const USkyAtmosphereComponent>(Object);
-	return SkyAtmosphereComponent->AerialPespectiveViewDistanceScale;
-}
-
-void SetSkyAtmosphereComponentAerialPerspectiveViewDistanceScale(UObject* Object, float InAerialPerspectiveViewDistanceScale)
-{
-	USkyAtmosphereComponent* SkyAtmosphereComponent = CastChecked<USkyAtmosphereComponent>(Object);
-	SkyAtmosphereComponent->SetAerialPespectiveViewDistanceScale(InAerialPerspectiveViewDistanceScale);
-}
-
-FIntermediateColor GetSkyAtmosphereComponentMieAbsorption(const UObject* Object, EColorPropertyType InColorType)
-{
-	// USkyAtmosphereComponent::MieAbsorption is a FLinearColor. FIntermediateColor is always blended and dealt with
-	// in linear space so it should be safe to reinterpret it as linear.
-	ensure(InColorType == EColorPropertyType::Linear);
-
-	const USkyAtmosphereComponent* SkyAtmosphereComponent = CastChecked<const USkyAtmosphereComponent>(Object);
-	return FIntermediateColor(SkyAtmosphereComponent->MieAbsorption);
-}
-
-void SetSkyAtmosphereComponentMieAbsorption(UObject* Object, EColorPropertyType InColorType, const FIntermediateColor& InColor)
-{
-	// USkyAtmosphereComponent::MieAbsorption is a FLinearColor. FIntermediateColor is always blended and dealt with
-	// in linear space so it should be safe to reinterpret it as linear.
-	ensure(InColorType == EColorPropertyType::Linear);
-
-	USkyAtmosphereComponent* SkyAtmosphereComponent = CastChecked<USkyAtmosphereComponent>(Object);
-	SkyAtmosphereComponent->SetMieAbsorption(InColor.GetLinearColor());
-}
-
-FIntermediateColor GetSkyAtmosphereComponentMieScattering(const UObject* Object, EColorPropertyType InColorType)
-{
-	// USkyAtmosphereComponent::MieScattering is a FLinearColor. FIntermediateColor is always blended and dealt with
-	// in linear space so it should be safe to reinterpret it as linear.
-	ensure(InColorType == EColorPropertyType::Linear);
-
-	const USkyAtmosphereComponent* SkyAtmosphereComponent = CastChecked<const USkyAtmosphereComponent>(Object);
-	return FIntermediateColor(SkyAtmosphereComponent->MieScattering);
-}
-
-void SetSkyAtmosphereComponentMieScattering(UObject* Object, EColorPropertyType InColorType, const FIntermediateColor& InColor)
-{
-	// USkyAtmosphereComponent::MieScattering is a FLinearColor. FIntermediateColor is always blended and dealt with
-	// in linear space so it should be safe to reinterpret it as linear.
-	ensure(InColorType == EColorPropertyType::Linear);
-
-	USkyAtmosphereComponent* SkyAtmosphereComponent = CastChecked<USkyAtmosphereComponent>(Object);
-	SkyAtmosphereComponent->SetMieScattering(InColor.GetLinearColor());
-}
-
-FIntermediateColor GetSkyAtmosphereComponentRayleighScattering(const UObject* Object, EColorPropertyType InColorType)
-{
-	// USkyAtmosphereComponent::RayleighScattering is a FLinearColor. FIntermediateColor is always blended and dealt with
-	// in linear space so it should be safe to reinterpret it as linear.
-	ensure(InColorType == EColorPropertyType::Linear);
-
-	const USkyAtmosphereComponent* SkyAtmosphereComponent = CastChecked<const USkyAtmosphereComponent>(Object);
-	return FIntermediateColor(SkyAtmosphereComponent->RayleighScattering);
-}
-
-void SetSkyAtmosphereComponentRayleighScattering(UObject* Object, EColorPropertyType InColorType, const FIntermediateColor& InColor)
-{
-	// USkyAtmosphereComponent::RayleighScattering is a FLinearColor. FIntermediateColor is always blended and dealt with
-	// in linear space so it should be safe to reinterpret it as linear.
-	ensure(InColorType == EColorPropertyType::Linear);
-
-	USkyAtmosphereComponent* SkyAtmosphereComponent = CastChecked<USkyAtmosphereComponent>(Object);
-	SkyAtmosphereComponent->SetRayleighScattering(InColor.GetLinearColor());
-}
-
-FIntermediateColor GetSkyAtmosphereComponentSkyLuminanceFactor(const UObject* Object, EColorPropertyType InColorType)
-{
-	// USkyAtmosphereComponent::SkyLuminanceFactor is a FLinearColor. FIntermediateColor is always blended and dealt with
-	// in linear space so it should be safe to reinterpret it as linear.
-	ensure(InColorType == EColorPropertyType::Linear);
-
-	const USkyAtmosphereComponent* SkyAtmosphereComponent = CastChecked<const USkyAtmosphereComponent>(Object);
-	return FIntermediateColor(SkyAtmosphereComponent->SkyLuminanceFactor);
-}
-
-void SetSkyAtmosphereComponentSkyLuminanceFactor(UObject* Object, EColorPropertyType InColorType, const FIntermediateColor& InColor)
-{
-	// USkyAtmosphereComponent::SkyLuminanceFactor is a FLinearColor. FIntermediateColor is always blended and dealt with
-	// in linear space so it should be safe to reinterpret it as linear.
-	ensure(InColorType == EColorPropertyType::Linear);
-
-	USkyAtmosphereComponent* SkyAtmosphereComponent = CastChecked<USkyAtmosphereComponent>(Object);
-	SkyAtmosphereComponent->SetSkyLuminanceFactor(InColor.GetLinearColor());
-}
-
-float GetMassScale(const UObject* Object)
-{
-	const UPrimitiveComponent* PrimitiveComponent = CastChecked<const UPrimitiveComponent>(Object);
-	return PrimitiveComponent->GetMassScale();
-}
-
-void SetMassScale(UObject* Object, float InMassScale)
-{
-	UPrimitiveComponent* PrimitiveComponent = CastChecked<UPrimitiveComponent>(Object);
-	PrimitiveComponent->SetMassScale(NAME_None, InMassScale);
-}
-
-float GetExponentialHeightFogComponentFogDensity(const UObject* Object)
-{
-	const UExponentialHeightFogComponent* ExponentialHeightFogComponent = CastChecked<const UExponentialHeightFogComponent>(Object);
-	return ExponentialHeightFogComponent->FogDensity; 
-}
-
-void SetExponentialHeightFogComponentFogDensity(UObject* Object, float InFogDensity)
-{
-	UExponentialHeightFogComponent* ExponentialHeightFogComponent = CastChecked<UExponentialHeightFogComponent>(Object);
-	ExponentialHeightFogComponent->SetFogDensity(InFogDensity);
-}
-
-float GetExponentialHeightFogComponentFogHeightFalloff(const UObject* Object)
-{
-	const UExponentialHeightFogComponent* ExponentialHeightFogComponent = CastChecked<const UExponentialHeightFogComponent>(Object);
-	return ExponentialHeightFogComponent->FogHeightFalloff;
-}
-
-void SetExponentialHeightFogComponentFogHeightFalloff(UObject* Object, float InFogHeightFalloff)
-{
-	UExponentialHeightFogComponent* ExponentialHeightFogComponent = CastChecked<UExponentialHeightFogComponent>(Object);
-	ExponentialHeightFogComponent->SetFogHeightFalloff(InFogHeightFalloff);
-}
-
-float GetExponentialHeightFogComponentFogMaxOpacity(const UObject* Object)
-{
-	const UExponentialHeightFogComponent* ExponentialHeightFogComponent = CastChecked<const UExponentialHeightFogComponent>(Object);
-	return ExponentialHeightFogComponent->FogMaxOpacity;
-}
-
-void SetExponentialHeightFogComponentFogMaxOpacity(UObject* Object, float InFogMaxOpacity)
-{
-	UExponentialHeightFogComponent* ExponentialHeightFogComponent = CastChecked<UExponentialHeightFogComponent>(Object);
-	ExponentialHeightFogComponent->SetFogMaxOpacity(InFogMaxOpacity);
-}
-
-float GetExponentialHeightFogComponentDirectionalInscatteringExponent(const UObject* Object)
-{
-	const UExponentialHeightFogComponent* ExponentialHeightFogComponent = CastChecked<const UExponentialHeightFogComponent>(Object);
-	return ExponentialHeightFogComponent->DirectionalInscatteringExponent;
-}
-
-void SetExponentialHeightFogComponentDirectionalInscatteringExponent(UObject* Object, float InDirectionalInscatteringExponent)
-{
-	UExponentialHeightFogComponent* ExponentialHeightFogComponent = CastChecked<UExponentialHeightFogComponent>(Object);
-	ExponentialHeightFogComponent->SetDirectionalInscatteringExponent(InDirectionalInscatteringExponent);
-}
-
-FIntermediateColor GetExponentialHeightFogVolumetricFogAlbedo(const UObject* Object, EColorPropertyType InColorType)
-{
-	// UExponentialHeightFog::VolumetricFogAlbedo is an FColor. FIntermediateColor is always blended and dealt with
-	// in linear space, so reinterpret the color.
-	ensure(InColorType == EColorPropertyType::Color);
-
-	const UExponentialHeightFogComponent* ExponentialHeightFogComponent = CastChecked<const UExponentialHeightFogComponent>(Object);
-	return FIntermediateColor(ExponentialHeightFogComponent->VolumetricFogAlbedo);
-}
-
-void SetExponentialHeightFogVolumetricFogAlbedo(UObject* Object, EColorPropertyType InColorType, const FIntermediateColor& InColor)
-{
-	// UExponentialHeightFog::VolumetricFogAlbedo is an FColor. FIntermediateColor is always blended and dealt with
-	// in linear space, so reinterpret the color back to FColor when setting via the public interface.
-	ensure(InColorType == EColorPropertyType::Color);
-
-	UExponentialHeightFogComponent* ExponentialHeightFogComponent = CastChecked<UExponentialHeightFogComponent>(Object);
-	ExponentialHeightFogComponent->SetVolumetricFogAlbedo(InColor.GetColor());
-}
-
-float GetSecondFogDataFogDensity(const UObject* Object)
-{
-	const UExponentialHeightFogComponent* ExponentialHeightFogComponent = CastChecked<const UExponentialHeightFogComponent>(Object);
-	return ExponentialHeightFogComponent->SecondFogData.FogDensity;
-}
-
-void SetSecondFogDataFogDensity(UObject* Object, float InFogDensity)
-{
-	UExponentialHeightFogComponent* ExponentialHeightFogComponent = CastChecked<UExponentialHeightFogComponent>(Object);
-	ExponentialHeightFogComponent->SetSecondFogDensity(InFogDensity);
-}
-
-float GetSecondFogDataFogHeightFalloff(const UObject* Object)
-{
-	const UExponentialHeightFogComponent* ExponentialHeightFogComponent = CastChecked<const UExponentialHeightFogComponent>(Object);
-	return ExponentialHeightFogComponent->SecondFogData.FogHeightFalloff;
-}
-
-void SetSecondFogDataFogHeightFalloff(UObject* Object, float InFogHeightFalloff)
-{
-	UExponentialHeightFogComponent* ExponentialHeightFogComponent = CastChecked<UExponentialHeightFogComponent>(Object);
-	ExponentialHeightFogComponent->SetSecondFogHeightFalloff(InFogHeightFalloff);
-}
-
-float GetSecondFogDataFogHeightOffset(const UObject* Object)
-{
-	const UExponentialHeightFogComponent* ExponentialHeightFogComponent = CastChecked<const UExponentialHeightFogComponent>(Object);
-	return ExponentialHeightFogComponent->SecondFogData.FogHeightOffset;
-}
-
-void SetSecondFogDataFogHeightOffset(UObject* Object, float InFogHeightOffset)
-{
-	UExponentialHeightFogComponent* ExponentialHeightFogComponent = CastChecked<UExponentialHeightFogComponent>(Object);
-	ExponentialHeightFogComponent->SetSecondFogHeightOffset(InFogHeightOffset);
-}
-
 uint8 GetSkeletalMeshAnimationMode(const UObject* Object)
 {
 	const USkeletalMeshComponent* SkeletalMeshComponent = CastChecked<const USkeletalMeshComponent>(Object);
@@ -826,6 +453,40 @@ struct FDoubleVectorHandler : TPropertyComponentHandler<FDoubleVectorPropertyTra
 	}
 };
 
+struct FBoolHandler : TPropertyComponentHandler<FBoolPropertyTraits, bool>
+{
+	virtual void DispatchInitializePropertyMetaDataTasks(const FPropertyDefinition& Definition, FSystemTaskPrerequisites& InPrerequisites, FSystemSubsequentTasks& Subsequents, UMovieSceneEntitySystemLinker* Linker) override
+	{
+		FBuiltInComponentTypes* BuiltInComponents = FBuiltInComponentTypes::Get();
+		FMovieSceneTracksComponentTypes* TrackComponents = FMovieSceneTracksComponentTypes::Get();
+
+		FEntityTaskBuilder()
+		.Read(BuiltInComponents->BoundObject)
+		.Read(BuiltInComponents->PropertyBinding)
+		.Write(TrackComponents->Bool.MetaDataComponents.GetType<0>())
+		.FilterAll({ BuiltInComponents->Tags.NeedsLink })
+		.Iterate_PerEntity(&Linker->EntityManager, [](UObject* Object, const FMovieScenePropertyBinding& Binding, FBoolPropertyTraits::FBoolMetaData& OutMetaData)
+		{
+			FBoolProperty* BoundProperty = CastField<FBoolProperty>(FTrackInstancePropertyBindings::FindProperty(Object, Binding.PropertyPath.ToString()));
+			if (ensure(BoundProperty))
+			{
+				if (BoundProperty->IsNativeBool())
+				{
+					OutMetaData.BitFieldSize = 0;
+					OutMetaData.BitIndex     = 0;
+				}
+				else
+				{
+					auto FieldMask = BoundProperty->GetFieldMask();
+					static_assert(std::is_same_v<decltype(FieldMask), uint8>, "Unexpected size of field mask returned from FBoolProperty::FieldMask");
+
+					OutMetaData.BitFieldSize = static_cast<uint8>(BoundProperty->ElementSize);
+					OutMetaData.BitIndex     = static_cast<uint8>(FMath::CountTrailingZeros(FieldMask));
+				}
+			}
+		});
+	}
+};
 
 struct FComponentTransformHandler : TPropertyComponentHandler<FComponentTransformPropertyTraits, double, double, double, double, double, double, double, double, double>
 {
@@ -857,6 +518,7 @@ FMovieSceneTracksComponentTypes::FMovieSceneTracksComponentTypes()
 	ComponentRegistry->NewPropertyType(FloatParameter, TEXT("float parameter"));
 	ComponentRegistry->NewPropertyType(ColorParameter, TEXT("color parameter"));
 
+	Bool.MetaDataComponents.Initialize(ComponentRegistry, TEXT("Bool Bitfield"));
 	Color.MetaDataComponents.Initialize(ComponentRegistry, TEXT("Color Type"));
 	FloatVector.MetaDataComponents.Initialize(ComponentRegistry, TEXT("Num Float Vector Channels"));
 	DoubleVector.MetaDataComponents.Initialize(ComponentRegistry, TEXT("Num Double Vector Channels"));
@@ -906,7 +568,7 @@ FMovieSceneTracksComponentTypes::FMovieSceneTracksComponentTypes()
 	.AddSoleChannel(BuiltInComponents->BoolResult)
 	.SetBlenderSystem<UMovieScenePiecewiseBoolBlenderSystem>()
 	.SetCustomAccessors(&Accessors.Bool)
-	.Commit();
+	.Commit(FBoolHandler());
 
 	// Set up FTransform properties
 	BuiltInComponents->PropertyRegistry.DefineCompositeProperty(Transform, TEXT("Apply FTransform Properties"))
@@ -996,88 +658,6 @@ FMovieSceneTracksComponentTypes::FMovieSceneTracksComponentTypes()
 	.SetBlenderSystem<UMovieScenePiecewiseDoubleBlenderSystem>()
 	.Commit(FColorParameterHandler());
 
-	// We have some custom accessors for well-known types.
-
-	// LightComponent
-	Accessors.Color.Add(
-			ULightComponent::StaticClass(), GET_MEMBER_NAME_CHECKED(ULightComponent, LightColor), 
-			GetLightComponentLightColor, SetLightComponentLightColor);
-	Accessors.Float.Add(
-			ULightComponent::StaticClass(), GET_MEMBER_NAME_CHECKED(ULightComponent, Intensity),
-			GetLightComponentIntensity, SetLightComponentIntensity);
-
-	// SkyLightComponent
-	Accessors.Color.Add(
-			USkyLightComponent::StaticClass(), GET_MEMBER_NAME_CHECKED(USkyLightComponent, LightColor), 
-			GetSkyLightComponentLightColor, SetSkyLightComponentLightColor);
-	Accessors.Float.Add(
-			USkyLightComponent::StaticClass(), GET_MEMBER_NAME_CHECKED(USkyLightComponent, Intensity),
-			GetSkyLightComponentIntensity, SetSkyLightComponentIntensity);
-
-	// SkyAtmosphereComponent
-	Accessors.Float.Add(
-			USkyAtmosphereComponent::StaticClass(), GET_MEMBER_NAME_CHECKED(USkyAtmosphereComponent, MieScatteringScale),
-			GetSkyAtmosphereComponentMieScatteringScale, SetSkyAtmosphereComponentMieScatteringScale);
-	Accessors.Float.Add(
-			USkyAtmosphereComponent::StaticClass(), GET_MEMBER_NAME_CHECKED(USkyAtmosphereComponent, OtherAbsorptionScale),
-			GetSkyAtmosphereComponentOtherAbsorptionScale, SetSkyAtmosphereComponentOtherAbsorptionScale);
-	Accessors.Float.Add(
-			USkyAtmosphereComponent::StaticClass(), GET_MEMBER_NAME_CHECKED(USkyAtmosphereComponent, RayleighExponentialDistribution),
-			GetSkyAtmosphereComponentRayleighExponentialDistribution, SetSkyAtmosphereComponentRayleighExponentialDistribution);
-	Accessors.Float.Add(
-			USkyAtmosphereComponent::StaticClass(), GET_MEMBER_NAME_CHECKED(USkyAtmosphereComponent, RayleighScatteringScale),
-			GetSkyAtmosphereComponentRayleighScatteringScale, SetSkyAtmosphereComponentRayleighScatteringScale);
-	Accessors.Float.Add(
-			USkyAtmosphereComponent::StaticClass(), GET_MEMBER_NAME_CHECKED(USkyAtmosphereComponent, AerialPespectiveViewDistanceScale),
-			GetSkyAtmosphereComponentAerialPerspectiveViewDistanceScale, SetSkyAtmosphereComponentAerialPerspectiveViewDistanceScale);
-	Accessors.Color.Add(
-			USkyAtmosphereComponent::StaticClass(), GET_MEMBER_NAME_CHECKED(USkyAtmosphereComponent, MieAbsorption),
-			GetSkyAtmosphereComponentMieAbsorption, SetSkyAtmosphereComponentMieAbsorption);
-	Accessors.Color.Add(
-			USkyAtmosphereComponent::StaticClass(), GET_MEMBER_NAME_CHECKED(USkyAtmosphereComponent, MieScattering),
-			GetSkyAtmosphereComponentMieScattering, SetSkyAtmosphereComponentMieScattering);
-	Accessors.Color.Add(
-			USkyAtmosphereComponent::StaticClass(), GET_MEMBER_NAME_CHECKED(USkyAtmosphereComponent, RayleighScattering),
-			GetSkyAtmosphereComponentRayleighScattering, SetSkyAtmosphereComponentRayleighScattering);
-	Accessors.Color.Add(
-			USkyAtmosphereComponent::StaticClass(), GET_MEMBER_NAME_CHECKED(USkyAtmosphereComponent, SkyLuminanceFactor),
-			GetSkyAtmosphereComponentSkyLuminanceFactor, SetSkyAtmosphereComponentSkyLuminanceFactor);
-
-	// PrimitiveComponent
-	const FString MassScalePath = FString::Printf(TEXT("%s.%s"), GET_MEMBER_NAME_STRING_CHECKED(UPrimitiveComponent, BodyInstance), GET_MEMBER_NAME_STRING_CHECKED(FBodyInstance, MassScale));
-	Accessors.Float.Add(
-		UPrimitiveComponent::StaticClass(), *MassScalePath,
-		GetMassScale, SetMassScale);
-
-	// ExponentialHeightFogComponent
-	Accessors.Float.Add(
-			UExponentialHeightFogComponent::StaticClass(), GET_MEMBER_NAME_CHECKED(UExponentialHeightFogComponent, FogDensity),
-			GetExponentialHeightFogComponentFogDensity, SetExponentialHeightFogComponentFogDensity);
-	Accessors.Float.Add(
-			UExponentialHeightFogComponent::StaticClass(), GET_MEMBER_NAME_CHECKED(UExponentialHeightFogComponent, FogHeightFalloff),
-			GetExponentialHeightFogComponentFogHeightFalloff, SetExponentialHeightFogComponentFogHeightFalloff);
-	Accessors.Float.Add(
-			UExponentialHeightFogComponent::StaticClass(), GET_MEMBER_NAME_CHECKED(UExponentialHeightFogComponent, DirectionalInscatteringExponent),
-			GetExponentialHeightFogComponentDirectionalInscatteringExponent, SetExponentialHeightFogComponentDirectionalInscatteringExponent);
-	Accessors.Float.Add(
-			UExponentialHeightFogComponent::StaticClass(), GET_MEMBER_NAME_CHECKED(UExponentialHeightFogComponent, FogMaxOpacity),
-			GetExponentialHeightFogComponentFogMaxOpacity, SetExponentialHeightFogComponentFogMaxOpacity);
-	Accessors.Color.Add(
-			UExponentialHeightFogComponent::StaticClass(), GET_MEMBER_NAME_CHECKED(UExponentialHeightFogComponent, VolumetricFogAlbedo),
-			GetExponentialHeightFogVolumetricFogAlbedo, SetExponentialHeightFogVolumetricFogAlbedo);
-	const FString SecondFogDataFogDensityPath = FString::Printf(TEXT("%s.%s"), GET_MEMBER_NAME_STRING_CHECKED(UExponentialHeightFogComponent, SecondFogData), GET_MEMBER_NAME_STRING_CHECKED(FExponentialHeightFogData, FogDensity));
-	Accessors.Float.Add(
-			UExponentialHeightFogComponent::StaticClass(), *SecondFogDataFogDensityPath,
-			GetSecondFogDataFogDensity, SetSecondFogDataFogDensity);
-	const FString SecondFogDataFogHeightFalloffPath = FString::Printf(TEXT("%s.%s"), GET_MEMBER_NAME_STRING_CHECKED(UExponentialHeightFogComponent, SecondFogData), GET_MEMBER_NAME_STRING_CHECKED(FExponentialHeightFogData, FogHeightFalloff));
-	Accessors.Float.Add(
-			UExponentialHeightFogComponent::StaticClass(), *SecondFogDataFogHeightFalloffPath,
-			GetSecondFogDataFogHeightFalloff, SetSecondFogDataFogHeightFalloff);
-	const FString SecondFogDataFogHeightOffsetPath = FString::Printf(TEXT("%s.%s"), GET_MEMBER_NAME_STRING_CHECKED(UExponentialHeightFogComponent, SecondFogData), GET_MEMBER_NAME_STRING_CHECKED(FExponentialHeightFogData, FogHeightOffset));
-	Accessors.Float.Add(
-			UExponentialHeightFogComponent::StaticClass(), *SecondFogDataFogHeightOffsetPath,
-			GetSecondFogDataFogHeightOffset, SetSecondFogDataFogHeightOffset);
-
 	Accessors.Byte.Add(
 		USkeletalMeshComponent::StaticClass(), USkeletalMeshComponent::GetAnimationModePropertyNameChecked(),
 		GetSkeletalMeshAnimationMode, SetSkeletalMeshAnimationMode);
@@ -1123,8 +703,6 @@ FMovieSceneTracksComponentTypes::FMovieSceneTracksComponentTypes()
 	// --------------------------------------------------------------------------------------------
 	// Set up component transforms
 	{
-		Accessors.ComponentTransform.Add(USceneComponent::StaticClass(), "Transform", &GetComponentTransform, &SetComponentTransformAndVelocity);
-
 		BuiltInComponents->PropertyRegistry.DefineCompositeProperty(ComponentTransform, TEXT("Call USceneComponent::SetRelativeTransform"))
 		.AddComposite(BuiltInComponents->DoubleResult[0], &FIntermediate3DTransform::T_X)
 		.AddComposite(BuiltInComponents->DoubleResult[1], &FIntermediate3DTransform::T_Y)
@@ -1172,6 +750,8 @@ FMovieSceneTracksComponentTypes::FMovieSceneTracksComponentTypes()
 	// --------------------------------------------------------------------------------------------
 	// Set up SkeletalAnimation components
 	ComponentRegistry->Factories.DuplicateChildComponent(SkeletalAnimation);
+
+	InitializeMovieSceneTracksAccessors(this);
 }
 
 FMovieSceneTracksComponentTypes::~FMovieSceneTracksComponentTypes()
diff --git a/Engine/Source/Runtime/MovieSceneTracks/Private/MovieSceneTracksCustomAccessors.cpp b/Engine/Source/Runtime/MovieSceneTracks/Private/MovieSceneTracksCustomAccessors.cpp
new file mode 100644
--- /dev/null
+++ b/Engine/Source/Runtime/MovieSceneTracks/Private/MovieSceneTracksCustomAccessors.cpp
@@ -0,0 +1,487 @@
+// Copyright Epic Games, Inc. All Rights Reserved.
+
+#include "MovieSceneTracksCustomAccessors.h"
+#include "MovieSceneTracksComponentTypes.h"
+#include "Components/ExponentialHeightFogComponent.h"
+#include "Components/LightComponent.h"
+#include "Components/PrimitiveComponent.h"
+#include "Components/SkyAtmosphereComponent.h"
+#include "Components/SkyLightComponent.h"
+
+namespace UE::MovieScene
+{
+
+FIntermediate3DTransform GetComponentTransform(const UObject* Object)
+{
+	const USceneComponent* SceneComponent = CastChecked<const USceneComponent>(Object);
+	FIntermediate3DTransform Result(SceneComponent->GetRelativeLocation(), SceneComponent->GetRelativeRotation(), SceneComponent->GetRelativeScale3D());
+	return Result;
+}
+
+void SetComponentTranslationAndRotation(USceneComponent* SceneComponent, const FIntermediate3DTransform& Transform)
+{
+	// If this is a simulating component, teleport since sequencer takes over. 
+	// Teleport will not have no velocity, but it's computed later by sequencer so that it will be correct for physics.
+	// @todo: We would really rather not 
+	AActor* Actor = SceneComponent->GetOwner();
+	USceneComponent* RootComponent = Actor ? Actor->GetRootComponent() : nullptr;
+	bool bIsSimulatingPhysics = RootComponent ? RootComponent->IsSimulatingPhysics() : false;
+
+	FVector Translation = Transform.GetTranslation();
+	FRotator Rotation = Transform.GetRotation();
+	SceneComponent->SetRelativeLocationAndRotation(Translation, Rotation, false, nullptr, bIsSimulatingPhysics ? ETeleportType::ResetPhysics : ETeleportType::None);
+
+	// Force the location and rotation values to avoid Rot->Quat->Rot conversions
+	SceneComponent->SetRelativeLocation_Direct(Translation);
+	SceneComponent->SetRelativeRotation_Direct(Rotation);
+}
+
+void SetComponentTransform(USceneComponent* SceneComponent, const FIntermediate3DTransform& Transform)
+{
+	// If this is a simulating component, teleport since sequencer takes over. 
+	// Teleport will not have no velocity, but it's computed later by sequencer so that it will be correct for physics.
+	// @todo: We would really rather not 
+	AActor* Actor = SceneComponent->GetOwner();
+	USceneComponent* RootComponent = Actor ? Actor->GetRootComponent() : nullptr;
+	bool bIsSimulatingPhysics = RootComponent ? RootComponent->IsSimulatingPhysics() : false;
+
+	FVector Translation = Transform.GetTranslation();
+	FRotator Rotation = Transform.GetRotation();
+	SceneComponent->SetRelativeTransform(FTransform(Rotation, Translation, Transform.GetScale()), false, nullptr, bIsSimulatingPhysics ? ETeleportType::ResetPhysics : ETeleportType::None);
+
+	// Force the location and rotation values to avoid Rot->Quat->Rot conversions
+	SceneComponent->SetRelativeLocation_Direct(Translation);
+	SceneComponent->SetRelativeRotation_Direct(Rotation);
+}
+
+void SetComponentTransformAndVelocity(UObject* Object, const FIntermediate3DTransform& InTransform)
+{
+	UE::MovieScene::FIntermediate3DTransform::ApplyTransformTo(CastChecked<USceneComponent>(Object), InTransform);
+}
+
+FIntermediateColor GetLightComponentLightColor(const UObject* Object, EColorPropertyType InColorType)
+{
+	ensure(InColorType == EColorPropertyType::Color);
+
+	const ULightComponent* LightComponent = CastChecked<const ULightComponent>(Object);
+	return FIntermediateColor(LightComponent->GetLightColor());
+}
+
+void SetLightComponentLightColor(UObject* Object, EColorPropertyType InColorType, const FIntermediateColor& InColor)
+{
+	// This is a little esoteric - ULightComponentBase::LightColor is the UPROPERTY that generates the meta-data
+	// for this custom callback, but it is an FColor, even though the public get/set functions expose it as an
+	// FLinearColor. FIntermediateColor is always blended and dealt with in linear space, so it's fine to 
+	// simply reinterpret the color
+	ensure(InColorType == EColorPropertyType::Color);
+
+	const bool bConvertBackToSRgb = true;
+	ULightComponent* LightComponent = CastChecked<ULightComponent>(Object);
+	LightComponent->SetLightColor(InColor.GetLinearColor(), bConvertBackToSRgb);
+}
+
+float GetLightComponentIntensity(const UObject* Object)
+{
+	const ULightComponent* LightComponent = CastChecked<const ULightComponent>(Object);
+	return LightComponent->Intensity;
+}
+
+void SetLightComponentIntensity(UObject* Object, float InIntensity)
+{
+	ULightComponent* LightComponent = CastChecked<ULightComponent>(Object);
+	LightComponent->SetIntensity(InIntensity);
+}
+
+float GetLightComponentVolumetricScatteringIntensity(const UObject* Object)
+{
+	const ULightComponent* LightComponent = CastChecked<const ULightComponent>(Object);
+	return LightComponent->VolumetricScatteringIntensity;
+}
+
+void SetLightComponentVolumetricScatteringIntensity(UObject* Object, float InVolumetricScatteringIntensity)
+{
+	ULightComponent* LightComponent = CastChecked<ULightComponent>(Object);
+	LightComponent->SetVolumetricScatteringIntensity(InVolumetricScatteringIntensity);
+}
+
+FIntermediateColor GetSkyLightComponentLightColor(const UObject* Object, EColorPropertyType InColorType)
+{
+	ensure(InColorType == EColorPropertyType::Color);
+
+	const USkyLightComponent* SkyLightComponent = CastChecked<const USkyLightComponent>(Object);
+	return FIntermediateColor(SkyLightComponent->GetLightColor());
+}
+
+void SetSkyLightComponentLightColor(UObject* Object, EColorPropertyType InColorType, const FIntermediateColor& InColor)
+{
+	// This is a little esoteric - ULightComponentBase::LightColor is the UPROPERTY that generates the meta-data
+	// for this custom callback, but it is an FColor, even though the public get/set functions expose it as an
+	// FLinearColor. FIntermediateColor is always blended and dealt with in linear space, so it's fine to 
+	// simply reinterpret the color
+	ensure(InColorType == EColorPropertyType::Color);
+
+	USkyLightComponent* SkyLightComponent = CastChecked<USkyLightComponent>(Object);
+	SkyLightComponent->SetLightColor(InColor.GetLinearColor());
+}
+
+float GetSkyLightComponentIntensity(const UObject* Object)
+{
+	const USkyLightComponent* SkyLightComponent = CastChecked<const USkyLightComponent>(Object);
+	return SkyLightComponent->Intensity;
+}
+
+void SetSkyLightComponentIntensity(UObject* Object, float InIntensity)
+{
+	USkyLightComponent* SkyLightComponent = CastChecked<USkyLightComponent>(Object);
+	SkyLightComponent->SetIntensity(InIntensity);
+}
+
+float GetSkyAtmosphereComponentMieScatteringScale(const UObject* Object)
+{
+	const USkyAtmosphereComponent* SkyAtmosphereComponent = CastChecked<const USkyAtmosphereComponent>(Object);
+	return SkyAtmosphereComponent->MieScatteringScale;
+}
+
+void SetSkyAtmosphereComponentMieScatteringScale(UObject* Object, float InMieScatteringScale)
+{
+	USkyAtmosphereComponent* SkyAtmosphereComponent = CastChecked<USkyAtmosphereComponent>(Object);
+	SkyAtmosphereComponent->SetMieScatteringScale(InMieScatteringScale);
+}
+
+float GetSkyAtmosphereComponentOtherAbsorptionScale(const UObject* Object)
+{
+	const USkyAtmosphereComponent* SkyAtmosphereComponent = CastChecked<const USkyAtmosphereComponent>(Object);
+	return SkyAtmosphereComponent->OtherAbsorptionScale;
+}
+
+void SetSkyAtmosphereComponentOtherAbsorptionScale(UObject* Object, float InOtherAbsorptionScale)
+{
+	USkyAtmosphereComponent* SkyAtmosphereComponent = CastChecked<USkyAtmosphereComponent>(Object);
+	SkyAtmosphereComponent->SetOtherAbsorptionScale(InOtherAbsorptionScale);
+}
+
+float GetSkyAtmosphereComponentRayleighScatteringScale(const UObject* Object)
+{
+	const USkyAtmosphereComponent* SkyAtmosphereComponent = CastChecked<const USkyAtmosphereComponent>(Object);
+	return SkyAtmosphereComponent->RayleighScatteringScale;
+}
+
+void SetSkyAtmosphereComponentRayleighScatteringScale(UObject* Object, float InRayleighScatteringScale)
+{
+	USkyAtmosphereComponent* SkyAtmosphereComponent = CastChecked<USkyAtmosphereComponent>(Object);
+	SkyAtmosphereComponent->SetRayleighScatteringScale(InRayleighScatteringScale);
+}
+
+float GetSkyAtmosphereComponentRayleighExponentialDistribution(const UObject* Object)
+{
+	const USkyAtmosphereComponent* SkyAtmosphereComponent = CastChecked<const USkyAtmosphereComponent>(Object);
+	return SkyAtmosphereComponent->RayleighExponentialDistribution;
+}
+
+void SetSkyAtmosphereComponentRayleighExponentialDistribution(UObject* Object, float InRayleighExponentialDistribution)
+{
+	USkyAtmosphereComponent* SkyAtmosphereComponent = CastChecked<USkyAtmosphereComponent>(Object);
+	SkyAtmosphereComponent->SetRayleighExponentialDistribution(InRayleighExponentialDistribution);
+}
+
+float GetSkyAtmosphereComponentAerialPerspectiveViewDistanceScale(const UObject* Object)
+{
+	const USkyAtmosphereComponent* SkyAtmosphereComponent = CastChecked<const USkyAtmosphereComponent>(Object);
+	return SkyAtmosphereComponent->AerialPespectiveViewDistanceScale;
+}
+
+void SetSkyAtmosphereComponentAerialPerspectiveViewDistanceScale(UObject* Object, float InAerialPerspectiveViewDistanceScale)
+{
+	USkyAtmosphereComponent* SkyAtmosphereComponent = CastChecked<USkyAtmosphereComponent>(Object);
+	SkyAtmosphereComponent->SetAerialPespectiveViewDistanceScale(InAerialPerspectiveViewDistanceScale);
+}
+
+FIntermediateColor GetSkyAtmosphereComponentMieAbsorption(const UObject* Object, EColorPropertyType InColorType)
+{
+	// USkyAtmosphereComponent::MieAbsorption is a FLinearColor. FIntermediateColor is always blended and dealt with
+	// in linear space so it should be safe to reinterpret it as linear.
+	ensure(InColorType == EColorPropertyType::Linear);
+
+	const USkyAtmosphereComponent* SkyAtmosphereComponent = CastChecked<const USkyAtmosphereComponent>(Object);
+	return FIntermediateColor(SkyAtmosphereComponent->MieAbsorption);
+}
+
+void SetSkyAtmosphereComponentMieAbsorption(UObject* Object, EColorPropertyType InColorType, const FIntermediateColor& InColor)
+{
+	// USkyAtmosphereComponent::MieAbsorption is a FLinearColor. FIntermediateColor is always blended and dealt with
+	// in linear space so it should be safe to reinterpret it as linear.
+	ensure(InColorType == EColorPropertyType::Linear);
+
+	USkyAtmosphereComponent* SkyAtmosphereComponent = CastChecked<USkyAtmosphereComponent>(Object);
+	SkyAtmosphereComponent->SetMieAbsorption(InColor.GetLinearColor());
+}
+
+FIntermediateColor GetSkyAtmosphereComponentMieScattering(const UObject* Object, EColorPropertyType InColorType)
+{
+	// USkyAtmosphereComponent::MieScattering is a FLinearColor. FIntermediateColor is always blended and dealt with
+	// in linear space so it should be safe to reinterpret it as linear.
+	ensure(InColorType == EColorPropertyType::Linear);
+
+	const USkyAtmosphereComponent* SkyAtmosphereComponent = CastChecked<const USkyAtmosphereComponent>(Object);
+	return FIntermediateColor(SkyAtmosphereComponent->MieScattering);
+}
+
+void SetSkyAtmosphereComponentMieScattering(UObject* Object, EColorPropertyType InColorType, const FIntermediateColor& InColor)
+{
+	// USkyAtmosphereComponent::MieScattering is a FLinearColor. FIntermediateColor is always blended and dealt with
+	// in linear space so it should be safe to reinterpret it as linear.
+	ensure(InColorType == EColorPropertyType::Linear);
+
+	USkyAtmosphereComponent* SkyAtmosphereComponent = CastChecked<USkyAtmosphereComponent>(Object);
+	SkyAtmosphereComponent->SetMieScattering(InColor.GetLinearColor());
+}
+
+FIntermediateColor GetSkyAtmosphereComponentRayleighScattering(const UObject* Object, EColorPropertyType InColorType)
+{
+	// USkyAtmosphereComponent::RayleighScattering is a FLinearColor. FIntermediateColor is always blended and dealt with
+	// in linear space so it should be safe to reinterpret it as linear.
+	ensure(InColorType == EColorPropertyType::Linear);
+
+	const USkyAtmosphereComponent* SkyAtmosphereComponent = CastChecked<const USkyAtmosphereComponent>(Object);
+	return FIntermediateColor(SkyAtmosphereComponent->RayleighScattering);
+}
+
+void SetSkyAtmosphereComponentRayleighScattering(UObject* Object, EColorPropertyType InColorType, const FIntermediateColor& InColor)
+{
+	// USkyAtmosphereComponent::RayleighScattering is a FLinearColor. FIntermediateColor is always blended and dealt with
+	// in linear space so it should be safe to reinterpret it as linear.
+	ensure(InColorType == EColorPropertyType::Linear);
+
+	USkyAtmosphereComponent* SkyAtmosphereComponent = CastChecked<USkyAtmosphereComponent>(Object);
+	SkyAtmosphereComponent->SetRayleighScattering(InColor.GetLinearColor());
+}
+
+FIntermediateColor GetSkyAtmosphereComponentSkyLuminanceFactor(const UObject* Object, EColorPropertyType InColorType)
+{
+	// USkyAtmosphereComponent::SkyLuminanceFactor is a FLinearColor. FIntermediateColor is always blended and dealt with
+	// in linear space so it should be safe to reinterpret it as linear.
+	ensure(InColorType == EColorPropertyType::Linear);
+
+	const USkyAtmosphereComponent* SkyAtmosphereComponent = CastChecked<const USkyAtmosphereComponent>(Object);
+	return FIntermediateColor(SkyAtmosphereComponent->SkyLuminanceFactor);
+}
+
+void SetSkyAtmosphereComponentSkyLuminanceFactor(UObject* Object, EColorPropertyType InColorType, const FIntermediateColor& InColor)
+{
+	// USkyAtmosphereComponent::SkyLuminanceFactor is a FLinearColor. FIntermediateColor is always blended and dealt with
+	// in linear space so it should be safe to reinterpret it as linear.
+	ensure(InColorType == EColorPropertyType::Linear);
+
+	USkyAtmosphereComponent* SkyAtmosphereComponent = CastChecked<USkyAtmosphereComponent>(Object);
+	SkyAtmosphereComponent->SetSkyLuminanceFactor(InColor.GetLinearColor());
+}
+
+float GetMassScale(const UObject* Object)
+{
+	const UPrimitiveComponent* PrimitiveComponent = CastChecked<const UPrimitiveComponent>(Object);
+	return PrimitiveComponent->GetMassScale();
+}
+
+void SetMassScale(UObject* Object, float InMassScale)
+{
+	UPrimitiveComponent* PrimitiveComponent = CastChecked<UPrimitiveComponent>(Object);
+	PrimitiveComponent->SetMassScale(NAME_None, InMassScale);
+}
+
+float GetExponentialHeightFogComponentFogDensity(const UObject* Object)
+{
+	const UExponentialHeightFogComponent* ExponentialHeightFogComponent = CastChecked<const UExponentialHeightFogComponent>(Object);
+	return ExponentialHeightFogComponent->FogDensity; 
+}
+
+void SetExponentialHeightFogComponentFogDensity(UObject* Object, float InFogDensity)
+{
+	UExponentialHeightFogComponent* ExponentialHeightFogComponent = CastChecked<UExponentialHeightFogComponent>(Object);
+	ExponentialHeightFogComponent->SetFogDensity(InFogDensity);
+}
+
+float GetExponentialHeightFogComponentFogHeightFalloff(const UObject* Object)
+{
+	const UExponentialHeightFogComponent* ExponentialHeightFogComponent = CastChecked<const UExponentialHeightFogComponent>(Object);
+	return ExponentialHeightFogComponent->FogHeightFalloff;
+}
+
+void SetExponentialHeightFogComponentFogHeightFalloff(UObject* Object, float InFogHeightFalloff)
+{
+	UExponentialHeightFogComponent* ExponentialHeightFogComponent = CastChecked<UExponentialHeightFogComponent>(Object);
+	ExponentialHeightFogComponent->SetFogHeightFalloff(InFogHeightFalloff);
+}
+
+float GetExponentialHeightFogComponentFogMaxOpacity(const UObject* Object)
+{
+	const UExponentialHeightFogComponent* ExponentialHeightFogComponent = CastChecked<const UExponentialHeightFogComponent>(Object);
+	return ExponentialHeightFogComponent->FogMaxOpacity;
+}
+
+void SetExponentialHeightFogComponentFogMaxOpacity(UObject* Object, float InFogMaxOpacity)
+{
+	UExponentialHeightFogComponent* ExponentialHeightFogComponent = CastChecked<UExponentialHeightFogComponent>(Object);
+	ExponentialHeightFogComponent->SetFogMaxOpacity(InFogMaxOpacity);
+}
+
+float GetExponentialHeightFogComponentDirectionalInscatteringExponent(const UObject* Object)
+{
+	const UExponentialHeightFogComponent* ExponentialHeightFogComponent = CastChecked<const UExponentialHeightFogComponent>(Object);
+	return ExponentialHeightFogComponent->DirectionalInscatteringExponent;
+}
+
+void SetExponentialHeightFogComponentDirectionalInscatteringExponent(UObject* Object, float InDirectionalInscatteringExponent)
+{
+	UExponentialHeightFogComponent* ExponentialHeightFogComponent = CastChecked<UExponentialHeightFogComponent>(Object);
+	ExponentialHeightFogComponent->SetDirectionalInscatteringExponent(InDirectionalInscatteringExponent);
+}
+
+FIntermediateColor GetExponentialHeightFogVolumetricFogAlbedo(const UObject* Object, EColorPropertyType InColorType)
+{
+	// UExponentialHeightFog::VolumetricFogAlbedo is an FColor. FIntermediateColor is always blended and dealt with
+	// in linear space, so reinterpret the color.
+	ensure(InColorType == EColorPropertyType::Color);
+
+	const UExponentialHeightFogComponent* ExponentialHeightFogComponent = CastChecked<const UExponentialHeightFogComponent>(Object);
+	return FIntermediateColor(ExponentialHeightFogComponent->VolumetricFogAlbedo);
+}
+
+void SetExponentialHeightFogVolumetricFogAlbedo(UObject* Object, EColorPropertyType InColorType, const FIntermediateColor& InColor)
+{
+	// UExponentialHeightFog::VolumetricFogAlbedo is an FColor. FIntermediateColor is always blended and dealt with
+	// in linear space, so reinterpret the color back to FColor when setting via the public interface.
+	ensure(InColorType == EColorPropertyType::Color);
+
+	UExponentialHeightFogComponent* ExponentialHeightFogComponent = CastChecked<UExponentialHeightFogComponent>(Object);
+	ExponentialHeightFogComponent->SetVolumetricFogAlbedo(InColor.GetColor());
+}
+
+float GetSecondFogDataFogDensity(const UObject* Object)
+{
+	const UExponentialHeightFogComponent* ExponentialHeightFogComponent = CastChecked<const UExponentialHeightFogComponent>(Object);
+	return ExponentialHeightFogComponent->SecondFogData.FogDensity;
+}
+
+void SetSecondFogDataFogDensity(UObject* Object, float InFogDensity)
+{
+	UExponentialHeightFogComponent* ExponentialHeightFogComponent = CastChecked<UExponentialHeightFogComponent>(Object);
+	ExponentialHeightFogComponent->SetSecondFogDensity(InFogDensity);
+}
+
+float GetSecondFogDataFogHeightFalloff(const UObject* Object)
+{
+	const UExponentialHeightFogComponent* ExponentialHeightFogComponent = CastChecked<const UExponentialHeightFogComponent>(Object);
+	return ExponentialHeightFogComponent->SecondFogData.FogHeightFalloff;
+}
+
+void SetSecondFogDataFogHeightFalloff(UObject* Object, float InFogHeightFalloff)
+{
+	UExponentialHeightFogComponent* ExponentialHeightFogComponent = CastChecked<UExponentialHeightFogComponent>(Object);
+	ExponentialHeightFogComponent->SetSecondFogHeightFalloff(InFogHeightFalloff);
+}
+
+float GetSecondFogDataFogHeightOffset(const UObject* Object)
+{
+	const UExponentialHeightFogComponent* ExponentialHeightFogComponent = CastChecked<const UExponentialHeightFogComponent>(Object);
+	return ExponentialHeightFogComponent->SecondFogData.FogHeightOffset;
+}
+
+void SetSecondFogDataFogHeightOffset(UObject* Object, float InFogHeightOffset)
+{
+	UExponentialHeightFogComponent* ExponentialHeightFogComponent = CastChecked<UExponentialHeightFogComponent>(Object);
+	ExponentialHeightFogComponent->SetSecondFogHeightOffset(InFogHeightOffset);
+}
+
+void InitializeMovieSceneTracksAccessors(FMovieSceneTracksComponentTypes* TracksComponents)
+{
+	// We have some custom accessors for well-known types.
+
+	// LightComponent
+	TracksComponents->Accessors.Color.Add(
+			ULightComponent::StaticClass(), GET_MEMBER_NAME_CHECKED(ULightComponent, LightColor), 
+			GetLightComponentLightColor, SetLightComponentLightColor);
+
+	TracksComponents->Accessors.Float.Add(
+			ULightComponent::StaticClass(), GET_MEMBER_NAME_CHECKED(ULightComponent, Intensity),
+			GetLightComponentIntensity, SetLightComponentIntensity);
+
+	TracksComponents->Accessors.Float.Add(
+			ULightComponent::StaticClass(), GET_MEMBER_NAME_CHECKED(ULightComponent, VolumetricScatteringIntensity),
+			GetLightComponentVolumetricScatteringIntensity, SetLightComponentVolumetricScatteringIntensity);
+
+	// SkyLightComponent
+	TracksComponents->Accessors.Color.Add(
+			USkyLightComponent::StaticClass(), GET_MEMBER_NAME_CHECKED(USkyLightComponent, LightColor), 
+			GetSkyLightComponentLightColor, SetSkyLightComponentLightColor);
+	TracksComponents->Accessors.Float.Add(
+			USkyLightComponent::StaticClass(), GET_MEMBER_NAME_CHECKED(USkyLightComponent, Intensity),
+			GetSkyLightComponentIntensity, SetSkyLightComponentIntensity);
+
+	// SkyAtmosphereComponent
+	TracksComponents->Accessors.Float.Add(
+			USkyAtmosphereComponent::StaticClass(), GET_MEMBER_NAME_CHECKED(USkyAtmosphereComponent, MieScatteringScale),
+			GetSkyAtmosphereComponentMieScatteringScale, SetSkyAtmosphereComponentMieScatteringScale);
+	TracksComponents->Accessors.Float.Add(
+			USkyAtmosphereComponent::StaticClass(), GET_MEMBER_NAME_CHECKED(USkyAtmosphereComponent, OtherAbsorptionScale),
+			GetSkyAtmosphereComponentOtherAbsorptionScale, SetSkyAtmosphereComponentOtherAbsorptionScale);
+	TracksComponents->Accessors.Float.Add(
+			USkyAtmosphereComponent::StaticClass(), GET_MEMBER_NAME_CHECKED(USkyAtmosphereComponent, RayleighExponentialDistribution),
+			GetSkyAtmosphereComponentRayleighExponentialDistribution, SetSkyAtmosphereComponentRayleighExponentialDistribution);
+	TracksComponents->Accessors.Float.Add(
+			USkyAtmosphereComponent::StaticClass(), GET_MEMBER_NAME_CHECKED(USkyAtmosphereComponent, RayleighScatteringScale),
+			GetSkyAtmosphereComponentRayleighScatteringScale, SetSkyAtmosphereComponentRayleighScatteringScale);
+	TracksComponents->Accessors.Float.Add(
+			USkyAtmosphereComponent::StaticClass(), GET_MEMBER_NAME_CHECKED(USkyAtmosphereComponent, AerialPespectiveViewDistanceScale),
+			GetSkyAtmosphereComponentAerialPerspectiveViewDistanceScale, SetSkyAtmosphereComponentAerialPerspectiveViewDistanceScale);
+	TracksComponents->Accessors.Color.Add(
+			USkyAtmosphereComponent::StaticClass(), GET_MEMBER_NAME_CHECKED(USkyAtmosphereComponent, MieAbsorption),
+			GetSkyAtmosphereComponentMieAbsorption, SetSkyAtmosphereComponentMieAbsorption);
+	TracksComponents->Accessors.Color.Add(
+			USkyAtmosphereComponent::StaticClass(), GET_MEMBER_NAME_CHECKED(USkyAtmosphereComponent, MieScattering),
+			GetSkyAtmosphereComponentMieScattering, SetSkyAtmosphereComponentMieScattering);
+	TracksComponents->Accessors.Color.Add(
+			USkyAtmosphereComponent::StaticClass(), GET_MEMBER_NAME_CHECKED(USkyAtmosphereComponent, RayleighScattering),
+			GetSkyAtmosphereComponentRayleighScattering, SetSkyAtmosphereComponentRayleighScattering);
+	TracksComponents->Accessors.Color.Add(
+			USkyAtmosphereComponent::StaticClass(), GET_MEMBER_NAME_CHECKED(USkyAtmosphereComponent, SkyLuminanceFactor),
+			GetSkyAtmosphereComponentSkyLuminanceFactor, SetSkyAtmosphereComponentSkyLuminanceFactor);
+
+	// PrimitiveComponent
+	const FString MassScalePath = FString::Printf(TEXT("%s.%s"), GET_MEMBER_NAME_STRING_CHECKED(UPrimitiveComponent, BodyInstance), GET_MEMBER_NAME_STRING_CHECKED(FBodyInstance, MassScale));
+	TracksComponents->Accessors.Float.Add(
+		UPrimitiveComponent::StaticClass(), *MassScalePath,
+		GetMassScale, SetMassScale);
+
+	// ExponentialHeightFogComponent
+	TracksComponents->Accessors.Float.Add(
+			UExponentialHeightFogComponent::StaticClass(), GET_MEMBER_NAME_CHECKED(UExponentialHeightFogComponent, FogDensity),
+			GetExponentialHeightFogComponentFogDensity, SetExponentialHeightFogComponentFogDensity);
+	TracksComponents->Accessors.Float.Add(
+			UExponentialHeightFogComponent::StaticClass(), GET_MEMBER_NAME_CHECKED(UExponentialHeightFogComponent, FogHeightFalloff),
+			GetExponentialHeightFogComponentFogHeightFalloff, SetExponentialHeightFogComponentFogHeightFalloff);
+	TracksComponents->Accessors.Float.Add(
+			UExponentialHeightFogComponent::StaticClass(), GET_MEMBER_NAME_CHECKED(UExponentialHeightFogComponent, DirectionalInscatteringExponent),
+			GetExponentialHeightFogComponentDirectionalInscatteringExponent, SetExponentialHeightFogComponentDirectionalInscatteringExponent);
+	TracksComponents->Accessors.Float.Add(
+			UExponentialHeightFogComponent::StaticClass(), GET_MEMBER_NAME_CHECKED(UExponentialHeightFogComponent, FogMaxOpacity),
+			GetExponentialHeightFogComponentFogMaxOpacity, SetExponentialHeightFogComponentFogMaxOpacity);
+	TracksComponents->Accessors.Color.Add(
+			UExponentialHeightFogComponent::StaticClass(), GET_MEMBER_NAME_CHECKED(UExponentialHeightFogComponent, VolumetricFogAlbedo),
+			GetExponentialHeightFogVolumetricFogAlbedo, SetExponentialHeightFogVolumetricFogAlbedo);
+	const FString SecondFogDataFogDensityPath = FString::Printf(TEXT("%s.%s"), GET_MEMBER_NAME_STRING_CHECKED(UExponentialHeightFogComponent, SecondFogData), GET_MEMBER_NAME_STRING_CHECKED(FExponentialHeightFogData, FogDensity));
+	TracksComponents->Accessors.Float.Add(
+			UExponentialHeightFogComponent::StaticClass(), *SecondFogDataFogDensityPath,
+			GetSecondFogDataFogDensity, SetSecondFogDataFogDensity);
+	const FString SecondFogDataFogHeightFalloffPath = FString::Printf(TEXT("%s.%s"), GET_MEMBER_NAME_STRING_CHECKED(UExponentialHeightFogComponent, SecondFogData), GET_MEMBER_NAME_STRING_CHECKED(FExponentialHeightFogData, FogHeightFalloff));
+	TracksComponents->Accessors.Float.Add(
+			UExponentialHeightFogComponent::StaticClass(), *SecondFogDataFogHeightFalloffPath,
+			GetSecondFogDataFogHeightFalloff, SetSecondFogDataFogHeightFalloff);
+	const FString SecondFogDataFogHeightOffsetPath = FString::Printf(TEXT("%s.%s"), GET_MEMBER_NAME_STRING_CHECKED(UExponentialHeightFogComponent, SecondFogData), GET_MEMBER_NAME_STRING_CHECKED(FExponentialHeightFogData, FogHeightOffset));
+	TracksComponents->Accessors.Float.Add(
+			UExponentialHeightFogComponent::StaticClass(), *SecondFogDataFogHeightOffsetPath,
+			GetSecondFogDataFogHeightOffset, SetSecondFogDataFogHeightOffset);
+
+	TracksComponents->Accessors.ComponentTransform.Add(USceneComponent::StaticClass(), "Transform", &GetComponentTransform, &SetComponentTransformAndVelocity);
+}
+
+} // namespace UE::MovieScene
\ No newline at end of file
diff --git a/Engine/Source/Runtime/MovieSceneTracks/Private/MovieSceneTracksCustomAccessors.h b/Engine/Source/Runtime/MovieSceneTracks/Private/MovieSceneTracksCustomAccessors.h
new file mode 100644
--- /dev/null
+++ b/Engine/Source/Runtime/MovieSceneTracks/Private/MovieSceneTracksCustomAccessors.h
@@ -0,0 +1,24 @@
+// Copyright Epic Games, Inc. All Rights Reserved.
+
+#pragma once
+
+#include "CoreTypes.h"
+
+class UObject;
+class USceneComponent;
+
+namespace UE::MovieScene
+{
+
+struct FMovieSceneTracksComponentTypes;
+struct FIntermediate3DTransform;
+
+void InitializeMovieSceneTracksAccessors(FMovieSceneTracksComponentTypes* TracksComponents);
+
+MOVIESCENETRACKS_API FIntermediate3DTransform GetComponentTransform(const UObject* Object);
+MOVIESCENETRACKS_API void SetComponentTranslationAndRotation(USceneComponent* SceneComponent, const FIntermediate3DTransform& Transform);
+MOVIESCENETRACKS_API void SetComponentTransform(USceneComponent* SceneComponent, const FIntermediate3DTransform& Transform);
+MOVIESCENETRACKS_API void SetComponentTransformAndVelocity(UObject* Object, const FIntermediate3DTransform& InTransform);
+
+
+} // namespace UE::MovieScene
\ No newline at end of file
diff --git a/Engine/Source/Runtime/MovieSceneTracks/Private/Tracks/MovieScenePropertyTrack.cpp b/Engine/Source/Runtime/MovieSceneTracks/Private/Tracks/MovieScenePropertyTrack.cpp
--- a/Engine/Source/Runtime/MovieSceneTracks/Private/Tracks/MovieScenePropertyTrack.cpp
+++ b/Engine/Source/Runtime/MovieSceneTracks/Private/Tracks/MovieScenePropertyTrack.cpp
@@ -408,7 +408,10 @@ void FMovieScenePropertyTrackEntityImportHelper::PopulateEvaluationField(UMovieS
 					const FString EditConditionPropertyName = LeafProperty->GetMetaData("EditCondition");
 					if (!EditConditionPropertyName.IsEmpty())
 					{
-						EditConditionPropertyPath.Append(".");
+						if (!EditConditionPropertyPath.IsEmpty())
+						{
+							EditConditionPropertyPath.Append(".");
+						}
 						EditConditionPropertyPath.Append(EditConditionPropertyName);
 						bHasEditCondition = true;
 					}
@@ -426,7 +429,11 @@ void FMovieScenePropertyTrackEntityImportHelper::PopulateEvaluationField(UMovieS
 			FString EditConditionPropertyName = PropertyPathSegments[PropertyPathSegments.Num() - 1];
 			EditConditionPropertyName.InsertAt(0, TEXT("bOverride_"));
 
-			EditConditionPropertyPath.Append(".");
+
+			if (!EditConditionPropertyPath.IsEmpty())
+			{
+				EditConditionPropertyPath.Append(".");
+			}
 			EditConditionPropertyPath.Append(EditConditionPropertyName);
 			bHasEditCondition = true;
 		}
diff --git a/Engine/Source/Runtime/MovieSceneTracks/Public/MovieSceneTracksComponentTypes.h b/Engine/Source/Runtime/MovieSceneTracks/Public/MovieSceneTracksComponentTypes.h
--- a/Engine/Source/Runtime/MovieSceneTracks/Public/MovieSceneTracksComponentTypes.h
+++ b/Engine/Source/Runtime/MovieSceneTracks/Public/MovieSceneTracksComponentTypes.h
@@ -172,6 +172,8 @@ struct FFadeComponentData
 
 struct FFloatPropertyTraits
 {
+	static constexpr bool bIsComposite = false;
+
 	using StorageType  = double;
 	using CustomAccessorStorageType = float;
 	using MetaDataType = TPropertyMetaData<>;
@@ -220,11 +222,6 @@ struct FFloatPropertyTraits
 		const float SetterValue = (float)InValue;
 		FloatTraitsImpl::SetObjectPropertyValue(InObject, PropertyBindings, SetterValue);
 	}
-
-	static float CombineComposites(double InValue)
-	{
-		return (float)InValue;
-	}
 };
 
 struct FColorPropertyTraits
@@ -439,17 +436,91 @@ struct FFloatVectorPropertyTraits
 	}
 };
 
-using FBoolPropertyTraits               = TDirectPropertyTraits<bool>;
-using FBytePropertyTraits               = TDirectPropertyTraits<uint8>;
-using FEnumPropertyTraits               = TDirectPropertyTraits<uint8>;
-using FIntPropertyTraits                = TDirectPropertyTraits<int32>;
-using FDoublePropertyTraits             = TDirectPropertyTraits<double>;
+struct FBoolPropertyTraits
+{
+	static constexpr bool bIsComposite = false;
+
+	struct FBoolMetaData
+	{
+		/** Size of the bitfield/bool property in bytes. Must be one of 1(uint8), 2(uint16), 4(uint32), 8(uint64), or 0 signifying that this is not a bitfield, but a regular bool.. */
+		uint8 BitFieldSize = 0;
+		uint8 BitIndex = 0;
+	};
+	using StorageType  = bool;
+	using MetaDataType = TPropertyMetaData<FBoolMetaData>;
+	using TraitsType   = TDirectPropertyTraits<bool>;
+	using ParamType    = bool;
+
+	/** Property Value Getters  */
+	static void GetObjectPropertyValue(const UObject* InObject, FBoolMetaData MetaData, const FCustomPropertyAccessor& BaseCustomAccessor, bool& OutValue)
+	{
+		const TCustomPropertyAccessor<TraitsType>& CustomAccessor = static_cast<const TCustomPropertyAccessor<TraitsType>&>(BaseCustomAccessor);
+		OutValue = (*CustomAccessor.Functions.Getter)(InObject);
+	}
+	static void GetObjectPropertyValue(const UObject* InObject, FBoolMetaData MetaData, uint16 PropertyOffset, bool& OutValue)
+	{
+		const void* PropertyAddress = reinterpret_cast<const uint8*>(InObject) + PropertyOffset;
+		switch(MetaData.BitFieldSize)
+		{
+		case 0: OutValue = (*reinterpret_cast<const bool  *>(PropertyAddress)); return; // 0 means no bitfield
+		case 1: OutValue = (*reinterpret_cast<const uint8 *>(PropertyAddress) & (uint8 (1u) << MetaData.BitIndex)); return;
+		case 2: OutValue = (*reinterpret_cast<const uint16*>(PropertyAddress) & (uint16(1u) << MetaData.BitIndex)); return;
+		case 4: OutValue = (*reinterpret_cast<const uint32*>(PropertyAddress) & (uint32(1u) << MetaData.BitIndex)); return;
+		case 8: OutValue = (*reinterpret_cast<const uint64*>(PropertyAddress) & (uint64(1u) << MetaData.BitIndex)); return;
+		}
+	}
+	static void GetObjectPropertyValue(const UObject* InObject, FBoolMetaData MetaData, FTrackInstancePropertyBindings* PropertyBindings, bool& OutValue)
+	{
+		OutValue = PropertyBindings->GetCurrentValue<bool>(*InObject);
+	}
+	static void GetObjectPropertyValue(const UObject* InObject, FBoolMetaData MetaData, const FName& PropertyPath, bool& OutValue)
+	{
+		TOptional<bool> Property = FTrackInstancePropertyBindings::StaticValue<bool>(InObject, *PropertyPath.ToString());
+		if (Property)
+		{
+			OutValue = MoveTemp(Property.GetValue());
+		}
+	}
+
+	/** Property Value Setters  */
+	static void SetObjectPropertyValue(UObject* InObject, FBoolMetaData MetaData, const FCustomPropertyAccessor& BaseCustomAccessor, bool InValue)
+	{
+		const TCustomPropertyAccessor<TraitsType>& CustomAccessor = static_cast<const TCustomPropertyAccessor<TraitsType>&>(BaseCustomAccessor);
+		(*CustomAccessor.Functions.Setter)(InObject, InValue);
+	}
+	static void SetObjectPropertyValue(UObject* InObject, FBoolMetaData MetaData, uint16 PropertyOffset, bool InValue)
+	{
+		void* PropertyAddress = reinterpret_cast<uint8*>(InObject) + PropertyOffset;
+
+		// Perform a branchless set by getting the current value, removing the bit, then a bitwise or with InValue in the bit's position.
+		// For example, bit index 4 of a uint8 field:
+		//			f(true)  = (Value & 0b11101111) | 0b00010000
+		//			f(false) = (Value & 0b11101111) | 0b00000000
+		switch(MetaData.BitFieldSize)
+		{
+		case 0: *reinterpret_cast<bool  *>(PropertyAddress) = InValue; return; // 0 means no bitfield
+		case 1: *reinterpret_cast<uint8 *>(PropertyAddress) = (*reinterpret_cast<uint8 *>(PropertyAddress) & ~(uint8 (1u)  << MetaData.BitIndex)) | (uint8 (InValue) << MetaData.BitIndex); return;
+		case 2: *reinterpret_cast<uint16*>(PropertyAddress) = (*reinterpret_cast<uint16*>(PropertyAddress) & ~(uint16(1u)  << MetaData.BitIndex)) | (uint16(InValue) << MetaData.BitIndex); return;
+		case 4: *reinterpret_cast<uint32*>(PropertyAddress) = (*reinterpret_cast<uint32*>(PropertyAddress) & ~(uint32(1u)  << MetaData.BitIndex)) | (uint32(InValue) << MetaData.BitIndex); return;
+		case 8: *reinterpret_cast<uint64*>(PropertyAddress) = (*reinterpret_cast<uint64*>(PropertyAddress) & ~(uint64(1u)  << MetaData.BitIndex)) | (uint64(InValue) << MetaData.BitIndex); return;
+		}
+	}
+	static void SetObjectPropertyValue(UObject* InObject, FBoolMetaData MetaData, FTrackInstancePropertyBindings* PropertyBindings, bool InValue)
+	{
+		PropertyBindings->CallFunction<bool>(*InObject, InValue);
+	}
+};
+
+using FBytePropertyTraits               = TDirectPropertyTraits<uint8, false>;
+using FEnumPropertyTraits               = TDirectPropertyTraits<uint8, false>;
+using FIntPropertyTraits                = TDirectPropertyTraits<int32, false>;
+using FDoublePropertyTraits             = TDirectPropertyTraits<double, false>;
 using FTransformPropertyTraits          = TIndirectPropertyTraits<FTransform, FIntermediate3DTransform>;
 using FEulerTransformPropertyTraits     = TIndirectPropertyTraits<FEulerTransform, FIntermediate3DTransform>;
 using FComponentTransformPropertyTraits = TDirectPropertyTraits<FIntermediate3DTransform>;
 using FStringPropertyTraits			    = TDirectPropertyTraits<FString>;
 
-using FFloatParameterTraits             = TIndirectPropertyTraits<float, double>;
+using FFloatParameterTraits             = TIndirectPropertyTraits<float, double, false>;
 using FColorParameterTraits             = TIndirectPropertyTraits<FLinearColor, FIntermediateColor>;
 
 struct MOVIESCENETRACKS_API FMovieSceneTracksComponentTypes
