Sequencer: Fixed skeletal anim blending across sequences

This change essentially just removes the grouping by IMovieScenePlayer that was previously done, to allow animations from different players to blend together (as expected)

diff --git a/Engine/Source/Runtime/MovieSceneTracks/Private/Systems/MovieSceneSkeletalAnimationSystem.cpp b/Engine/Source/Runtime/MovieSceneTracks/Private/Systems/MovieSceneSkeletalAnimationSystem.cpp
--- a/Engine/Source/Runtime/MovieSceneTracks/Private/Systems/MovieSceneSkeletalAnimationSystem.cpp
+++ b/Engine/Source/Runtime/MovieSceneTracks/Private/Systems/MovieSceneSkeletalAnimationSystem.cpp
@@ -300,7 +300,7 @@ void FBoneTransformFinalizeData::BoneTransformFinalized()
 void FSkeletalAnimationSystemData::ResetSkeletalAnimations()
 {
 	//clear out the delegates
-	for (TTuple<UObject*, FBoundObjectActiveSkeletalAnimations>& Pair : SkeletalAnimations)
+	for (TTuple<USkeletalMeshComponent*, FBoundObjectActiveSkeletalAnimations>& Pair : SkeletalAnimations)
 	{
 		Pair.Value.BoneTransformFinalizeData.Unregister();
 	}
@@ -359,15 +359,53 @@ struct FGatherSkeletalAnimations
 			// Calculate the time at which to evaluate the animation
 			const UMovieSceneSkeletalAnimationSection* AnimSection = SkeletalAnimation.Section;
 			const FMovieSceneSkeletalAnimationParams& AnimParams = AnimSection->Params;
-			if (AnimParams.Animation == nullptr)
+
+			// Get the bound skeletal mesh component.
+			USkeletalMeshComponent* SkeletalMeshComponent = SkeletalMeshComponentFromObject(BoundObject);
+			if (!SkeletalMeshComponent || AnimParams.Animation == nullptr)
 			{
 				continue;
 			}
 			const float EvalTime = AnimParams.MapTimeToAnimation(AnimSection, Context.GetTime(), Context.GetFrameRate());
 			const float PreviousEvalTime = AnimParams.MapTimeToAnimation(AnimSection, Context.GetPreviousTime(), Context.GetFrameRate());
 
-			FBoundObjectActiveSkeletalAnimations& BoundObjectAnimations = SystemData->SkeletalAnimations.FindOrAdd(BoundObject);
-			BoundObjectAnimations.Animations.Add(FActiveSkeletalAnimation{ Player, AnimSection, Context, EntityID, RootInstanceHandle, PreviousEvalTime, EvalTime, Weight, bWantsRestoreState });
+			const FSequenceInstance& RootInstance = InstanceRegistry->GetInstance(RootInstanceHandle);
+			const FMovieSceneContext& RootContext = RootInstance.GetContext();
+			const double RootDeltaTime = (RootContext.HasJumped() ? FFrameTime(0) : RootContext.GetRange().Size<FFrameTime>() ) / RootContext.GetFrameRate();
+
+			const EMovieScenePlayerStatus::Type PlayerStatus = Player->GetPlaybackStatus();
+
+			const bool bResetDynamics = PlayerStatus == EMovieScenePlayerStatus::Stepping || 
+				PlayerStatus == EMovieScenePlayerStatus::Jumping || 
+				PlayerStatus == EMovieScenePlayerStatus::Scrubbing || 
+				(RootDeltaTime == 0.0f && PlayerStatus != EMovieScenePlayerStatus::Stopped);
+
+			const bool bPreviewPlayback = ShouldUsePreviewPlayback(PlayerStatus, *BoundObject);
+
+			// If the playback status is jumping, ie. one such occurrence is setting the time for thumbnail generation, disable anim notifies updates because it could fire audio.
+			// If the playback status is scrubbing, we disable notifies for now because we can't properly fire them in all cases until we get evaluation range info.
+			// We now layer this with the passed in notify toggle to force a disable in this case.
+			const bool bFireNotifies = !bPreviewPlayback || (PlayerStatus != EMovieScenePlayerStatus::Jumping && PlayerStatus != EMovieScenePlayerStatus::Stopped && PlayerStatus != EMovieScenePlayerStatus::Scrubbing);
+			const bool bPlaying = PlayerStatus == EMovieScenePlayerStatus::Playing;
+
+			FBoundObjectActiveSkeletalAnimations& BoundObjectAnimations = SystemData->SkeletalAnimations.FindOrAdd(SkeletalMeshComponent);
+
+			FActiveSkeletalAnimation Animation;
+			Animation.AnimSection        = AnimSection;
+			Animation.Context            = Context;
+			Animation.EntityID           = EntityID;
+			Animation.RootInstanceHandle = RootInstanceHandle;
+			Animation.FromEvalTime       = PreviousEvalTime;
+			Animation.ToEvalTime         = EvalTime;
+			Animation.BlendWeight        = Weight;
+			Animation.PlayerStatus       = PlayerStatus;
+			Animation.bFireNotifies      = bFireNotifies;
+			Animation.bPlaying           = bPlaying;
+			Animation.bResetDynamics     = bResetDynamics;
+			Animation.bWantsRestoreState = bWantsRestoreState;
+			Animation.bPreviewPlayback   = bPreviewPlayback;
+
+			BoundObjectAnimations.Animations.Add(Animation);
 
 			if (FMotionVectorSimulation::IsEnabled())
 			{
@@ -383,9 +421,46 @@ struct FGatherSkeletalAnimations
 
 				const float SimulatedWeight = SimulatedManualWeight * AnimSection->EvaluateEasing(SimulatedTime);
 
-				BoundObjectAnimations.SimulatedAnimations.Add(FActiveSkeletalAnimation{ Player, AnimSection, Context, EntityID, RootInstanceHandle, EvalTime, SimulatedEvalTime, SimulatedWeight, bWantsRestoreState });
+				Animation.BlendWeight = SimulatedWeight;
+				Animation.FromEvalTime = EvalTime;
+				Animation.ToEvalTime = SimulatedEvalTime;
+				BoundObjectAnimations.SimulatedAnimations.Add(Animation);
+			}
+		}
+	}
+
+private:
+
+	static bool ShouldUsePreviewPlayback(EMovieScenePlayerStatus::Type PlayerStatus, UObject& RuntimeObject)
+	{
+		// We also use PreviewSetAnimPosition in PIE when not playing, as we can preview in PIE.
+		bool bIsNotInPIEOrNotPlaying = (RuntimeObject.GetWorld() && !RuntimeObject.GetWorld()->HasBegunPlay()) || PlayerStatus != EMovieScenePlayerStatus::Playing;
+		return GIsEditor && bIsNotInPIEOrNotPlaying;
+	}
+
+	static USkeletalMeshComponent* SkeletalMeshComponentFromObject(UObject* InObject)
+	{
+		// Check if we are bound directly to a skeletal mesh component.
+		USkeletalMeshComponent* SkeletalMeshComponent = Cast<USkeletalMeshComponent>(InObject);
+		if (SkeletalMeshComponent)
+		{
+			return SkeletalMeshComponent;
+		}
+
+		// Then check to see if we are controlling an actor. If so use its first skeletal mesh component.
+		AActor* Actor = Cast<AActor>(InObject);
+		if (!Actor)
+		{
+			if (UChildActorComponent* ChildActorComponent = Cast<UChildActorComponent>(InObject))
+			{
+				Actor = ChildActorComponent->GetChildActor();
 			}
 		}
+		if (Actor)
+		{
+			return Actor->FindComponentByClass<USkeletalMeshComponent>();
+		}
+		return nullptr;
 	}
 };
 
@@ -418,7 +493,7 @@ public:
 	}
 	void Run() const
 	{
-		for (const TTuple<UObject*, FBoundObjectActiveSkeletalAnimations>& Pair : SystemData->SkeletalAnimations)
+		for (const TTuple<USkeletalMeshComponent*, FBoundObjectActiveSkeletalAnimations>& Pair : SystemData->SkeletalAnimations)
 		{
 			EvaluateSkeletalAnimations(Pair.Key, Pair.Value);
 		}
@@ -426,12 +501,10 @@ public:
 
 private:
 
-	void EvaluateSkeletalAnimations(UObject* InObject, const FBoundObjectActiveSkeletalAnimations& InSkeletalAnimations) const
+	void EvaluateSkeletalAnimations(USkeletalMeshComponent* SkeletalMeshComponent, const FBoundObjectActiveSkeletalAnimations& InSkeletalAnimations) const
 	{
-		ensureMsgf(InObject, TEXT("Attempting to evaluate an Animation track with a null object."));
+		ensureMsgf(SkeletalMeshComponent, TEXT("Attempting to evaluate an Animation track with a null object."));
 
-		// Get the bound skeletal mesh component.
-		USkeletalMeshComponent* SkeletalMeshComponent = SkeletalMeshComponentFromObject(InObject);
 		if (!SkeletalMeshComponent || !SkeletalMeshComponent->GetSkeletalMeshAsset())
 		{
 			return;
@@ -464,66 +537,17 @@ private:
 			}
 		}
 
-		// Group skeletal animations by player, in case we have multiple ongoing sequences binding to the same object.
-		TMap<IMovieScenePlayer*, FBoundObjectActiveSkeletalAnimations> SkeletalAnimationsPerPlayer;
-		for (const FActiveSkeletalAnimation& Animation : InSkeletalAnimations.Animations)
-		{
-			FBoundObjectActiveSkeletalAnimations& SkeletalAnimationsForPlayer = SkeletalAnimationsPerPlayer.FindOrAdd(Animation.Player);
-			SkeletalAnimationsForPlayer.Animations.Add(Animation);
-		}
-		for (const FActiveSkeletalAnimation& SimulatedAnimation : InSkeletalAnimations.SimulatedAnimations)
-		{
-			FBoundObjectActiveSkeletalAnimations& SkeletalAnimationsForPlayer = SkeletalAnimationsPerPlayer.FindOrAdd(SimulatedAnimation.Player);
-			SkeletalAnimationsForPlayer.SimulatedAnimations.Add(SimulatedAnimation);
-		}
-
-		// Evaluate each group for the given bound object.
-		for (const TTuple<IMovieScenePlayer*, FBoundObjectActiveSkeletalAnimations>& Pair : SkeletalAnimationsPerPlayer)
-		{
-			EvaluateSkeletalAnimationsForPlayer(*Pair.Key,InObject, SkeletalMeshComponent, ExistingAnimInstance, SequencerInstance, Pair.Value);
-		}
-	}
-
-	void EvaluateSkeletalAnimationsForPlayer(
-			IMovieScenePlayer& InPlayer, 
-			UObject* InObject,
-			USkeletalMeshComponent* InSkeletalMeshComponent, 
-			UAnimInstance* InExistingAnimInstance, 
-			ISequencerAnimationSupport* InSequencerInstance, 
-			const FBoundObjectActiveSkeletalAnimations& InSkeletalAnimations) const
-	{
-		const bool bPreviewPlayback = ShouldUsePreviewPlayback(InPlayer, *InSkeletalMeshComponent);
-
-		const EMovieScenePlayerStatus::Type PlayerStatus = InPlayer.GetPlaybackStatus();
-
-		// If the playback status is jumping, ie. one such occurrence is setting the time for thumbnail generation, disable anim notifies updates because it could fire audio.
-		// If the playback status is scrubbing, we disable notifies for now because we can't properly fire them in all cases until we get evaluation range info.
-		// We now layer this with the passed in notify toggle to force a disable in this case.
-		const bool bFireNotifies = !bPreviewPlayback || (PlayerStatus != EMovieScenePlayerStatus::Jumping && PlayerStatus != EMovieScenePlayerStatus::Stopped && PlayerStatus != EMovieScenePlayerStatus::Scrubbing);
-
-		// When jumping from one cut to another cut, the delta time should be 0 so that anim notifies before the current position are not evaluated. Note, anim notifies at the current time should still be evaluated.
-		const FInstanceHandle RootInstanceHandle = InPlayer.GetEvaluationTemplate().GetRootInstanceHandle();
-		const FSequenceInstance& RootInstance = Linker->GetInstanceRegistry()->GetInstance(RootInstanceHandle);
-		const FMovieSceneContext& RootContext = RootInstance.GetContext();
-		const double RootDeltaTime = (RootContext.HasJumped() ? FFrameTime(0) : RootContext.GetRange().Size<FFrameTime>() ) / RootContext.GetFrameRate();
-
-		const bool bResetDynamics = PlayerStatus == EMovieScenePlayerStatus::Stepping || 
-			PlayerStatus == EMovieScenePlayerStatus::Jumping || 
-			PlayerStatus == EMovieScenePlayerStatus::Scrubbing || 
-			(RootDeltaTime == 0.0f && PlayerStatus != EMovieScenePlayerStatus::Stopped);
-
 		// Need to zero all weights first since we may be blending animation that are keeping state but are no longer active.
-
-		if (InSequencerInstance)
+		if (SequencerInstance)
 		{
-			InSequencerInstance->ResetNodes();
+			SequencerInstance->ResetNodes();
 		}
-		else if (InExistingAnimInstance)
+		else if (ExistingAnimInstance)
 		{
 			for (const TPair<FObjectKey, FMontagePlayerPerSectionData >& Pair : SystemData->MontageData)
 			{
 				int32 InstanceId = Pair.Value.MontageInstanceId;
-				FAnimMontageInstance* MontageInstanceToUpdate = InExistingAnimInstance->GetMontageInstanceForID(InstanceId);
+				FAnimMontageInstance* MontageInstanceToUpdate = ExistingAnimInstance->GetMontageInstanceForID(InstanceId);
 				if (MontageInstanceToUpdate)
 				{
 					MontageInstanceToUpdate->SetDesiredWeight(0.0f);
@@ -534,44 +558,37 @@ private:
 
 		if (InSkeletalAnimations.SimulatedAnimations.Num() != 0)
 		{
-			UMovieSceneMotionVectorSimulationSystem* MotionVectorSim = InPlayer.GetEvaluationTemplate().GetEntitySystemLinker()->FindSystem<UMovieSceneMotionVectorSimulationSystem>();
+			UMovieSceneMotionVectorSimulationSystem* MotionVectorSim = Linker->FindSystem<UMovieSceneMotionVectorSimulationSystem>();
 			if (MotionVectorSim && MotionVectorSim->IsSimulationEnabled())
 			{
-				ApplyAnimations(InPlayer, InObject, InSkeletalMeshComponent, InSkeletalAnimations.SimulatedAnimations, bPreviewPlayback, bFireNotifies, bResetDynamics);
-				InSkeletalMeshComponent->TickAnimation(0.f, false);
-				InSkeletalMeshComponent->ForceMotionVector();
+				ApplyAnimations(SkeletalMeshComponent, InSkeletalAnimations.SimulatedAnimations);
+				SkeletalMeshComponent->TickAnimation(0.f, false);
+				SkeletalMeshComponent->ForceMotionVector();
 		
-				SimulateMotionVectors(InSkeletalMeshComponent, MotionVectorSim);
+				SimulateMotionVectors(SkeletalMeshComponent, MotionVectorSim);
 			}
 		}
 
-		ApplyAnimations(InPlayer, InObject, InSkeletalMeshComponent, InSkeletalAnimations.Animations, bPreviewPlayback, bFireNotifies, bResetDynamics);
+		ApplyAnimations(SkeletalMeshComponent, InSkeletalAnimations.Animations);
 
 		// If the skeletal component has already ticked this frame because tick prerequisites weren't set up yet or a new binding was created, forcibly tick this component to update.
 		// This resolves first frame issues where the skeletal component ticks first, then the sequencer binding is resolved which sets up tick prerequisites
 		// for the next frame.
-		if (InSkeletalMeshComponent->PoseTickedThisFrame() || (InSequencerInstance && InSequencerInstance->GetSourceAnimInstance() != InExistingAnimInstance))
+		if (SkeletalMeshComponent->PoseTickedThisFrame() || (SequencerInstance && SequencerInstance->GetSourceAnimInstance() != ExistingAnimInstance))
 		{
-			InSkeletalMeshComponent->TickAnimation(0.f, false);
-
-			InSkeletalMeshComponent->RefreshBoneTransforms();
-			InSkeletalMeshComponent->RefreshFollowerComponents();
-			InSkeletalMeshComponent->UpdateComponentToWorld();
-			InSkeletalMeshComponent->FinalizeBoneTransform();
-			InSkeletalMeshComponent->MarkRenderTransformDirty();
-			InSkeletalMeshComponent->MarkRenderDynamicDataDirty();
+			SkeletalMeshComponent->TickAnimation(0.f, false);
+
+			SkeletalMeshComponent->RefreshBoneTransforms();
+			SkeletalMeshComponent->RefreshFollowerComponents();
+			SkeletalMeshComponent->UpdateComponentToWorld();
+			SkeletalMeshComponent->FinalizeBoneTransform();
+			SkeletalMeshComponent->MarkRenderTransformDirty();
+			SkeletalMeshComponent->MarkRenderDynamicDataDirty();
 		}
 	}
 
 private:
 
-	static bool ShouldUsePreviewPlayback(IMovieScenePlayer& Player, UObject& RuntimeObject)
-	{
-		// We also use PreviewSetAnimPosition in PIE when not playing, as we can preview in PIE.
-		bool bIsNotInPIEOrNotPlaying = (RuntimeObject.GetWorld() && !RuntimeObject.GetWorld()->HasBegunPlay()) || Player.GetPlaybackStatus() != EMovieScenePlayerStatus::Playing;
-		return GIsEditor && bIsNotInPIEOrNotPlaying;
-	}
-
 	static bool CanPlayAnimation(USkeletalMeshComponent* SkeletalMeshComponent, UAnimSequenceBase* AnimAssetBase)
 	{
 		return (SkeletalMeshComponent->GetSkeletalMeshAsset() &&
@@ -592,31 +609,6 @@ private:
 		return SkelAnimInstance;
 	}
 
-	static USkeletalMeshComponent* SkeletalMeshComponentFromObject(UObject* InObject)
-	{
-		// Check if we are bound directly to a skeletal mesh component.
-		USkeletalMeshComponent* SkeletalMeshComponent = Cast<USkeletalMeshComponent>(InObject);
-		if (SkeletalMeshComponent)
-		{
-			return SkeletalMeshComponent;
-		}
-
-		// Then check to see if we are controlling an actor. If so use its first skeletal mesh component.
-		AActor* Actor = Cast<AActor>(InObject);
-		if (!Actor)
-		{
-			if (UChildActorComponent* ChildActorComponent = Cast<UChildActorComponent>(InObject))
-			{
-				Actor = ChildActorComponent->GetChildActor();
-			}
-		}
-		if (Actor)
-		{
-			return Actor->FindComponentByClass<USkeletalMeshComponent>();
-		}
-		return nullptr;
-	}
-
 private:
 
 	/** Parameter structure for setting the skeletal animation position */
@@ -658,16 +650,9 @@ private:
 	}
 
 	void ApplyAnimations(
-		IMovieScenePlayer& Player,
-		UObject* InObject,
 		USkeletalMeshComponent* SkeletalMeshComponent,
-		TArrayView<const FActiveSkeletalAnimation> SkeletalAnimations,
-		bool bPreviewPlayback,
-		bool bFireNotifies,
-		bool bResetDynamics) const
+		TArrayView<const FActiveSkeletalAnimation> SkeletalAnimations) const
 	{
-		const EMovieScenePlayerStatus::Type PlayerStatus = Player.GetPlaybackStatus();
-
 		for (const FActiveSkeletalAnimation& SkeletalAnimation : SkeletalAnimations)
 		{
 			const UMovieSceneSkeletalAnimationSection* AnimSection = SkeletalAnimation.AnimSection;
@@ -679,7 +664,7 @@ private:
 			{
 				FTransform Transform = RootMotionParams.Transform.IsSet() ? RootMotionParams.Transform.GetValue() : FTransform::Identity;
 				TOptional<FTransform> InitialTransform = GetCurrentTransform(AnimSection->Params.SwapRootBone, SkeletalMeshComponent);
-				if (FBoundObjectActiveSkeletalAnimations* BoundObjectAnimations = SystemData->SkeletalAnimations.Find(InObject)) //InObject may be an AActor
+				if (FBoundObjectActiveSkeletalAnimations* BoundObjectAnimations = SystemData->SkeletalAnimations.Find(SkeletalMeshComponent))
 				{
 					BoundObjectAnimations->BoneTransformFinalizeData.Register(SkeletalMeshComponent, AnimSection->Params.SwapRootBone,Transform, InitialTransform);
 				}
@@ -709,11 +694,11 @@ private:
 			SetAnimPositionParams.ToPosition = SkeletalAnimation.ToEvalTime;
 			SetAnimPositionParams.Weight = SkeletalAnimation.BlendWeight;
 			SetAnimPositionParams.bWantsRestoreState = SkeletalAnimation.bWantsRestoreState;
-			SetAnimPositionParams.bPlaying = (PlayerStatus == EMovieScenePlayerStatus::Playing);
-			SetAnimPositionParams.bFireNotifies = (bFireNotifies && !AnimParams.bSkipAnimNotifiers && !bLooped);
-			SetAnimPositionParams.bResetDynamics = bResetDynamics;
+			SetAnimPositionParams.bPlaying = SkeletalAnimation.bPlaying;
+			SetAnimPositionParams.bFireNotifies = (SkeletalAnimation.bFireNotifies && !AnimParams.bSkipAnimNotifiers && !bLooped);
+			SetAnimPositionParams.bResetDynamics = SkeletalAnimation.bResetDynamics;
 
-			if (bPreviewPlayback)
+			if (SkeletalAnimation.bPreviewPlayback)
 			{
 				PreviewSetAnimPosition(SetAnimPositionParams);
 			}
diff --git a/Engine/Source/Runtime/MovieSceneTracks/Public/Systems/MovieSceneSkeletalAnimationSystem.h b/Engine/Source/Runtime/MovieSceneTracks/Public/Systems/MovieSceneSkeletalAnimationSystem.h
--- a/Engine/Source/Runtime/MovieSceneTracks/Public/Systems/MovieSceneSkeletalAnimationSystem.h
+++ b/Engine/Source/Runtime/MovieSceneTracks/Public/Systems/MovieSceneSkeletalAnimationSystem.h
@@ -19,15 +19,19 @@ namespace UE::MovieScene
 /** Information for a single skeletal animation playing on a bound object */
 struct FActiveSkeletalAnimation
 {
-	IMovieScenePlayer* Player;
 	const UMovieSceneSkeletalAnimationSection* AnimSection;
 	FMovieSceneContext Context;
 	FMovieSceneEntityID EntityID;
 	FRootInstanceHandle RootInstanceHandle;
+	double BlendWeight;
 	float FromEvalTime;
 	float ToEvalTime;
-	double BlendWeight;
-	bool bWantsRestoreState;
+	EMovieScenePlayerStatus::Type PlayerStatus;
+	uint8 bFireNotifies : 1;
+	uint8 bPlaying : 1;
+	uint8 bResetDynamics : 1;
+	uint8 bWantsRestoreState : 1;
+	uint8 bPreviewPlayback : 1;
 };
 
 /** DelegateHandle and Skeletal Mesh for bone transform finalized */
@@ -72,7 +76,7 @@ struct FSkeletalAnimationSystemData
 	void ResetSkeletalAnimations();
 
 	/** Map of active skeletal animations for each bound object */
-	TMap<UObject*, FBoundObjectActiveSkeletalAnimations> SkeletalAnimations;
+	TMap<USkeletalMeshComponent*, FBoundObjectActiveSkeletalAnimations> SkeletalAnimations;
 
 	/** Map of persistent montage data */
 	TMap<FObjectKey, FMontagePlayerPerSectionData> MontageData;
