Sequencer: Fixed erroneous task dependency in the quaternion blender system

diff --git a/Engine/Source/Runtime/MovieSceneTracks/Private/Systems/MovieSceneQuaternionBlenderSystem.cpp b/Engine/Source/Runtime/MovieSceneTracks/Private/Systems/MovieSceneQuaternionBlenderSystem.cpp
--- a/Engine/Source/Runtime/MovieSceneTracks/Private/Systems/MovieSceneQuaternionBlenderSystem.cpp
+++ b/Engine/Source/Runtime/MovieSceneTracks/Private/Systems/MovieSceneQuaternionBlenderSystem.cpp
@@ -243,7 +243,7 @@ void UMovieSceneQuaternionBlenderSystem::OnSchedulePersistentTasks(UE::MovieScen
 	.FilterNone({ BuiltInComponents->Tags.Ignored })
 	.Schedule_PerEntity<FAccumulateAbsoluteWeights>(&Linker->EntityManager, TaskScheduler, &Buffers);
 
-	TaskScheduler->AddPrerequisite(ResetBufferTask, AdditivesTask);
+	TaskScheduler->AddPrerequisite(ResetBufferTask, CollectWeightsTask);
 
 	// Next we accumulate normalized transforms
 	FTaskID AbsolutesTask = FEntityTaskBuilder()
