Sequencer: Fix for 0 duration sounds so that they default to 1 second as opposed to 0 when adding sounds to the sequence.

Noticed at least 1 sound cue in QAGame has -1 duration, so when added to Sequencer, it doesn't show up without this change.

diff --git a/Engine/Source/Runtime/MovieSceneTracks/Private/Sections/MovieSceneAudioSection.cpp b/Engine/Source/Runtime/MovieSceneTracks/Private/Sections/MovieSceneAudioSection.cpp
--- a/Engine/Source/Runtime/MovieSceneTracks/Private/Sections/MovieSceneAudioSection.cpp
+++ b/Engine/Source/Runtime/MovieSceneTracks/Private/Sections/MovieSceneAudioSection.cpp
@@ -326,7 +326,7 @@ TOptional<TRange<FFrameNumber> > UMovieSceneAudioSection::GetAutoSizeRange() con
 	// @todo ^^ Why? Infinte sections would mean there's no starting time?
 	FFrameTime DurationToUse = 1.f * FrameRate; // if all else fails, use 1 second duration
 
-	if (SoundDuration != INDEFINITELY_LOOPING_DURATION)
+	if (SoundDuration != INDEFINITELY_LOOPING_DURATION && SoundDuration > 0)
 	{
 		DurationToUse = FMath::Max(SoundDuration * FrameRate - StartFrameOffset, FFrameTime(1));
 	}
diff --git a/Engine/Source/Runtime/MovieSceneTracks/Private/Tracks/MovieSceneAudioTrack.cpp b/Engine/Source/Runtime/MovieSceneTracks/Private/Tracks/MovieSceneAudioTrack.cpp
--- a/Engine/Source/Runtime/MovieSceneTracks/Private/Tracks/MovieSceneAudioTrack.cpp
+++ b/Engine/Source/Runtime/MovieSceneTracks/Private/Tracks/MovieSceneAudioTrack.cpp
@@ -89,8 +89,8 @@ UMovieSceneSection* UMovieSceneAudioTrack::AddNewSoundOnRow(USoundBase* Sound, F
 	// @todo ^^ Why? Infinte sections would mean there's no starting time?
 	FFrameTime DurationToUse = 1.f * FrameRate; // if all else fails, use 1 second duration
 
-	float SoundDuration = MovieSceneHelpers::GetSoundDuration(Sound);
-	if (SoundDuration != INDEFINITELY_LOOPING_DURATION)
+	const float SoundDuration = MovieSceneHelpers::GetSoundDuration(Sound);
+	if (SoundDuration != INDEFINITELY_LOOPING_DURATION && SoundDuration > 0)
 	{
 		DurationToUse = SoundDuration * FrameRate;
 	}
