[Engine]

diff --git a/Engine/Source/Runtime/MovieSceneTracks/Public/Channels/MovieSceneEventChannel.h b/Engine/Source/Runtime/MovieSceneTracks/Public/Channels/MovieSceneEventChannel.h
--- a/Engine/Source/Runtime/MovieSceneTracks/Public/Channels/MovieSceneEventChannel.h
+++ b/Engine/Source/Runtime/MovieSceneTracks/Public/Channels/MovieSceneEventChannel.h
@@ -22,7 +22,7 @@ struct FKeyHandle;
 
 
 USTRUCT()
-struct MOVIESCENETRACKS_API FMovieSceneEventChannel : public FMovieSceneChannel
+struct FMovieSceneEventChannel : public FMovieSceneChannel
 {
 	GENERATED_BODY()
 
@@ -49,17 +49,17 @@ struct MOVIESCENETRACKS_API FMovieSceneEventChannel : public FMovieSceneChannel
 public:
 
 	// ~ FMovieSceneChannel Interface
-	virtual void GetKeys(const TRange<FFrameNumber>& WithinRange, TArray<FFrameNumber>* OutKeyTimes, TArray<FKeyHandle>* OutKeyHandles) override;
-	virtual void GetKeyTimes(TArrayView<const FKeyHandle> InHandles, TArrayView<FFrameNumber> OutKeyTimes) override;
-	virtual void SetKeyTimes(TArrayView<const FKeyHandle> InHandles, TArrayView<const FFrameNumber> InKeyTimes) override;
-	virtual void DuplicateKeys(TArrayView<const FKeyHandle> InHandles, TArrayView<FKeyHandle> OutNewHandles) override;
-	virtual void DeleteKeys(TArrayView<const FKeyHandle> InHandles) override;
-	virtual void DeleteKeysFrom(FFrameNumber InTime, bool bDeleteKeysBefore) override;
-	virtual void ChangeFrameResolution(FFrameRate SourceRate, FFrameRate DestinationRate) override;
-	virtual TRange<FFrameNumber> ComputeEffectiveRange() const override;
-	virtual int32 GetNumKeys() const override;
-	virtual void Reset() override;
-	virtual void Offset(FFrameNumber DeltaPosition) override;
+	MOVIESCENETRACKS_API virtual void GetKeys(const TRange<FFrameNumber>& WithinRange, TArray<FFrameNumber>* OutKeyTimes, TArray<FKeyHandle>* OutKeyHandles) override;
+	MOVIESCENETRACKS_API virtual void GetKeyTimes(TArrayView<const FKeyHandle> InHandles, TArrayView<FFrameNumber> OutKeyTimes) override;
+	MOVIESCENETRACKS_API virtual void SetKeyTimes(TArrayView<const FKeyHandle> InHandles, TArrayView<const FFrameNumber> InKeyTimes) override;
+	MOVIESCENETRACKS_API virtual void DuplicateKeys(TArrayView<const FKeyHandle> InHandles, TArrayView<FKeyHandle> OutNewHandles) override;
+	MOVIESCENETRACKS_API virtual void DeleteKeys(TArrayView<const FKeyHandle> InHandles) override;
+	MOVIESCENETRACKS_API virtual void DeleteKeysFrom(FFrameNumber InTime, bool bDeleteKeysBefore) override;
+	MOVIESCENETRACKS_API virtual void ChangeFrameResolution(FFrameRate SourceRate, FFrameRate DestinationRate) override;
+	MOVIESCENETRACKS_API virtual TRange<FFrameNumber> ComputeEffectiveRange() const override;
+	MOVIESCENETRACKS_API virtual int32 GetNumKeys() const override;
+	MOVIESCENETRACKS_API virtual void Reset() override;
+	MOVIESCENETRACKS_API virtual void Offset(FFrameNumber DeltaPosition) override;
 
 private:
 
diff --git a/Engine/Source/Runtime/MovieSceneTracks/Public/Channels/MovieSceneFloatPerlinNoiseChannel.h b/Engine/Source/Runtime/MovieSceneTracks/Public/Channels/MovieSceneFloatPerlinNoiseChannel.h
--- a/Engine/Source/Runtime/MovieSceneTracks/Public/Channels/MovieSceneFloatPerlinNoiseChannel.h
+++ b/Engine/Source/Runtime/MovieSceneTracks/Public/Channels/MovieSceneFloatPerlinNoiseChannel.h
@@ -8,13 +8,13 @@
 #include "MovieSceneFloatPerlinNoiseChannel.generated.h"
 
 USTRUCT()
-struct MOVIESCENETRACKS_API FMovieSceneFloatPerlinNoiseChannel : public FMovieSceneChannel
+struct FMovieSceneFloatPerlinNoiseChannel : public FMovieSceneChannel
 {
 	GENERATED_BODY()
 
 	using CurveValueType = float;
 
-	FMovieSceneFloatPerlinNoiseChannel();
+	MOVIESCENETRACKS_API FMovieSceneFloatPerlinNoiseChannel();
 
 	/**
 	 * Evaluate this channel at the given time
@@ -22,7 +22,7 @@ struct MOVIESCENETRACKS_API FMovieSceneFloatPerlinNoiseChannel : public FMovieSc
 	 * @param InSeconds		The time, in seconds, to evaluate at
 	 * @return				The evaluated noise value
 	 */
-	float Evaluate(double InSeconds) const;
+	MOVIESCENETRACKS_API float Evaluate(double InSeconds) const;
 
 	/**
 	 * Evaluate this channel at the given time
@@ -32,7 +32,7 @@ struct MOVIESCENETRACKS_API FMovieSceneFloatPerlinNoiseChannel : public FMovieSc
 	 * @param OutValue		The evaluated noise value
 	 * @return				Whether the noise was successfully evaluated
 	 */
-	bool Evaluate(const UMovieSceneSection* InSection, FFrameTime InTime, float& OutValue) const;
+	MOVIESCENETRACKS_API bool Evaluate(const UMovieSceneSection* InSection, FFrameTime InTime, float& OutValue) const;
 
 	/** The noise parameters */
 	UPROPERTY(EditAnywhere, Category="Perlin Noise")
@@ -46,7 +46,7 @@ struct MOVIESCENETRACKS_API FMovieSceneFloatPerlinNoiseChannel : public FMovieSc
 	 * @params InSeconds		The time at which to evaluate the noise
 	 * @return					The evaluated noise value
 	 */
-	static float Evaluate(const FPerlinNoiseParams& InParams, double InSeconds);
+	static MOVIESCENETRACKS_API float Evaluate(const FPerlinNoiseParams& InParams, double InSeconds);
 };
 
 template<>
diff --git a/Engine/Source/Runtime/MovieSceneTracks/Public/Channels/MovieSceneFloatPerlinNoiseChannelContainer.h b/Engine/Source/Runtime/MovieSceneTracks/Public/Channels/MovieSceneFloatPerlinNoiseChannelContainer.h
--- a/Engine/Source/Runtime/MovieSceneTracks/Public/Channels/MovieSceneFloatPerlinNoiseChannelContainer.h
+++ b/Engine/Source/Runtime/MovieSceneTracks/Public/Channels/MovieSceneFloatPerlinNoiseChannelContainer.h
@@ -10,8 +10,8 @@
 /**
  * Float perlin noise channel overriden container
  */
-UCLASS(meta=(DisplayName="Float Perlin Noise", ToolTip="Override a channel to use float perlin noise"))
-class MOVIESCENETRACKS_API UMovieSceneFloatPerlinNoiseChannelContainer : public UMovieSceneChannelOverrideContainer
+UCLASS(meta=(DisplayName="Float Perlin Noise", ToolTip="Override a channel to use float perlin noise"), MinimalAPI)
+class UMovieSceneFloatPerlinNoiseChannelContainer : public UMovieSceneChannelOverrideContainer
 {
 	GENERATED_BODY()
 
@@ -19,9 +19,9 @@ public:
 
 	using ChannelType = FMovieSceneFloatPerlinNoiseChannel;
 
-	void InitializeOverride(FMovieSceneChannel* InChannel) override;
-	bool SupportsOverride(FName DefaultChannelTypeName) const override;
-	void ImportEntityImpl(
+	MOVIESCENETRACKS_API void InitializeOverride(FMovieSceneChannel* InChannel) override;
+	MOVIESCENETRACKS_API bool SupportsOverride(FName DefaultChannelTypeName) const override;
+	MOVIESCENETRACKS_API void ImportEntityImpl(
 			const UE::MovieScene::FChannelOverrideEntityImportParams& OverrideParams, 
 			const UE::MovieScene::FEntityImportParams& ImportParams, UE::MovieScene::FImportedEntity* OutImportedEntity) override;
 
@@ -29,9 +29,9 @@ public:
 	FMovieSceneChannel* GetChannel() override { return &PerlinNoiseChannel; }
 
 #if WITH_EDITOR
-	FMovieSceneChannelHandle AddChannelProxy(FName ChannelName, FMovieSceneChannelProxyData& ProxyData, const FMovieSceneChannelMetaData& MetaData) override;
+	MOVIESCENETRACKS_API FMovieSceneChannelHandle AddChannelProxy(FName ChannelName, FMovieSceneChannelProxyData& ProxyData, const FMovieSceneChannelMetaData& MetaData) override;
 #else
-	void AddChannelProxy(FName ChannelName, FMovieSceneChannelProxyData& ProxyData) override;
+	MOVIESCENETRACKS_API void AddChannelProxy(FName ChannelName, FMovieSceneChannelProxyData& ProxyData) override;
 #endif
 
 private:
diff --git a/Engine/Source/Runtime/MovieSceneTracks/Public/Channels/MovieSceneStringChannel.h b/Engine/Source/Runtime/MovieSceneTracks/Public/Channels/MovieSceneStringChannel.h
--- a/Engine/Source/Runtime/MovieSceneTracks/Public/Channels/MovieSceneStringChannel.h
+++ b/Engine/Source/Runtime/MovieSceneTracks/Public/Channels/MovieSceneStringChannel.h
@@ -27,7 +27,7 @@ struct FKeyHandle;
 struct FPropertyTag;
 
 USTRUCT()
-struct MOVIESCENETRACKS_API FMovieSceneStringChannel : public FMovieSceneChannel
+struct FMovieSceneStringChannel : public FMovieSceneChannel
 {
 	GENERATED_BODY()
 
@@ -38,7 +38,7 @@ struct MOVIESCENETRACKS_API FMovieSceneStringChannel : public FMovieSceneChannel
 	/**
 	 * Serialize this type from another
 	 */
-	bool SerializeFromMismatchedTag(const FPropertyTag& Tag, FStructuredArchive::FSlot Slot);
+	MOVIESCENETRACKS_API bool SerializeFromMismatchedTag(const FPropertyTag& Tag, FStructuredArchive::FSlot Slot);
 
 	/**
 	 * Access a mutable interface for this channel's data
@@ -82,7 +82,7 @@ struct MOVIESCENETRACKS_API FMovieSceneStringChannel : public FMovieSceneChannel
 	 * @param InTime     The time to evaluate at
 	 * @return A pointer to the string, or nullptr
 	 */
-	const FString* Evaluate(FFrameTime InTime) const;
+	MOVIESCENETRACKS_API const FString* Evaluate(FFrameTime InTime) const;
 
 	/**
 	 * Add keys with these times to channel. The number of elements in both arrays much match or nothing is added.
@@ -99,19 +99,19 @@ struct MOVIESCENETRACKS_API FMovieSceneStringChannel : public FMovieSceneChannel
 public:
 
 	// ~ FMovieSceneChannel Interface
-	virtual void GetKeys(const TRange<FFrameNumber>& WithinRange, TArray<FFrameNumber>* OutKeyTimes, TArray<FKeyHandle>* OutKeyHandles) override;
-	virtual void GetKeyTimes(TArrayView<const FKeyHandle> InHandles, TArrayView<FFrameNumber> OutKeyTimes) override;
-	virtual void SetKeyTimes(TArrayView<const FKeyHandle> InHandles, TArrayView<const FFrameNumber> InKeyTimes) override;
-	virtual void DuplicateKeys(TArrayView<const FKeyHandle> InHandles, TArrayView<FKeyHandle> OutNewHandles) override;
-	virtual void DeleteKeys(TArrayView<const FKeyHandle> InHandles) override;
-	virtual void DeleteKeysFrom(FFrameNumber InTime, bool bDeleteKeysBefore) override;
-	virtual void ChangeFrameResolution(FFrameRate SourceRate, FFrameRate DestinationRate) override;
-	virtual TRange<FFrameNumber> ComputeEffectiveRange() const override;
-	virtual int32 GetNumKeys() const override;
-	virtual void Reset() override;
-	virtual void Offset(FFrameNumber DeltaPosition) override;
-	virtual void Optimize(const FKeyDataOptimizationParams& InParameters) override;
-	virtual void ClearDefault() override;
+	MOVIESCENETRACKS_API virtual void GetKeys(const TRange<FFrameNumber>& WithinRange, TArray<FFrameNumber>* OutKeyTimes, TArray<FKeyHandle>* OutKeyHandles) override;
+	MOVIESCENETRACKS_API virtual void GetKeyTimes(TArrayView<const FKeyHandle> InHandles, TArrayView<FFrameNumber> OutKeyTimes) override;
+	MOVIESCENETRACKS_API virtual void SetKeyTimes(TArrayView<const FKeyHandle> InHandles, TArrayView<const FFrameNumber> InKeyTimes) override;
+	MOVIESCENETRACKS_API virtual void DuplicateKeys(TArrayView<const FKeyHandle> InHandles, TArrayView<FKeyHandle> OutNewHandles) override;
+	MOVIESCENETRACKS_API virtual void DeleteKeys(TArrayView<const FKeyHandle> InHandles) override;
+	MOVIESCENETRACKS_API virtual void DeleteKeysFrom(FFrameNumber InTime, bool bDeleteKeysBefore) override;
+	MOVIESCENETRACKS_API virtual void ChangeFrameResolution(FFrameRate SourceRate, FFrameRate DestinationRate) override;
+	MOVIESCENETRACKS_API virtual TRange<FFrameNumber> ComputeEffectiveRange() const override;
+	MOVIESCENETRACKS_API virtual int32 GetNumKeys() const override;
+	MOVIESCENETRACKS_API virtual void Reset() override;
+	MOVIESCENETRACKS_API virtual void Offset(FFrameNumber DeltaPosition) override;
+	MOVIESCENETRACKS_API virtual void Optimize(const FKeyDataOptimizationParams& InParameters) override;
+	MOVIESCENETRACKS_API virtual void ClearDefault() override;
 
 public:
 
diff --git a/Engine/Source/Runtime/MovieSceneTracks/Public/EntitySystem/Interrogation/MovieSceneInterrogatedPropertyInstantiator.h b/Engine/Source/Runtime/MovieSceneTracks/Public/EntitySystem/Interrogation/MovieSceneInterrogatedPropertyInstantiator.h
--- a/Engine/Source/Runtime/MovieSceneTracks/Public/EntitySystem/Interrogation/MovieSceneInterrogatedPropertyInstantiator.h
+++ b/Engine/Source/Runtime/MovieSceneTracks/Public/EntitySystem/Interrogation/MovieSceneInterrogatedPropertyInstantiator.h
@@ -14,8 +14,8 @@
 class UMovieSceneBlenderSystem;
 
 /** Class responsible for resolving all property types registered with FBuiltInComponentTypes::PropertyRegistry */
-UCLASS()
-class MOVIESCENETRACKS_API UMovieSceneInterrogatedPropertyInstantiatorSystem
+UCLASS(MinimalAPI)
+class UMovieSceneInterrogatedPropertyInstantiatorSystem
 	: public UMovieSceneEntityInstantiatorSystem
 {
 public:
@@ -25,7 +25,7 @@ public:
 
 	GENERATED_BODY()
 
-	UMovieSceneInterrogatedPropertyInstantiatorSystem(const FObjectInitializer& ObjInit);
+	MOVIESCENETRACKS_API UMovieSceneInterrogatedPropertyInstantiatorSystem(const FObjectInitializer& ObjInit);
 
 	/**
 	 * Recompose a value from the constituent parts specified in InQuery, taking into accounts the wieghtings of the specific channel defined by ChannelCompositeIndex.
@@ -36,7 +36,7 @@ public:
 	 * @param InCurrentValue The value of the property to recompose
 	 * @return A result containing the recomposed value for each of the entities specified in InQuery
 	 */
-	FValueRecompositionResult RecomposeBlendChannel(const UE::MovieScene::FPropertyDefinition& PropertyDefinition, int32 ChannelCompositeIndex, const UE::MovieScene::FDecompositionQuery& InQuery, double InCurrentValue);
+	MOVIESCENETRACKS_API FValueRecompositionResult RecomposeBlendChannel(const UE::MovieScene::FPropertyDefinition& PropertyDefinition, int32 ChannelCompositeIndex, const UE::MovieScene::FDecompositionQuery& InQuery, double InCurrentValue);
 
 	/**
 	 * Recompose a value from the constituent parts specified in InQuery, taking into account the weightings of each channel.
@@ -90,9 +90,9 @@ public:
 	};
 
 	// TOverlappingEntityTracker handler interface
-	void InitializeOutput(UE::MovieScene::FInterrogationKey Key, TArrayView<const FMovieSceneEntityID> Inputs, FPropertyInfo* Output, UE::MovieScene::FEntityOutputAggregate Aggregate);
-	void UpdateOutput(UE::MovieScene::FInterrogationKey Key, TArrayView<const FMovieSceneEntityID> Inputs, FPropertyInfo* Output, UE::MovieScene::FEntityOutputAggregate Aggregate);
-	void DestroyOutput(UE::MovieScene::FInterrogationKey Key, FPropertyInfo* Output, UE::MovieScene::FEntityOutputAggregate Aggregate);
+	MOVIESCENETRACKS_API void InitializeOutput(UE::MovieScene::FInterrogationKey Key, TArrayView<const FMovieSceneEntityID> Inputs, FPropertyInfo* Output, UE::MovieScene::FEntityOutputAggregate Aggregate);
+	MOVIESCENETRACKS_API void UpdateOutput(UE::MovieScene::FInterrogationKey Key, TArrayView<const FMovieSceneEntityID> Inputs, FPropertyInfo* Output, UE::MovieScene::FEntityOutputAggregate Aggregate);
+	MOVIESCENETRACKS_API void DestroyOutput(UE::MovieScene::FInterrogationKey Key, FPropertyInfo* Output, UE::MovieScene::FEntityOutputAggregate Aggregate);
 
 	const FPropertyInfo* FindPropertyInfo(UE::MovieScene::FInterrogationKey Key) const
 	{
@@ -106,12 +106,12 @@ public:
 
 private:
 
-	virtual void OnRun(FSystemTaskPrerequisites& InPrerequisites, FSystemSubsequentTasks& Subsequents) override;
-	virtual bool IsRelevantImpl(UMovieSceneEntitySystemLinker* InLinker) const override;
+	MOVIESCENETRACKS_API virtual void OnRun(FSystemTaskPrerequisites& InPrerequisites, FSystemSubsequentTasks& Subsequents) override;
+	MOVIESCENETRACKS_API virtual bool IsRelevantImpl(UMovieSceneEntitySystemLinker* InLinker) const override;
 
-	bool PropertySupportsFastPath(TArrayView<const FMovieSceneEntityID> Inputs, FPropertyInfo* Output) const;
-	UClass* ResolveBlenderClass(const UE::MovieScene::FPropertyDefinition& PropertyDefinition, TArrayView<const FMovieSceneEntityID> Inputs) const;
-	UE::MovieScene::FPropertyRecomposerPropertyInfo FindPropertyFromSource(FMovieSceneEntityID EntityID, UObject* Object) const;
+	MOVIESCENETRACKS_API bool PropertySupportsFastPath(TArrayView<const FMovieSceneEntityID> Inputs, FPropertyInfo* Output) const;
+	MOVIESCENETRACKS_API UClass* ResolveBlenderClass(const UE::MovieScene::FPropertyDefinition& PropertyDefinition, TArrayView<const FMovieSceneEntityID> Inputs) const;
+	MOVIESCENETRACKS_API UE::MovieScene::FPropertyRecomposerPropertyInfo FindPropertyFromSource(FMovieSceneEntityID EntityID, UObject* Object) const;
 
 private:
 
diff --git a/Engine/Source/Runtime/MovieSceneTracks/Public/EntitySystem/Interrogation/MovieSceneInterrogationLinker.h b/Engine/Source/Runtime/MovieSceneTracks/Public/EntitySystem/Interrogation/MovieSceneInterrogationLinker.h
--- a/Engine/Source/Runtime/MovieSceneTracks/Public/EntitySystem/Interrogation/MovieSceneInterrogationLinker.h
+++ b/Engine/Source/Runtime/MovieSceneTracks/Public/EntitySystem/Interrogation/MovieSceneInterrogationLinker.h
@@ -68,14 +68,14 @@ struct FInterrogationParams
  *        MySystem->Interrogate(FInterogationKey::Default(), OutData);
  *    }
  */
-class MOVIESCENETRACKS_API FInterrogationChannels
+class FInterrogationChannels
 {
 public:
 
-	FInterrogationChannels();
-	~FInterrogationChannels();
+	MOVIESCENETRACKS_API FInterrogationChannels();
+	MOVIESCENETRACKS_API ~FInterrogationChannels();
 
-	void Reset();
+	MOVIESCENETRACKS_API void Reset();
 
 	const FSparseInterrogationChannelInfo& GetSparseChannelInfo() const
 	{
@@ -88,19 +88,19 @@ public:
 	 * @param ParentChannel     The channel that should be considered this channel's parent, or FInterrogationChannel::Invalid if there is none
 	 * @return A new interrogation channel
 	 */
-	FInterrogationChannel AllocateChannel(FInterrogationChannel ParentChannel, const FMovieScenePropertyBinding& PropertyBinding);
+	MOVIESCENETRACKS_API FInterrogationChannel AllocateChannel(FInterrogationChannel ParentChannel, const FMovieScenePropertyBinding& PropertyBinding);
 
 
 	/**
 	 * Allocate a new interrogation channel that relates to a specific object
 	 */
-	FInterrogationChannel AllocateChannel(UObject* Object, const FMovieScenePropertyBinding& PropertyBinding);
+	MOVIESCENETRACKS_API FInterrogationChannel AllocateChannel(UObject* Object, const FMovieScenePropertyBinding& PropertyBinding);
 
 
 	/**
 	 * Allocate a new interrogation channel that relates to a specific object
 	 */
-	FInterrogationChannel AllocateChannel(UObject* Object, FInterrogationChannel ParentChannel, const FMovieScenePropertyBinding& PropertyBinding);
+	MOVIESCENETRACKS_API FInterrogationChannel AllocateChannel(UObject* Object, FInterrogationChannel ParentChannel, const FMovieScenePropertyBinding& PropertyBinding);
 
 
 	/**
@@ -110,7 +110,7 @@ public:
 	 * @param CurrentValue      A value to use if this channel has now animated data after the interrogation
 	 * @return A new interrogation channel that can be passed to ImportUnboundTrack
 	 */
-	FInterrogationChannel AllocateUnboundChannel(FInterrogationChannel ParentChannel, const FTransform& CurrentValueLocalSpace);
+	MOVIESCENETRACKS_API FInterrogationChannel AllocateUnboundChannel(FInterrogationChannel ParentChannel, const FTransform& CurrentValueLocalSpace);
 
 
 	/**
@@ -121,7 +121,7 @@ public:
 	 * @param SequenceID        The current sequence ID for the interrogation
 	 * @return The channel that was either pre-existing or allocated for SceneComponent
 	 */
-	FInterrogationChannel ImportTransformHierarchy(USceneComponent* SceneComponent);
+	MOVIESCENETRACKS_API FInterrogationChannel ImportTransformHierarchy(USceneComponent* SceneComponent);
 
 
 	/**
@@ -130,7 +130,7 @@ public:
 	 * @param Object    The object to find a channel for
 	 * @return The (potentially invalid) interrogation channel that the specified object is on
 	 */
-	FInterrogationChannel FindChannel(UObject* Object);
+	MOVIESCENETRACKS_API FInterrogationChannel FindChannel(UObject* Object);
 
 
 	/**
@@ -138,14 +138,14 @@ public:
 	 *
 	 * @param InChannel The channel to activate
 	 */
-	void ActivateChannel(FInterrogationChannel InChannel);
+	MOVIESCENETRACKS_API void ActivateChannel(FInterrogationChannel InChannel);
 
 	/**
 	 * Called to deactivate the specified channel. Inactive channels will only ever use their object's current value.
 	 *
 	 * @param InChannel The channel to deactivate
 	 */
-	void DeactivateChannel(FInterrogationChannel InChannel);
+	MOVIESCENETRACKS_API void DeactivateChannel(FInterrogationChannel InChannel);
 
 
 	/**
@@ -154,7 +154,7 @@ public:
 	 * @param Params    The desired time to interrogate at
 	 * @return A unique index identifier for the specified time, or INDEX_NONE if the maximum number have been reached
 	 */
-	int32 AddInterrogation(const FInterrogationParams& Params);
+	MOVIESCENETRACKS_API int32 AddInterrogation(const FInterrogationParams& Params);
 
 public:
 
@@ -184,7 +184,7 @@ public:
 	 * @param SceneComponent    The scene component to query
 	 * @param OutTransforms     Array to output transforms into, one per Interrogation
 	 */
-	void QueryLocalSpaceTransforms(UMovieSceneEntitySystemLinker* Linker, USceneComponent* SceneComponent, TArray<FIntermediate3DTransform>& OutTransforms) const;
+	MOVIESCENETRACKS_API void QueryLocalSpaceTransforms(UMovieSceneEntitySystemLinker* Linker, USceneComponent* SceneComponent, TArray<FIntermediate3DTransform>& OutTransforms) const;
 
 	/**
 	 * Query local space transforms
@@ -192,14 +192,14 @@ public:
 	 * @param InChannel         The channel to query
 	 * @param OutTransforms     Array to output transforms into, one per Interrogation
 	 */
-	void QueryLocalSpaceTransforms(UMovieSceneEntitySystemLinker* Linker, FInterrogationChannel InChannel, TArray<FIntermediate3DTransform>& OutTransforms) const;
+	MOVIESCENETRACKS_API void QueryLocalSpaceTransforms(UMovieSceneEntitySystemLinker* Linker, FInterrogationChannel InChannel, TArray<FIntermediate3DTransform>& OutTransforms) const;
 
 	/**
 	 * Query all local space transforms, even including channels that do not have any variable track data
 	 *
 	 * @param OutTransformsByChannel  Sparse array to receive transforms allocated by their Channel index
 	 */
-	void QueryLocalSpaceTransforms(UMovieSceneEntitySystemLinker* Linker, TSparseArray<TArray<FIntermediate3DTransform>>& OutTransformsByChannel) const;
+	MOVIESCENETRACKS_API void QueryLocalSpaceTransforms(UMovieSceneEntitySystemLinker* Linker, TSparseArray<TArray<FIntermediate3DTransform>>& OutTransformsByChannel) const;
 
 	/**
 	 * Query a specific set of channels for their local space transforms as defined by set bits within ChannelsToQuery
@@ -207,7 +207,7 @@ public:
 	 * @param ChannelsToQuery         Bit array containing set bits for each channel to query
 	 * @param OutTransformsByChannel  Sparse array to receive transforms allocated by their Channel index
 	 */
-	void QueryLocalSpaceTransforms(UMovieSceneEntitySystemLinker* Linker, const TBitArray<>& ChannelsToQuery, TSparseArray<TArray<FIntermediate3DTransform>>& OutTransformsByChannel) const;
+	MOVIESCENETRACKS_API void QueryLocalSpaceTransforms(UMovieSceneEntitySystemLinker* Linker, const TBitArray<>& ChannelsToQuery, TSparseArray<TArray<FIntermediate3DTransform>>& OutTransformsByChannel) const;
 
 public:
 
@@ -217,7 +217,7 @@ public:
 	 * @param SceneComponent    The scene component to query
 	 * @param OutTransforms     Array to output transforms into, one per Interrogation
 	 */
-	void QueryWorldSpaceTransforms(UMovieSceneEntitySystemLinker* Linker, USceneComponent* SceneComponent, TArray<FTransform>& OutTransforms) const;
+	MOVIESCENETRACKS_API void QueryWorldSpaceTransforms(UMovieSceneEntitySystemLinker* Linker, USceneComponent* SceneComponent, TArray<FTransform>& OutTransforms) const;
 
 	/**
 	 * Query world space transforms for a channel
@@ -225,14 +225,14 @@ public:
 	 * @param InChannel         The channel to query
 	 * @param OutTransforms     Array to output transforms into, one per Interrogation
 	 */
-	void QueryWorldSpaceTransforms(UMovieSceneEntitySystemLinker* Linker, FInterrogationChannel InChannel, TArray<FTransform>& OutTransforms) const;
+	MOVIESCENETRACKS_API void QueryWorldSpaceTransforms(UMovieSceneEntitySystemLinker* Linker, FInterrogationChannel InChannel, TArray<FTransform>& OutTransforms) const;
 
 	/**
 	 * Query all world space transforms, even including channels that do not have any variable track data
 	 *
 	 * @param OutTransformsByChannel  Sparse array to receive transforms allocated by their Channel index
 	 */
-	void QueryWorldSpaceTransforms(UMovieSceneEntitySystemLinker* Linker, TSparseArray<TArray<FTransform>>& OutTransformsByChannel) const;
+	MOVIESCENETRACKS_API void QueryWorldSpaceTransforms(UMovieSceneEntitySystemLinker* Linker, TSparseArray<TArray<FTransform>>& OutTransformsByChannel) const;
 
 	/**
 	 * Query a specific set of channels for their world space transforms as defined by set bits within ChannelsToQuery
@@ -240,7 +240,7 @@ public:
 	 * @param ChannelsToQuery         Bit array containing set bits for each channel to query
 	 * @param OutTransformsByChannel  Sparse array to receive transforms allocated by their Channel index
 	 */
-	void QueryWorldSpaceTransforms(UMovieSceneEntitySystemLinker* Linker, const TBitArray<>& ChannelsToQuery, TSparseArray<TArray<FTransform>>& OutTransformsByChannel) const;
+	MOVIESCENETRACKS_API void QueryWorldSpaceTransforms(UMovieSceneEntitySystemLinker* Linker, const TBitArray<>& ChannelsToQuery, TSparseArray<TArray<FTransform>>& OutTransformsByChannel) const;
 
 private:
 
@@ -268,25 +268,25 @@ protected:
 
 
 
-class MOVIESCENETRACKS_API FSystemInterrogator
+class FSystemInterrogator
 	: FGCObject
 	, IInterrogationExtension
 {
 public:
-	FSystemInterrogator();
-	~FSystemInterrogator();
+	MOVIESCENETRACKS_API FSystemInterrogator();
+	MOVIESCENETRACKS_API ~FSystemInterrogator();
 
 
 	/**
 	 * Gets the custom system category for interrogation-specific systems
 	 */
-	static EEntitySystemCategory GetInterrogationCategory();
+	static MOVIESCENETRACKS_API EEntitySystemCategory GetInterrogationCategory();
 
 
 	/**
 	 * Gets the custom system category for systems who should be excluded from interrogation linkers
 	 */
-	static EEntitySystemCategory GetExcludedFromInterrogationCategory();
+	static MOVIESCENETRACKS_API EEntitySystemCategory GetExcludedFromInterrogationCategory();
 
 	
 	/**
@@ -306,7 +306,7 @@ public:
 	 * @param Track      The track to import
 	 * @param InChannel  The channel to import this track onto. FInterrogationChannel::Default() can be used if this interrogator is only being used for a single output.
 	 */
-	void ImportTrack(UMovieSceneTrack* Track, FInterrogationChannel InChannel);
+	MOVIESCENETRACKS_API void ImportTrack(UMovieSceneTrack* Track, FInterrogationChannel InChannel);
 
 
 	/**
@@ -318,7 +318,7 @@ public:
 	 * @param ObjectBindingID  The binding ID for the object binding that this track resides within
 	 * @param InChannel        The channel to import this track onto. FInterrogationChannel::Default() can be used if this interrogator is only being used for a single output.
 	 */
-	void ImportTrack(UMovieSceneTrack* Track, const FGuid& ObjectBindingID, FInterrogationChannel InChannel);
+	MOVIESCENETRACKS_API void ImportTrack(UMovieSceneTrack* Track, const FGuid& ObjectBindingID, FInterrogationChannel InChannel);
 
 
 	/**
@@ -329,7 +329,7 @@ public:
 	 * @param SequenceID        The current sequence ID for the interrogation
 	 * @return The channel that was either pre-existing or allocated for SceneComponent
 	 */
-	FInterrogationChannel ImportTransformHierarchy(USceneComponent* SceneComponent, IMovieScenePlayer* InPlayer, FMovieSceneSequenceID SequenceID);
+	MOVIESCENETRACKS_API FInterrogationChannel ImportTransformHierarchy(USceneComponent* SceneComponent, IMovieScenePlayer* InPlayer, FMovieSceneSequenceID SequenceID);
 
 
 	/**
@@ -340,7 +340,7 @@ public:
 	 * @param SequenceID        The current sequence ID for the interrogation
 	 * @return The channel that was either pre-existing or allocated for SceneComponent
 	 */
-	FInterrogationChannel ImportLocalTransforms(USceneComponent* SceneComponent, IMovieScenePlayer* InPlayer, FMovieSceneSequenceID SequenceID);
+	MOVIESCENETRACKS_API FInterrogationChannel ImportLocalTransforms(USceneComponent* SceneComponent, IMovieScenePlayer* InPlayer, FMovieSceneSequenceID SequenceID);
 
 
 	/**
@@ -373,19 +373,19 @@ public:
 	 * @param Params    The desired time to interrogate at
 	 * @return A unique index identifier for the specified time, or INDEX_NONE if the maximum number have been reached
 	 */
-	int32 AddInterrogation(const FInterrogationParams& Params);
+	MOVIESCENETRACKS_API int32 AddInterrogation(const FInterrogationParams& Params);
 
 
 	/**
 	 * Flush this interrogator by running all the systems relevant to the current data and populating the interrogation outputs.
 	 */
-	void Update();
+	MOVIESCENETRACKS_API void Update();
 
 
 	/**
 	 * Reset this linker back to its original state
 	 */
-	void Reset();
+	MOVIESCENETRACKS_API void Reset();
 
 public:
 
@@ -393,14 +393,14 @@ public:
 	 * Indicate that consumers of this class require a reverse-lookup table for imported entities to be maintained such that the various FindEntity functions can be called.
 	 * (Not enabled by default due to performance cost with high interrogation counts)
 	 */
-	void TrackImportedEntities(bool bInTrackImportedEntities);
+	MOVIESCENETRACKS_API void TrackImportedEntities(bool bInTrackImportedEntities);
 
 
 	/**
 	 * Find an entity given the entity's owner.
 	 * @note: Must call TrackImportedEntities(true) prior to calling ImportTrack for this function to return the correct entity
 	 */
-	FMovieSceneEntityID FindEntityFromOwner(FInterrogationKey InterrogationKey, UObject* Owner, uint32 EntityID) const;
+	MOVIESCENETRACKS_API FMovieSceneEntityID FindEntityFromOwner(FInterrogationKey InterrogationKey, UObject* Owner, uint32 EntityID) const;
 
 
 	/**
@@ -431,7 +431,7 @@ public:
 
 private:
 
-	void FindPropertyOutputEntityIDs(const FPropertyDefinition& PropertyDefinition, FInterrogationChannel Channel, TArray<FMovieSceneEntityID>& OutEntityIDs) const;
+	MOVIESCENETRACKS_API void FindPropertyOutputEntityIDs(const FPropertyDefinition& PropertyDefinition, FInterrogationChannel Channel, TArray<FMovieSceneEntityID>& OutEntityIDs) const;
 
 public:
 
@@ -576,18 +576,18 @@ private:
 	/**
 	 * Import transform tracks from this binding
 	 */
-	void ImportTransformTracks(const FMovieSceneBinding& Binding, FInterrogationChannel Channel);
+	MOVIESCENETRACKS_API void ImportTransformTracks(const FMovieSceneBinding& Binding, FInterrogationChannel Channel);
 
 	/**
 	 * Import an entity for the specified query and index
 	 */
-	void InterrogateEntity(int32 InterrogationIndex, const FMovieSceneEvaluationFieldEntityQuery& Query);
+	MOVIESCENETRACKS_API void InterrogateEntity(int32 InterrogationIndex, const FMovieSceneEvaluationFieldEntityQuery& Query);
 
 
 private:
 
-	virtual void AddReferencedObjects(FReferenceCollector& Collector);
-	virtual FString GetReferencerName() const;
+	MOVIESCENETRACKS_API virtual void AddReferencedObjects(FReferenceCollector& Collector);
+	MOVIESCENETRACKS_API virtual FString GetReferencerName() const;
 
 protected:
 
diff --git a/Engine/Source/Runtime/MovieSceneTracks/Public/Evaluation/MovieSceneCameraShakePreviewer.h b/Engine/Source/Runtime/MovieSceneTracks/Public/Evaluation/MovieSceneCameraShakePreviewer.h
--- a/Engine/Source/Runtime/MovieSceneTracks/Public/Evaluation/MovieSceneCameraShakePreviewer.h
+++ b/Engine/Source/Runtime/MovieSceneTracks/Public/Evaluation/MovieSceneCameraShakePreviewer.h
@@ -38,24 +38,24 @@ public:
 /**
  * A class that owns a gameplay camera shake manager, so that we can us it to preview shakes in editor.
  */
-class MOVIESCENETRACKS_API FCameraShakePreviewer : public FGCObject
+class FCameraShakePreviewer : public FGCObject
 {
 public:
-	FCameraShakePreviewer();
-	~FCameraShakePreviewer();
+	MOVIESCENETRACKS_API FCameraShakePreviewer();
+	MOVIESCENETRACKS_API ~FCameraShakePreviewer();
 
-	void Initialize(UWorld* InWorld);
+	MOVIESCENETRACKS_API void Initialize(UWorld* InWorld);
 	bool IsInitialized() const { return PreviewCameraShake != nullptr; }
-	void Teardown();
+	MOVIESCENETRACKS_API void Teardown();
 
 	APlayerCameraManager* GetCameraManager() const { return PreviewCamera; }
 	UCameraModifier_CameraShake* GetCameraModifier() const { return PreviewCameraShake; }
 
-	void ModifyView(FEditorViewportViewModifierParams& Params);
+	MOVIESCENETRACKS_API void ModifyView(FEditorViewportViewModifierParams& Params);
 
-	void RegisterViewModifier();
-	void UnRegisterViewModifier();
-	void Update(float DeltaTime, bool bIsPlaying);
+	MOVIESCENETRACKS_API void RegisterViewModifier();
+	MOVIESCENETRACKS_API void UnRegisterViewModifier();
+	MOVIESCENETRACKS_API void Update(float DeltaTime, bool bIsPlaying);
 
 private:
 	// FGCObject interface
@@ -63,8 +63,8 @@ private:
 	virtual FString GetReferencerName() const override { return TEXT("SCameraShakePreviewer"); }
 
 private:
-	void OnModifyView(FEditorViewportViewModifierParams& Params);
-	void OnLevelViewportClientListChanged();
+	MOVIESCENETRACKS_API void OnModifyView(FEditorViewportViewModifierParams& Params);
+	MOVIESCENETRACKS_API void OnLevelViewportClientListChanged();
 
 private:
 	APreviewPlayerCameraManager* PreviewCamera;
diff --git a/Engine/Source/Runtime/MovieSceneTracks/Public/Evaluation/MovieSceneCameraShakeTemplate.h b/Engine/Source/Runtime/MovieSceneTracks/Public/Evaluation/MovieSceneCameraShakeTemplate.h
--- a/Engine/Source/Runtime/MovieSceneTracks/Public/Evaluation/MovieSceneCameraShakeTemplate.h
+++ b/Engine/Source/Runtime/MovieSceneTracks/Public/Evaluation/MovieSceneCameraShakeTemplate.h
@@ -18,8 +18,8 @@ struct FMovieSceneContext;
  * Factory methods for these evaluators are registered on the FMovieSceneCameraShakeEvaluatorRegistry. An evaluator is created for
  * each camera shake running in a sequence.
  */
-UCLASS()
-class MOVIESCENETRACKS_API UMovieSceneCameraShakeEvaluator : public UObject
+UCLASS(MinimalAPI)
+class UMovieSceneCameraShakeEvaluator : public UObject
 {
 	GENERATED_BODY()
 
@@ -33,7 +33,7 @@ public:
 DECLARE_DELEGATE_RetVal_OneParam(UMovieSceneCameraShakeEvaluator*, FMovieSceneBuildShakeEvaluator, UCameraShakeBase*);
 
 /** Registry for factories of shake evaluators */
-struct MOVIESCENETRACKS_API FMovieSceneCameraShakeEvaluatorRegistry
+struct FMovieSceneCameraShakeEvaluatorRegistry
 {
 public:
 	/** Registers a new shake evaluator factory */
@@ -63,7 +63,7 @@ public:
 
 private:
 	/** List of registered factory methods */
-	static TArray<FMovieSceneBuildShakeEvaluator> ShakeEvaluatorBuilders;
+	static MOVIESCENETRACKS_API TArray<FMovieSceneBuildShakeEvaluator> ShakeEvaluatorBuilders;
 };
 
 /** Section template for a camera anim */
diff --git a/Engine/Source/Runtime/MovieSceneTracks/Public/MovieSceneTracksComponentTypes.h b/Engine/Source/Runtime/MovieSceneTracks/Public/MovieSceneTracksComponentTypes.h
--- a/Engine/Source/Runtime/MovieSceneTracks/Public/MovieSceneTracksComponentTypes.h
+++ b/Engine/Source/Runtime/MovieSceneTracks/Public/MovieSceneTracksComponentTypes.h
@@ -27,7 +27,7 @@ struct FMovieSceneObjectBindingID;
 
 /** Component data for Perlin Noise channels */
 USTRUCT()
-struct MOVIESCENETRACKS_API FPerlinNoiseParams
+struct FPerlinNoiseParams
 {
 	GENERATED_BODY()
 
@@ -43,11 +43,11 @@ struct MOVIESCENETRACKS_API FPerlinNoiseParams
 	UPROPERTY(EditAnywhere, Category = "Perlin Noise")
 	float Offset;
 
-	FPerlinNoiseParams();
-	FPerlinNoiseParams(float InFrequency, double InAmplitude);
+	MOVIESCENETRACKS_API FPerlinNoiseParams();
+	MOVIESCENETRACKS_API FPerlinNoiseParams(float InFrequency, double InAmplitude);
 
 	/** Generates a new random offset between [0, InMaxOffset] and sets it */
-	void RandomizeOffset(float InMaxOffset = 100.f);
+	MOVIESCENETRACKS_API void RandomizeOffset(float InMaxOffset = 100.f);
 };
 
 /** Component data for the level visibility system */
@@ -523,9 +523,9 @@ using FStringPropertyTraits			    = TDirectPropertyTraits<FString>;
 using FFloatParameterTraits             = TIndirectPropertyTraits<float, double, false>;
 using FColorParameterTraits             = TIndirectPropertyTraits<FLinearColor, FIntermediateColor>;
 
-struct MOVIESCENETRACKS_API FMovieSceneTracksComponentTypes
+struct FMovieSceneTracksComponentTypes
 {
-	~FMovieSceneTracksComponentTypes();
+	MOVIESCENETRACKS_API ~FMovieSceneTracksComponentTypes();
 
 	TPropertyComponents<FBoolPropertyTraits> Bool;
 	TPropertyComponents<FBytePropertyTraits> Byte;
@@ -596,9 +596,9 @@ struct MOVIESCENETRACKS_API FMovieSceneTracksComponentTypes
 	TComponentTypeID<FLevelVisibilityComponentData> LevelVisibility;
 	TComponentTypeID<FMovieSceneDataLayerComponentData> DataLayer;
 
-	static void Destroy();
+	static MOVIESCENETRACKS_API void Destroy();
 
-	static FMovieSceneTracksComponentTypes* Get();
+	static MOVIESCENETRACKS_API FMovieSceneTracksComponentTypes* Get();
 
 private:
 	FMovieSceneTracksComponentTypes();
diff --git a/Engine/Source/Runtime/MovieSceneTracks/Public/Sections/MovieScene3DPathSection.h b/Engine/Source/Runtime/MovieSceneTracks/Public/Sections/MovieScene3DPathSection.h
--- a/Engine/Source/Runtime/MovieSceneTracks/Public/Sections/MovieScene3DPathSection.h
+++ b/Engine/Source/Runtime/MovieSceneTracks/Public/Sections/MovieScene3DPathSection.h
@@ -51,11 +51,11 @@ public:
 	 */
 	void SetPathBindingID(const FMovieSceneObjectBindingID& InPathBindingID);
 
-	MOVIESCENETRACKS_API MovieScene3DPathSection_Axis GetFrontAxisEnum() const { return FrontAxisEnum; }
-	MOVIESCENETRACKS_API MovieScene3DPathSection_Axis GetUpAxisEnum() const { return UpAxisEnum; }
-	MOVIESCENETRACKS_API bool GetFollow() const { return bFollow; }
-	MOVIESCENETRACKS_API bool GetReverse() const { return bReverse; }
-	MOVIESCENETRACKS_API bool GetForceUpright() const { return bForceUpright; }
+	MovieScene3DPathSection_Axis GetFrontAxisEnum() const { return FrontAxisEnum; }
+	MovieScene3DPathSection_Axis GetUpAxisEnum() const { return UpAxisEnum; }
+	bool GetFollow() const { return bFollow; }
+	bool GetReverse() const { return bReverse; }
+	bool GetForceUpright() const { return bForceUpright; }
 
 public:
 
diff --git a/Engine/Source/Runtime/MovieSceneTracks/Public/Sections/MovieScene3DTransformSection.h b/Engine/Source/Runtime/MovieSceneTracks/Public/Sections/MovieScene3DTransformSection.h
--- a/Engine/Source/Runtime/MovieSceneTracks/Public/Sections/MovieScene3DTransformSection.h
+++ b/Engine/Source/Runtime/MovieSceneTracks/Public/Sections/MovieScene3DTransformSection.h
@@ -191,8 +191,8 @@ private:
 /**
 * This object contains information needed for constraint channels on the transform section
 */
-UCLASS()
-class MOVIESCENETRACKS_API UMovieScene3DTransformSectionConstraints : public UObject
+UCLASS(MinimalAPI)
+class UMovieScene3DTransformSectionConstraints : public UObject
 {
 	GENERATED_BODY()
 
@@ -204,7 +204,7 @@ public:
 
 	/** When undo/redoing we need to recreate channel proxies after we are done*/
 #if WITH_EDITOR
-	virtual void PostEditUndo() override;
+	MOVIESCENETRACKS_API virtual void PostEditUndo() override;
 #endif
 };
 
@@ -370,12 +370,12 @@ public:
 	/**
 	 * Return the trajectory visibility
 	 */
-	MOVIESCENETRACKS_API EShow3DTrajectory GetShow3DTrajectory() const { return Show3DTrajectory; }
+	EShow3DTrajectory GetShow3DTrajectory() const { return Show3DTrajectory; }
 
 	/**
 	 * Return the trajectory visibility
 	 */
-	MOVIESCENETRACKS_API void SetShow3DTrajectory(EShow3DTrajectory InShow3DTrajectory) { Show3DTrajectory = InShow3DTrajectory; }
+	void SetShow3DTrajectory(EShow3DTrajectory InShow3DTrajectory) { Show3DTrajectory = InShow3DTrajectory; }
 
 private:
 
diff --git a/Engine/Source/Runtime/MovieSceneTracks/Public/Sections/MovieSceneActorReferenceSection.h b/Engine/Source/Runtime/MovieSceneTracks/Public/Sections/MovieSceneActorReferenceSection.h
--- a/Engine/Source/Runtime/MovieSceneTracks/Public/Sections/MovieSceneActorReferenceSection.h
+++ b/Engine/Source/Runtime/MovieSceneTracks/Public/Sections/MovieSceneActorReferenceSection.h
@@ -69,7 +69,7 @@ struct FMovieSceneActorReferenceKey
 
 /** A curve of events */
 USTRUCT()
-struct MOVIESCENETRACKS_API FMovieSceneActorReferenceData : public FMovieSceneChannel
+struct FMovieSceneActorReferenceData : public FMovieSceneChannel
 {
 	GENERATED_BODY()
 
@@ -104,23 +104,23 @@ struct MOVIESCENETRACKS_API FMovieSceneActorReferenceData : public FMovieSceneCh
 	 * @param OutValue   A value to receive the result
 	 * @return true if the channel was evaluated successfully, false otherwise
 	 */
-	bool Evaluate(FFrameTime InTime, FMovieSceneActorReferenceKey& OutValue) const;
+	MOVIESCENETRACKS_API bool Evaluate(FFrameTime InTime, FMovieSceneActorReferenceKey& OutValue) const;
 
 public:
 
 	// ~ FMovieSceneChannel Interface
-	virtual void GetKeys(const TRange<FFrameNumber>& WithinRange, TArray<FFrameNumber>* OutKeyTimes, TArray<FKeyHandle>* OutKeyHandles) override;
-	virtual void GetKeyTimes(TArrayView<const FKeyHandle> InHandles, TArrayView<FFrameNumber> OutKeyTimes) override;
-	virtual void SetKeyTimes(TArrayView<const FKeyHandle> InHandles, TArrayView<const FFrameNumber> InKeyTimes) override;
-	virtual void DuplicateKeys(TArrayView<const FKeyHandle> InHandles, TArrayView<FKeyHandle> OutNewHandles) override;
-	virtual void DeleteKeys(TArrayView<const FKeyHandle> InHandles) override;
-	virtual void DeleteKeysFrom(FFrameNumber InTime, bool bDeleteKeysBefore) override;
-	virtual void ChangeFrameResolution(FFrameRate SourceRate, FFrameRate DestinationRate) override;
-	virtual TRange<FFrameNumber> ComputeEffectiveRange() const override;
-	virtual int32 GetNumKeys() const override;
-	virtual void Reset() override;
-	virtual void Offset(FFrameNumber DeltaPosition) override;
-	virtual void ClearDefault() override;
+	MOVIESCENETRACKS_API virtual void GetKeys(const TRange<FFrameNumber>& WithinRange, TArray<FFrameNumber>* OutKeyTimes, TArray<FKeyHandle>* OutKeyHandles) override;
+	MOVIESCENETRACKS_API virtual void GetKeyTimes(TArrayView<const FKeyHandle> InHandles, TArrayView<FFrameNumber> OutKeyTimes) override;
+	MOVIESCENETRACKS_API virtual void SetKeyTimes(TArrayView<const FKeyHandle> InHandles, TArrayView<const FFrameNumber> InKeyTimes) override;
+	MOVIESCENETRACKS_API virtual void DuplicateKeys(TArrayView<const FKeyHandle> InHandles, TArrayView<FKeyHandle> OutNewHandles) override;
+	MOVIESCENETRACKS_API virtual void DeleteKeys(TArrayView<const FKeyHandle> InHandles) override;
+	MOVIESCENETRACKS_API virtual void DeleteKeysFrom(FFrameNumber InTime, bool bDeleteKeysBefore) override;
+	MOVIESCENETRACKS_API virtual void ChangeFrameResolution(FFrameRate SourceRate, FFrameRate DestinationRate) override;
+	MOVIESCENETRACKS_API virtual TRange<FFrameNumber> ComputeEffectiveRange() const override;
+	MOVIESCENETRACKS_API virtual int32 GetNumKeys() const override;
+	MOVIESCENETRACKS_API virtual void Reset() override;
+	MOVIESCENETRACKS_API virtual void Offset(FFrameNumber DeltaPosition) override;
+	MOVIESCENETRACKS_API virtual void ClearDefault() override;
 
 public:
 
@@ -214,4 +214,4 @@ private:
 inline bool EvaluateChannel(const FMovieSceneActorReferenceData* InChannel, FFrameTime InTime, FMovieSceneActorReferenceKey& OutValue)
 {
 	return InChannel->Evaluate(InTime, OutValue);
-}
\ No newline at end of file
+}
diff --git a/Engine/Source/Runtime/MovieSceneTracks/Public/Sections/MovieSceneAudioSection.h b/Engine/Source/Runtime/MovieSceneTracks/Public/Sections/MovieSceneAudioSection.h
--- a/Engine/Source/Runtime/MovieSceneTracks/Public/Sections/MovieSceneAudioSection.h
+++ b/Engine/Source/Runtime/MovieSceneTracks/Public/Sections/MovieSceneAudioSection.h
@@ -22,8 +22,8 @@ class USoundBase;
 /**
  * Audio section, for use in the audio track, or by attached audio objects
  */
-UCLASS()
-class MOVIESCENETRACKS_API UMovieSceneAudioSection
+UCLASS(MinimalAPI)
+class UMovieSceneAudioSection
 	: public UMovieSceneSection
 	, public IMovieSceneEntityProvider
 {
@@ -33,7 +33,7 @@ public:
 
 	/** Sets this section's sound */
 	UFUNCTION(BlueprintCallable, Category="Sequencer|Section")
-	void SetSound(class USoundBase* InSound);
+	MOVIESCENETRACKS_API void SetSound(class USoundBase* InSound);
 
 	/** Gets the sound for this section */
 	UFUNCTION(BlueprintPure, Category = "Sequencer|Section")
@@ -135,10 +135,10 @@ public:
 	/*
 	 * @return The attach component given the bound actor and the actor attach key with the component and socket names
 	 */
-	USceneComponent* GetAttachComponent(const AActor* InParentActor, const FMovieSceneActorReferenceKey& Key) const;
+	MOVIESCENETRACKS_API USceneComponent* GetAttachComponent(const AActor* InParentActor, const FMovieSceneActorReferenceKey& Key) const;
 
 	/** ~UObject interface */
-	virtual void PostLoad() override;
+	MOVIESCENETRACKS_API virtual void PostLoad() override;
 
 	/** Called when subtitles are sent to the SubtitleManager.  Set this delegate if you want to hijack the subtitles for other purposes */
 	void SetOnQueueSubtitles(const FOnQueueSubtitles& InOnQueueSubtitles)
@@ -184,21 +184,21 @@ public:
 public:
 
 	//~ UMovieSceneSection interface
-	virtual TOptional<TRange<FFrameNumber> > GetAutoSizeRange() const override;
-	virtual void TrimSection(FQualifiedFrameTime TrimTime, bool bTrimLeft, bool bDeleteKeys) override;
-	virtual UMovieSceneSection* SplitSection(FQualifiedFrameTime SplitTime, bool bDeleteKeys) override;
-	virtual TOptional<FFrameTime> GetOffsetTime() const override;
-	virtual void MigrateFrameTimes(FFrameRate SourceRate, FFrameRate DestinationRate) override;
-	virtual void Serialize(FArchive& Ar) override;
-	virtual void PostEditImport() override;
-	virtual EMovieSceneChannelProxyType CacheChannelProxy() override;
+	MOVIESCENETRACKS_API virtual TOptional<TRange<FFrameNumber> > GetAutoSizeRange() const override;
+	MOVIESCENETRACKS_API virtual void TrimSection(FQualifiedFrameTime TrimTime, bool bTrimLeft, bool bDeleteKeys) override;
+	MOVIESCENETRACKS_API virtual UMovieSceneSection* SplitSection(FQualifiedFrameTime SplitTime, bool bDeleteKeys) override;
+	MOVIESCENETRACKS_API virtual TOptional<FFrameTime> GetOffsetTime() const override;
+	MOVIESCENETRACKS_API virtual void MigrateFrameTimes(FFrameRate SourceRate, FFrameRate DestinationRate) override;
+	MOVIESCENETRACKS_API virtual void Serialize(FArchive& Ar) override;
+	MOVIESCENETRACKS_API virtual void PostEditImport() override;
+	MOVIESCENETRACKS_API virtual EMovieSceneChannelProxyType CacheChannelProxy() override;
 
 	//~ IMovieSceneEntityProvider interface
-	virtual bool PopulateEvaluationFieldImpl(const TRange<FFrameNumber>& EffectiveRange, const FMovieSceneEvaluationFieldEntityMetaData& InMetaData, FMovieSceneEntityComponentFieldBuilder* OutFieldBuilder) override;
-	virtual void ImportEntityImpl(UMovieSceneEntitySystemLinker* EntityLinker, const FEntityImportParams& Params, FImportedEntity* OutImportedEntity) override;
+	MOVIESCENETRACKS_API virtual bool PopulateEvaluationFieldImpl(const TRange<FFrameNumber>& EffectiveRange, const FMovieSceneEvaluationFieldEntityMetaData& InMetaData, FMovieSceneEntityComponentFieldBuilder* OutFieldBuilder) override;
+	MOVIESCENETRACKS_API virtual void ImportEntityImpl(UMovieSceneEntitySystemLinker* EntityLinker, const FEntityImportParams& Params, FImportedEntity* OutImportedEntity) override;
 
 #if WITH_EDITOR
-	virtual void PostEditChangeProperty(FPropertyChangedEvent& PropertyChangedEvent) override;
+	MOVIESCENETRACKS_API virtual void PostEditChangeProperty(FPropertyChangedEvent& PropertyChangedEvent) override;
 #endif
 
 private:
@@ -212,7 +212,7 @@ private:
 		}	
 	}
 
-	void SetupSoundInputParameters(const USoundBase* InSoundBase);
+	MOVIESCENETRACKS_API void SetupSoundInputParameters(const USoundBase* InSoundBase);
 
 	/** The sound cue or wave that this section plays */
 	UPROPERTY(EditAnywhere, Category="Audio", BlueprintGetter=GetSound, BlueprintSetter=SetSound)
diff --git a/Engine/Source/Runtime/MovieSceneTracks/Public/Sections/MovieSceneBaseCacheSection.h b/Engine/Source/Runtime/MovieSceneTracks/Public/Sections/MovieSceneBaseCacheSection.h
--- a/Engine/Source/Runtime/MovieSceneTracks/Public/Sections/MovieSceneBaseCacheSection.h
+++ b/Engine/Source/Runtime/MovieSceneTracks/Public/Sections/MovieSceneBaseCacheSection.h
@@ -1,4 +1,4 @@
-﻿// Copyright Epic Games, Inc. All Rights Reserved.
+// Copyright Epic Games, Inc. All Rights Reserved.
 
 #pragma once
 
@@ -9,12 +9,12 @@
  * Base class for the cache parameters that will be used in all the cache sections
  */
 USTRUCT()
-struct MOVIESCENETRACKS_API FMovieSceneBaseCacheParams
+struct FMovieSceneBaseCacheParams
 {
 	GENERATED_BODY()
 
 	/** Constructor/Destructor */
-	FMovieSceneBaseCacheParams();
+	MOVIESCENETRACKS_API FMovieSceneBaseCacheParams();
 	virtual ~FMovieSceneBaseCacheParams() {}
 
 	/** Gets the animation sequence length, not modified by play rate */
diff --git a/Engine/Source/Runtime/MovieSceneTracks/Public/Sections/MovieSceneCinematicShotSection.h b/Engine/Source/Runtime/MovieSceneTracks/Public/Sections/MovieSceneCinematicShotSection.h
--- a/Engine/Source/Runtime/MovieSceneTracks/Public/Sections/MovieSceneCinematicShotSection.h
+++ b/Engine/Source/Runtime/MovieSceneTracks/Public/Sections/MovieSceneCinematicShotSection.h
@@ -18,8 +18,8 @@ struct FFrame;
 /**
  * Implements a cinematic shot section.
  */
-UCLASS(BlueprintType)
-class MOVIESCENETRACKS_API UMovieSceneCinematicShotSection
+UCLASS(BlueprintType, MinimalAPI)
+class UMovieSceneCinematicShotSection
 	: public UMovieSceneSubSection
 {
 	GENERATED_BODY()
@@ -27,18 +27,18 @@ class MOVIESCENETRACKS_API UMovieSceneCinematicShotSection
 public:
 
 	/** Object constructor. */
-	UMovieSceneCinematicShotSection(const FObjectInitializer& ObjInitializer);
+	MOVIESCENETRACKS_API UMovieSceneCinematicShotSection(const FObjectInitializer& ObjInitializer);
 
 private:
 
 	/** ~UObject interface */
-	virtual void PostLoad() override;
+	MOVIESCENETRACKS_API virtual void PostLoad() override;
 
 public:
 
 	/** @return The shot display name. if empty, returns the sequence's name*/
 	UFUNCTION(BlueprintPure, Category = "Sequencer|Section")
-	FString GetShotDisplayName() const;
+	MOVIESCENETRACKS_API FString GetShotDisplayName() const;
 
 	/** Set the shot display name */
 	UFUNCTION(BlueprintCallable, Category = "Sequencer|Section")
diff --git a/Engine/Source/Runtime/MovieSceneTracks/Public/Sections/MovieSceneConstrainedSection.h b/Engine/Source/Runtime/MovieSceneTracks/Public/Sections/MovieSceneConstrainedSection.h
--- a/Engine/Source/Runtime/MovieSceneTracks/Public/Sections/MovieSceneConstrainedSection.h
+++ b/Engine/Source/Runtime/MovieSceneTracks/Public/Sections/MovieSceneConstrainedSection.h
@@ -17,8 +17,8 @@
  * Functionality to add to sections that contain constraints
  */
 
-UINTERFACE()
-class MOVIESCENETRACKS_API UMovieSceneConstrainedSection : public UInterface
+UINTERFACE(MinimalAPI)
+class UMovieSceneConstrainedSection : public UInterface
 {
 public:
 	GENERATED_BODY()
@@ -27,12 +27,12 @@ public:
 /**
  * Interface to be added to UMovieSceneSection types when they contain entity data
  */
-class MOVIESCENETRACKS_API IMovieSceneConstrainedSection
+class IMovieSceneConstrainedSection
 {
 public:
 
 	GENERATED_BODY()
-	IMovieSceneConstrainedSection();
+	MOVIESCENETRACKS_API IMovieSceneConstrainedSection();
 
 	DECLARE_MULTICAST_DELEGATE_TwoParams(FConstraintChannelAddedEvent, IMovieSceneConstrainedSection*, FMovieSceneConstraintChannel*);
 
@@ -80,7 +80,7 @@ public:
 	/*
 	*  SetToNotRemoveChannel when we are told a constraint is removed, we need this sometimes sincet his will be destructive
 	*/
-	 void SetDoNoRemoveChannel(bool bInDoNotRemoveChannel);
+	 MOVIESCENETRACKS_API void SetDoNoRemoveChannel(bool bInDoNotRemoveChannel);
 
 	/*
 	*  Removed delegate that get's added by the track editor
diff --git a/Engine/Source/Runtime/MovieSceneTracks/Public/Sections/MovieSceneEventSection.h b/Engine/Source/Runtime/MovieSceneTracks/Public/Sections/MovieSceneEventSection.h
--- a/Engine/Source/Runtime/MovieSceneTracks/Public/Sections/MovieSceneEventSection.h
+++ b/Engine/Source/Runtime/MovieSceneTracks/Public/Sections/MovieSceneEventSection.h
@@ -19,7 +19,7 @@
 struct EventData;
 
 USTRUCT(BlueprintType)
-struct MOVIESCENETRACKS_API FMovieSceneEventParameters
+struct FMovieSceneEventParameters
 {
 	GENERATED_BODY()
 
@@ -62,19 +62,19 @@ struct MOVIESCENETRACKS_API FMovieSceneEventParameters
 	 *
 	 * @param OutStruct Structure to receive the instance
 	 */
-	void GetInstance(FStructOnScope& OutStruct) const;
+	MOVIESCENETRACKS_API void GetInstance(FStructOnScope& OutStruct) const;
 
 	/**
 	 * Overwrite this payload with another instance of the same type.
 	 *
 	 * @param InstancePtr A valid pointer to an instance of the type represented by GetStructType
 	 */
-	void OverwriteWith(uint8* InstancePtr);
+	MOVIESCENETRACKS_API void OverwriteWith(uint8* InstancePtr);
 
 	/**
 	 * Serialization implementation
 	 */
-	bool Serialize(FArchive& Ar);
+	MOVIESCENETRACKS_API bool Serialize(FArchive& Ar);
 
 	friend FArchive& operator<<(FArchive& Ar, FMovieSceneEventParameters& Payload)
 	{
@@ -129,7 +129,7 @@ struct FEventPayload
 
 /** A curve of events */
 USTRUCT()
-struct MOVIESCENETRACKS_API FMovieSceneEventSectionData : public FMovieSceneChannel
+struct FMovieSceneEventSectionData : public FMovieSceneChannel
 {
 	GENERATED_BODY()
 
@@ -137,7 +137,7 @@ struct MOVIESCENETRACKS_API FMovieSceneEventSectionData : public FMovieSceneChan
 		/**
 	 * Called after this section data has been serialized to upgrade old data
 	 */
-	void PostSerialize(const FArchive& Ar);
+	MOVIESCENETRACKS_API void PostSerialize(const FArchive& Ar);
 #endif
 
 	/**
@@ -173,17 +173,17 @@ struct MOVIESCENETRACKS_API FMovieSceneEventSectionData : public FMovieSceneChan
 public:
 
 	// ~ FMovieSceneChannel Interface
-	virtual void GetKeys(const TRange<FFrameNumber>& WithinRange, TArray<FFrameNumber>* OutKeyTimes, TArray<FKeyHandle>* OutKeyHandles) override;
-	virtual void GetKeyTimes(TArrayView<const FKeyHandle> InHandles, TArrayView<FFrameNumber> OutKeyTimes) override;
-	virtual void SetKeyTimes(TArrayView<const FKeyHandle> InHandles, TArrayView<const FFrameNumber> InKeyTimes) override;
-	virtual void DuplicateKeys(TArrayView<const FKeyHandle> InHandles, TArrayView<FKeyHandle> OutNewHandles) override;
-	virtual void DeleteKeys(TArrayView<const FKeyHandle> InHandles) override;
-	virtual void DeleteKeysFrom(FFrameNumber InTime, bool bDeleteKeysBefore) override;
-	virtual void ChangeFrameResolution(FFrameRate SourceRate, FFrameRate DestinationRate) override;
-	virtual TRange<FFrameNumber> ComputeEffectiveRange() const override;
-	virtual int32 GetNumKeys() const override;
-	virtual void Reset() override;
-	virtual void Offset(FFrameNumber DeltaPosition) override;
+	MOVIESCENETRACKS_API virtual void GetKeys(const TRange<FFrameNumber>& WithinRange, TArray<FFrameNumber>* OutKeyTimes, TArray<FKeyHandle>* OutKeyHandles) override;
+	MOVIESCENETRACKS_API virtual void GetKeyTimes(TArrayView<const FKeyHandle> InHandles, TArrayView<FFrameNumber> OutKeyTimes) override;
+	MOVIESCENETRACKS_API virtual void SetKeyTimes(TArrayView<const FKeyHandle> InHandles, TArrayView<const FFrameNumber> InKeyTimes) override;
+	MOVIESCENETRACKS_API virtual void DuplicateKeys(TArrayView<const FKeyHandle> InHandles, TArrayView<FKeyHandle> OutNewHandles) override;
+	MOVIESCENETRACKS_API virtual void DeleteKeys(TArrayView<const FKeyHandle> InHandles) override;
+	MOVIESCENETRACKS_API virtual void DeleteKeysFrom(FFrameNumber InTime, bool bDeleteKeysBefore) override;
+	MOVIESCENETRACKS_API virtual void ChangeFrameResolution(FFrameRate SourceRate, FFrameRate DestinationRate) override;
+	MOVIESCENETRACKS_API virtual TRange<FFrameNumber> ComputeEffectiveRange() const override;
+	MOVIESCENETRACKS_API virtual int32 GetNumKeys() const override;
+	MOVIESCENETRACKS_API virtual void Reset() override;
+	MOVIESCENETRACKS_API virtual void Offset(FFrameNumber DeltaPosition) override;
 	virtual void Optimize(const FKeyDataOptimizationParams& InParameters) override {}
 	virtual void ClearDefault() override {}
 
diff --git a/Engine/Source/Runtime/MovieSceneTracks/Public/Sections/MovieSceneLevelVisibilitySection.h b/Engine/Source/Runtime/MovieSceneTracks/Public/Sections/MovieSceneLevelVisibilitySection.h
--- a/Engine/Source/Runtime/MovieSceneTracks/Public/Sections/MovieSceneLevelVisibilitySection.h
+++ b/Engine/Source/Runtime/MovieSceneTracks/Public/Sections/MovieSceneLevelVisibilitySection.h
@@ -27,8 +27,8 @@ enum class ELevelVisibility : uint8
 /**
  * A section for use with the movie scene level visibility track, which controls streamed level visibility.
  */
-UCLASS()
-class MOVIESCENETRACKS_API UMovieSceneLevelVisibilitySection
+UCLASS(MinimalAPI)
+class UMovieSceneLevelVisibilitySection
 	: public UMovieSceneSection
 	, public IMovieSceneEntityProvider
 {
@@ -37,10 +37,10 @@ class MOVIESCENETRACKS_API UMovieSceneLevelVisibilitySection
 public:
 
 	UFUNCTION(BlueprintPure, Category = "Sequencer|Section")
-	ELevelVisibility GetVisibility() const;
+	MOVIESCENETRACKS_API ELevelVisibility GetVisibility() const;
 
 	UFUNCTION(BlueprintCallable, Category = "Sequencer|Section")
-	void SetVisibility(ELevelVisibility InVisibility);
+	MOVIESCENETRACKS_API void SetVisibility(ELevelVisibility InVisibility);
 
 	UFUNCTION(BlueprintPure, Category = "Sequencer|Section")
 	const TArray<FName>& GetLevelNames() const { return LevelNames; }
@@ -50,7 +50,7 @@ public:
 
 private:
 
-	virtual void ImportEntityImpl(UMovieSceneEntitySystemLinker* EntityLinker, const FEntityImportParams& Params, FImportedEntity* OutImportedEntity) override;
+	MOVIESCENETRACKS_API virtual void ImportEntityImpl(UMovieSceneEntitySystemLinker* EntityLinker, const FEntityImportParams& Params, FImportedEntity* OutImportedEntity) override;
 
 private:
 
diff --git a/Engine/Source/Runtime/MovieSceneTracks/Public/Sections/MovieSceneParameterSection.h b/Engine/Source/Runtime/MovieSceneTracks/Public/Sections/MovieSceneParameterSection.h
--- a/Engine/Source/Runtime/MovieSceneTracks/Public/Sections/MovieSceneParameterSection.h
+++ b/Engine/Source/Runtime/MovieSceneTracks/Public/Sections/MovieSceneParameterSection.h
@@ -17,20 +17,20 @@ struct FMovieSceneParameterPropertyInterface;
 struct FMovieSceneConstParameterPropertyInterface;
 
 
-UINTERFACE()
-class MOVIESCENETRACKS_API UMovieSceneParameterSectionExtender : public UInterface
+UINTERFACE(MinimalAPI)
+class UMovieSceneParameterSectionExtender : public UInterface
 {
 public:
 	GENERATED_BODY()
 };
 
-class MOVIESCENETRACKS_API IMovieSceneParameterSectionExtender
+class IMovieSceneParameterSectionExtender
 {
 public:
 
 	GENERATED_BODY()
 
-	void ExtendEntity(UMovieSceneParameterSection* Section, UMovieSceneEntitySystemLinker* EntityLinker, const UE::MovieScene::FEntityImportParams& Params, UE::MovieScene::FImportedEntity* OutImportedEntity);
+	MOVIESCENETRACKS_API void ExtendEntity(UMovieSceneParameterSection* Section, UMovieSceneEntitySystemLinker* EntityLinker, const UE::MovieScene::FEntityImportParams& Params, UE::MovieScene::FImportedEntity* OutImportedEntity);
 
 private:
 
@@ -158,7 +158,7 @@ struct FTransformParameterNameAndValue
  * Structure representing an bool  parameter and it's associated animation curve.
  */
 USTRUCT()
-struct MOVIESCENETRACKS_API FBoolParameterNameAndCurve
+struct FBoolParameterNameAndCurve
 {
 	GENERATED_USTRUCT_BODY()
 
@@ -166,7 +166,7 @@ struct MOVIESCENETRACKS_API FBoolParameterNameAndCurve
 	{}
 
 	/** Creates a new FScalarParameterNameAndCurve for a specific scalar parameter. */
-	FBoolParameterNameAndCurve(FName InParameterName);
+	MOVIESCENETRACKS_API FBoolParameterNameAndCurve(FName InParameterName);
 
 	/** The name of the scalar parameter which is being animated. */
 	UPROPERTY()
@@ -181,7 +181,7 @@ struct MOVIESCENETRACKS_API FBoolParameterNameAndCurve
  * Structure representing an animated scalar parameter and it's associated animation curve.
  */
 USTRUCT()
-struct MOVIESCENETRACKS_API FScalarParameterNameAndCurve
+struct FScalarParameterNameAndCurve
 {
 	GENERATED_USTRUCT_BODY()
 
@@ -189,7 +189,7 @@ struct MOVIESCENETRACKS_API FScalarParameterNameAndCurve
 	{}
 
 	/** Creates a new FScalarParameterNameAndCurve for a specific scalar parameter. */
-	FScalarParameterNameAndCurve(FName InParameterName);
+	MOVIESCENETRACKS_API FScalarParameterNameAndCurve(FName InParameterName);
 
 	/** The name of the scalar parameter which is being animated. */
 	UPROPERTY()
@@ -204,7 +204,7 @@ struct MOVIESCENETRACKS_API FScalarParameterNameAndCurve
  * Structure representing an animated vector2D parameter and it's associated animation curve.
  */
 USTRUCT()
-struct MOVIESCENETRACKS_API FVector2DParameterNameAndCurves
+struct FVector2DParameterNameAndCurves
 {
 	GENERATED_USTRUCT_BODY()
 
@@ -212,7 +212,7 @@ struct MOVIESCENETRACKS_API FVector2DParameterNameAndCurves
 	{}
 
 	/** Creates a new FVectorParameterNameAndCurve for a specific vector parameter. */
-	FVector2DParameterNameAndCurves(FName InParameterName);
+	MOVIESCENETRACKS_API FVector2DParameterNameAndCurves(FName InParameterName);
 
 	/** The name of the vector parameter which is being animated. */
 	UPROPERTY()
@@ -233,7 +233,7 @@ struct MOVIESCENETRACKS_API FVector2DParameterNameAndCurves
  * Structure representing an animated vector parameter and it's associated animation curve.
  */
 USTRUCT()
-struct MOVIESCENETRACKS_API FVectorParameterNameAndCurves
+struct FVectorParameterNameAndCurves
 {
 	GENERATED_USTRUCT_BODY()
 
@@ -241,7 +241,7 @@ struct MOVIESCENETRACKS_API FVectorParameterNameAndCurves
 	{}
 
 	/** Creates a new FVectorParameterNameAndCurve for a specific vector parameter. */
-	FVectorParameterNameAndCurves(FName InParameterName);
+	MOVIESCENETRACKS_API FVectorParameterNameAndCurves(FName InParameterName);
 
 	/** The name of the vector parameter which is being animated. */
 	UPROPERTY()
@@ -265,7 +265,7 @@ struct MOVIESCENETRACKS_API FVectorParameterNameAndCurves
 * Structure representing an animated color parameter and it's associated animation curve.
 */
 USTRUCT()
-struct MOVIESCENETRACKS_API FColorParameterNameAndCurves
+struct FColorParameterNameAndCurves
 {
 	GENERATED_USTRUCT_BODY()
 
@@ -273,7 +273,7 @@ struct MOVIESCENETRACKS_API FColorParameterNameAndCurves
 	{}
 
 	/** Creates a new FVectorParameterNameAndCurve for a specific color parameter. */
-	FColorParameterNameAndCurves(FName InParameterName);
+	MOVIESCENETRACKS_API FColorParameterNameAndCurves(FName InParameterName);
 
 	/** The name of the color parameter which is being animated. */
 	UPROPERTY()
@@ -300,7 +300,7 @@ struct MOVIESCENETRACKS_API FColorParameterNameAndCurves
 * Structure representing an animated transform parameter and it's associated animation curve.
 */
 USTRUCT()
-struct MOVIESCENETRACKS_API FTransformParameterNameAndCurves
+struct FTransformParameterNameAndCurves
 {
 	GENERATED_USTRUCT_BODY()
 
@@ -308,7 +308,7 @@ struct MOVIESCENETRACKS_API FTransformParameterNameAndCurves
 	{}
 
 	/** Creates a new FVectorParameterNameAndCurve for a specific color parameter. */
-	FTransformParameterNameAndCurves(FName InParameterName);
+	MOVIESCENETRACKS_API FTransformParameterNameAndCurves(FName InParameterName);
 
 	/** The name of the transform  parameter which is being animated. */
 	UPROPERTY()
@@ -329,8 +329,8 @@ struct MOVIESCENETRACKS_API FTransformParameterNameAndCurves
 /**
  * A single movie scene section which can contain data for multiple named parameters.
  */
-UCLASS()
-class MOVIESCENETRACKS_API UMovieSceneParameterSection
+UCLASS(MinimalAPI)
+class UMovieSceneParameterSection
 	: public UMovieSceneSection
 	, public IMovieSceneEntityProvider
 {
@@ -339,27 +339,27 @@ class MOVIESCENETRACKS_API UMovieSceneParameterSection
 public:
 	/** Adds a a key for a specific scalar parameter at the specified time with the specified value. */
 	UFUNCTION(BlueprintCallable, Category = "Sequencer|Section")
-	void AddScalarParameterKey(FName InParameterName, FFrameNumber InTime, float InValue);
+	MOVIESCENETRACKS_API void AddScalarParameterKey(FName InParameterName, FFrameNumber InTime, float InValue);
 
 	/** Adds a a key for a specific bool parameter at the specified time with the specified value. */
 	UFUNCTION(BlueprintCallable, Category = "Sequencer|Section")
-	void AddBoolParameterKey(FName InParameterName, FFrameNumber InTime, bool InValue);
+	MOVIESCENETRACKS_API void AddBoolParameterKey(FName InParameterName, FFrameNumber InTime, bool InValue);
 
 	/** Adds a a key for a specific vector2D parameter at the specified time with the specified value. */
 	UFUNCTION(BlueprintCallable, Category = "Sequencer|Section")
-	void AddVector2DParameterKey(FName InParameterName, FFrameNumber InTime, FVector2D InValue);
+	MOVIESCENETRACKS_API void AddVector2DParameterKey(FName InParameterName, FFrameNumber InTime, FVector2D InValue);
 
 	/** Adds a a key for a specific vector parameter at the specified time with the specified value. */
 	UFUNCTION(BlueprintCallable, Category = "Sequencer|Section")
-	void AddVectorParameterKey(FName InParameterName, FFrameNumber InTime, FVector InValue);
+	MOVIESCENETRACKS_API void AddVectorParameterKey(FName InParameterName, FFrameNumber InTime, FVector InValue);
 
 	/** Adds a a key for a specific color parameter at the specified time with the specified value. */
 	UFUNCTION(BlueprintCallable, Category = "Sequencer|Section")
-	void AddColorParameterKey(FName InParameterName, FFrameNumber InTime, FLinearColor InValue);
+	MOVIESCENETRACKS_API void AddColorParameterKey(FName InParameterName, FFrameNumber InTime, FLinearColor InValue);
 
 	/** Adds a a key for a specific color parameter at the specified time with the specified value. */
 	UFUNCTION(BlueprintCallable, Category = "Sequencer|Section")
-	void AddTransformParameterKey(FName InParameterName, FFrameNumber InTime, const FTransform& InValue);
+	MOVIESCENETRACKS_API void AddTransformParameterKey(FName InParameterName, FFrameNumber InTime, const FTransform& InValue);
 
 	/** 
 	 * Removes a scalar parameter from this section. 
@@ -368,7 +368,7 @@ public:
 	 * @returns True if a parameter with that name was found and removed, otherwise false.
 	 */
 	UFUNCTION(BlueprintCallable, Category = "Sequencer|Section")
-	bool RemoveScalarParameter(FName InParameterName);
+	MOVIESCENETRACKS_API bool RemoveScalarParameter(FName InParameterName);
 
 	/**
 	 * Removes a bool parameter from this section.
@@ -377,7 +377,7 @@ public:
 	 * @returns True if a parameter with that name was found and removed, otherwise false.
 	 */
 	UFUNCTION(BlueprintCallable, Category = "Sequencer|Section")
-	bool RemoveBoolParameter(FName InParameterName);
+	MOVIESCENETRACKS_API bool RemoveBoolParameter(FName InParameterName);
 
 	/**
 	 * Removes a vector2D parameter from this section.
@@ -386,7 +386,7 @@ public:
 	 * @returns True if a parameter with that name was found and removed, otherwise false.
 	 */
 	UFUNCTION(BlueprintCallable, Category = "Sequencer|Section")
-	bool RemoveVector2DParameter(FName InParameterName);
+	MOVIESCENETRACKS_API bool RemoveVector2DParameter(FName InParameterName);
 
 	/**
      * Removes a vector parameter from this section.
@@ -395,7 +395,7 @@ public:
     * @returns True if a parameter with that name was found and removed, otherwise false.
     */
 	UFUNCTION(BlueprintCallable, Category = "Sequencer|Section")
-	bool RemoveVectorParameter(FName InParameterName);
+	MOVIESCENETRACKS_API bool RemoveVectorParameter(FName InParameterName);
 
 	/**
 	 * Removes a color parameter from this section.
@@ -404,7 +404,7 @@ public:
 	 * @returns True if a parameter with that name was found and removed, otherwise false.
 	 */
 	UFUNCTION(BlueprintCallable, Category = "Sequencer|Section")
-	bool RemoveColorParameter(FName InParameterName);
+	MOVIESCENETRACKS_API bool RemoveColorParameter(FName InParameterName);
 
 	/**
 	 * Removes a transform parameter from this section.
@@ -413,47 +413,47 @@ public:
 	 * @returns True if a parameter with that name was found and removed, otherwise false.
 	 */
 	UFUNCTION(BlueprintCallable, Category = "Sequencer|Section")
-	bool RemoveTransformParameter(FName InParameterName);
+	MOVIESCENETRACKS_API bool RemoveTransformParameter(FName InParameterName);
 
 	/** Gets the animated scalar parameters and their associated curves. */
-	TArray<FScalarParameterNameAndCurve>& GetScalarParameterNamesAndCurves();
-	const TArray<FScalarParameterNameAndCurve>& GetScalarParameterNamesAndCurves() const;
+	MOVIESCENETRACKS_API TArray<FScalarParameterNameAndCurve>& GetScalarParameterNamesAndCurves();
+	MOVIESCENETRACKS_API const TArray<FScalarParameterNameAndCurve>& GetScalarParameterNamesAndCurves() const;
 
 	/** Gets the animated bool parameters and their associated curves. */
-	TArray<FBoolParameterNameAndCurve>& GetBoolParameterNamesAndCurves();
-	const TArray<FBoolParameterNameAndCurve>& GetBoolParameterNamesAndCurves() const;
+	MOVIESCENETRACKS_API TArray<FBoolParameterNameAndCurve>& GetBoolParameterNamesAndCurves();
+	MOVIESCENETRACKS_API const TArray<FBoolParameterNameAndCurve>& GetBoolParameterNamesAndCurves() const;
 
 	/** Gets the animated vector2D parameters and their associated curves. */
-	TArray<FVector2DParameterNameAndCurves>& GetVector2DParameterNamesAndCurves();
-	const TArray<FVector2DParameterNameAndCurves>& GetVector2DParameterNamesAndCurves() const;
+	MOVIESCENETRACKS_API TArray<FVector2DParameterNameAndCurves>& GetVector2DParameterNamesAndCurves();
+	MOVIESCENETRACKS_API const TArray<FVector2DParameterNameAndCurves>& GetVector2DParameterNamesAndCurves() const;
 
 	/** Gets the animated vector parameters and their associated curves. */
-	TArray<FVectorParameterNameAndCurves>& GetVectorParameterNamesAndCurves();
-	const TArray<FVectorParameterNameAndCurves>& GetVectorParameterNamesAndCurves() const;
+	MOVIESCENETRACKS_API TArray<FVectorParameterNameAndCurves>& GetVectorParameterNamesAndCurves();
+	MOVIESCENETRACKS_API const TArray<FVectorParameterNameAndCurves>& GetVectorParameterNamesAndCurves() const;
 
 	/** Gets the animated color parameters and their associated curves. */
-	TArray<FColorParameterNameAndCurves>& GetColorParameterNamesAndCurves();
-	const TArray<FColorParameterNameAndCurves>& GetColorParameterNamesAndCurves() const;
+	MOVIESCENETRACKS_API TArray<FColorParameterNameAndCurves>& GetColorParameterNamesAndCurves();
+	MOVIESCENETRACKS_API const TArray<FColorParameterNameAndCurves>& GetColorParameterNamesAndCurves() const;
 
 	/** Gets the animated transform parameters and their associated curves. */
-	TArray<FTransformParameterNameAndCurves>& GetTransformParameterNamesAndCurves();
-    const TArray<FTransformParameterNameAndCurves>& GetTransformParameterNamesAndCurves() const;
+	MOVIESCENETRACKS_API TArray<FTransformParameterNameAndCurves>& GetTransformParameterNamesAndCurves();
+    MOVIESCENETRACKS_API const TArray<FTransformParameterNameAndCurves>& GetTransformParameterNamesAndCurves() const;
 
 	/** Gets the set of all parameter names used by this section. */
 	UFUNCTION(BlueprintPure, Category = "Sequencer|Section")
-	void GetParameterNames(TSet<FName>& ParameterNames) const;
+	MOVIESCENETRACKS_API void GetParameterNames(TSet<FName>& ParameterNames) const;
 
-	virtual void ImportEntityImpl(UMovieSceneEntitySystemLinker* EntityLinker, const FEntityImportParams& Params, FImportedEntity* OutImportedEntity) override;
-	virtual bool PopulateEvaluationFieldImpl(const TRange<FFrameNumber>& EffectiveRange, const FMovieSceneEvaluationFieldEntityMetaData& InMetaData, FMovieSceneEntityComponentFieldBuilder* OutFieldBuilder) override;
+	MOVIESCENETRACKS_API virtual void ImportEntityImpl(UMovieSceneEntitySystemLinker* EntityLinker, const FEntityImportParams& Params, FImportedEntity* OutImportedEntity) override;
+	MOVIESCENETRACKS_API virtual bool PopulateEvaluationFieldImpl(const TRange<FFrameNumber>& EffectiveRange, const FMovieSceneEvaluationFieldEntityMetaData& InMetaData, FMovieSceneEntityComponentFieldBuilder* OutFieldBuilder) override;
 
-	void ExternalPopulateEvaluationField(const TRange<FFrameNumber>& EffectiveRange, const FMovieSceneEvaluationFieldEntityMetaData& InMetaData, FMovieSceneEntityComponentFieldBuilder* OutFieldBuilder);
+	MOVIESCENETRACKS_API void ExternalPopulateEvaluationField(const TRange<FFrameNumber>& EffectiveRange, const FMovieSceneEvaluationFieldEntityMetaData& InMetaData, FMovieSceneEntityComponentFieldBuilder* OutFieldBuilder);
 
 protected:
 
 	//~ UMovieSceneSection interface
-	virtual void Serialize(FArchive& Ar) override;
-	virtual void PostEditImport() override;
-	virtual void ReconstructChannelProxy();
+	MOVIESCENETRACKS_API virtual void Serialize(FArchive& Ar) override;
+	MOVIESCENETRACKS_API virtual void PostEditImport() override;
+	MOVIESCENETRACKS_API virtual void ReconstructChannelProxy();
 
 protected:
 	/** The bool parameter names and their associated curves. */
diff --git a/Engine/Source/Runtime/MovieSceneTracks/Public/Systems/ByteChannelEvaluatorSystem.h b/Engine/Source/Runtime/MovieSceneTracks/Public/Systems/ByteChannelEvaluatorSystem.h
--- a/Engine/Source/Runtime/MovieSceneTracks/Public/Systems/ByteChannelEvaluatorSystem.h
+++ b/Engine/Source/Runtime/MovieSceneTracks/Public/Systems/ByteChannelEvaluatorSystem.h
@@ -13,15 +13,15 @@ class UObject;
 /**
  * System that is responsible for evaluating byte channels.
  */
-UCLASS()
-class MOVIESCENETRACKS_API UByteChannelEvaluatorSystem : public UMovieSceneEntitySystem
+UCLASS(MinimalAPI)
+class UByteChannelEvaluatorSystem : public UMovieSceneEntitySystem
 {
 public:
 
 	GENERATED_BODY()
 
-	UByteChannelEvaluatorSystem(const FObjectInitializer& ObjInit);
+	MOVIESCENETRACKS_API UByteChannelEvaluatorSystem(const FObjectInitializer& ObjInit);
 
-	virtual void OnSchedulePersistentTasks(UE::MovieScene::IEntitySystemScheduler* TaskScheduler) override;
-	virtual void OnRun(FSystemTaskPrerequisites& InPrerequisites, FSystemSubsequentTasks& Subsequents) override;
+	MOVIESCENETRACKS_API virtual void OnSchedulePersistentTasks(UE::MovieScene::IEntitySystemScheduler* TaskScheduler) override;
+	MOVIESCENETRACKS_API virtual void OnRun(FSystemTaskPrerequisites& InPrerequisites, FSystemSubsequentTasks& Subsequents) override;
 };
diff --git a/Engine/Source/Runtime/MovieSceneTracks/Public/Systems/DoubleChannelEvaluatorSystem.h b/Engine/Source/Runtime/MovieSceneTracks/Public/Systems/DoubleChannelEvaluatorSystem.h
--- a/Engine/Source/Runtime/MovieSceneTracks/Public/Systems/DoubleChannelEvaluatorSystem.h
+++ b/Engine/Source/Runtime/MovieSceneTracks/Public/Systems/DoubleChannelEvaluatorSystem.h
@@ -29,19 +29,19 @@ namespace UE::MovieScene
 /**
  * System that is responsible for evaluating double channels.
  */
-UCLASS()
-class MOVIESCENETRACKS_API UDoubleChannelEvaluatorSystem : public UMovieSceneEntitySystem
+UCLASS(MinimalAPI)
+class UDoubleChannelEvaluatorSystem : public UMovieSceneEntitySystem
 {
 public:
 
 
 	GENERATED_BODY()
 
-	UDoubleChannelEvaluatorSystem(const FObjectInitializer& ObjInit);
+	MOVIESCENETRACKS_API UDoubleChannelEvaluatorSystem(const FObjectInitializer& ObjInit);
 
-	virtual void OnSchedulePersistentTasks(UE::MovieScene::IEntitySystemScheduler* TaskScheduler) override;
-	virtual void OnRun(FSystemTaskPrerequisites& InPrerequisites, FSystemSubsequentTasks& Subsequents) override;
-	virtual bool IsRelevantImpl(UMovieSceneEntitySystemLinker* InLinker) const override;
+	MOVIESCENETRACKS_API virtual void OnSchedulePersistentTasks(UE::MovieScene::IEntitySystemScheduler* TaskScheduler) override;
+	MOVIESCENETRACKS_API virtual void OnRun(FSystemTaskPrerequisites& InPrerequisites, FSystemSubsequentTasks& Subsequents) override;
+	MOVIESCENETRACKS_API virtual bool IsRelevantImpl(UMovieSceneEntitySystemLinker* InLinker) const override;
 
-	static void RegisterChannelType(TComponentTypeID<UE::MovieScene::FSourceDoubleChannel> SourceChannelType, TComponentTypeID<UE::MovieScene::Interpolation::FCachedInterpolation> CachedInterpolationType, TComponentTypeID<double> ResultType);
+	static MOVIESCENETRACKS_API void RegisterChannelType(TComponentTypeID<UE::MovieScene::FSourceDoubleChannel> SourceChannelType, TComponentTypeID<UE::MovieScene::Interpolation::FCachedInterpolation> CachedInterpolationType, TComponentTypeID<double> ResultType);
 };
diff --git a/Engine/Source/Runtime/MovieSceneTracks/Public/Systems/DoublePerlinNoiseChannelEvaluatorSystem.h b/Engine/Source/Runtime/MovieSceneTracks/Public/Systems/DoublePerlinNoiseChannelEvaluatorSystem.h
--- a/Engine/Source/Runtime/MovieSceneTracks/Public/Systems/DoublePerlinNoiseChannelEvaluatorSystem.h
+++ b/Engine/Source/Runtime/MovieSceneTracks/Public/Systems/DoublePerlinNoiseChannelEvaluatorSystem.h
@@ -8,13 +8,13 @@
 /**
  * System that is responsible for evaluating double perlin noise channels.
  */
-UCLASS()
-class MOVIESCENETRACKS_API UDoublePerlinNoiseChannelEvaluatorSystem : public UMovieSceneEntitySystem
+UCLASS(MinimalAPI)
+class UDoublePerlinNoiseChannelEvaluatorSystem : public UMovieSceneEntitySystem
 {
 public:
 	GENERATED_BODY()
 
-	UDoublePerlinNoiseChannelEvaluatorSystem(const FObjectInitializer& ObjInit);
-	virtual void OnSchedulePersistentTasks(UE::MovieScene::IEntitySystemScheduler* TaskScheduler) override;
-	virtual void OnRun(FSystemTaskPrerequisites& InPrerequisites, FSystemSubsequentTasks& Subsequents) override;
+	MOVIESCENETRACKS_API UDoublePerlinNoiseChannelEvaluatorSystem(const FObjectInitializer& ObjInit);
+	MOVIESCENETRACKS_API virtual void OnSchedulePersistentTasks(UE::MovieScene::IEntitySystemScheduler* TaskScheduler) override;
+	MOVIESCENETRACKS_API virtual void OnRun(FSystemTaskPrerequisites& InPrerequisites, FSystemSubsequentTasks& Subsequents) override;
 };
diff --git a/Engine/Source/Runtime/MovieSceneTracks/Public/Systems/FloatChannelEvaluatorSystem.h b/Engine/Source/Runtime/MovieSceneTracks/Public/Systems/FloatChannelEvaluatorSystem.h
--- a/Engine/Source/Runtime/MovieSceneTracks/Public/Systems/FloatChannelEvaluatorSystem.h
+++ b/Engine/Source/Runtime/MovieSceneTracks/Public/Systems/FloatChannelEvaluatorSystem.h
@@ -25,19 +25,19 @@ namespace UE::MovieScene
 /**
  * System that is responsible for evaluating double channels.
  */
-UCLASS()
-class MOVIESCENETRACKS_API UFloatChannelEvaluatorSystem : public UMovieSceneEntitySystem
+UCLASS(MinimalAPI)
+class UFloatChannelEvaluatorSystem : public UMovieSceneEntitySystem
 {
 public:
 
 
 	GENERATED_BODY()
 
-	UFloatChannelEvaluatorSystem(const FObjectInitializer& ObjInit);
+	MOVIESCENETRACKS_API UFloatChannelEvaluatorSystem(const FObjectInitializer& ObjInit);
 
-	virtual void OnSchedulePersistentTasks(UE::MovieScene::IEntitySystemScheduler* TaskScheduler) override;
-	virtual void OnRun(FSystemTaskPrerequisites& InPrerequisites, FSystemSubsequentTasks& Subsequents) override;
-	virtual bool IsRelevantImpl(UMovieSceneEntitySystemLinker* InLinker) const override;
+	MOVIESCENETRACKS_API virtual void OnSchedulePersistentTasks(UE::MovieScene::IEntitySystemScheduler* TaskScheduler) override;
+	MOVIESCENETRACKS_API virtual void OnRun(FSystemTaskPrerequisites& InPrerequisites, FSystemSubsequentTasks& Subsequents) override;
+	MOVIESCENETRACKS_API virtual bool IsRelevantImpl(UMovieSceneEntitySystemLinker* InLinker) const override;
 
-	static void RegisterChannelType(TComponentTypeID<UE::MovieScene::FSourceFloatChannel> SourceChannelType, TComponentTypeID<UE::MovieScene::Interpolation::FCachedInterpolation> CachedInterpolationType, TComponentTypeID<double> ResultType);
+	static MOVIESCENETRACKS_API void RegisterChannelType(TComponentTypeID<UE::MovieScene::FSourceFloatChannel> SourceChannelType, TComponentTypeID<UE::MovieScene::Interpolation::FCachedInterpolation> CachedInterpolationType, TComponentTypeID<double> ResultType);
 };
diff --git a/Engine/Source/Runtime/MovieSceneTracks/Public/Systems/FloatPerlinNoiseChannelEvaluatorSystem.h b/Engine/Source/Runtime/MovieSceneTracks/Public/Systems/FloatPerlinNoiseChannelEvaluatorSystem.h
--- a/Engine/Source/Runtime/MovieSceneTracks/Public/Systems/FloatPerlinNoiseChannelEvaluatorSystem.h
+++ b/Engine/Source/Runtime/MovieSceneTracks/Public/Systems/FloatPerlinNoiseChannelEvaluatorSystem.h
@@ -8,14 +8,14 @@
 /**
  * System that is responsible for evaluating double perlin noise channels.
  */
-UCLASS()
-class MOVIESCENETRACKS_API UFloatPerlinNoiseChannelEvaluatorSystem : public UMovieSceneEntitySystem
+UCLASS(MinimalAPI)
+class UFloatPerlinNoiseChannelEvaluatorSystem : public UMovieSceneEntitySystem
 {
 public:
 	GENERATED_BODY()
 
-	UFloatPerlinNoiseChannelEvaluatorSystem(const FObjectInitializer& ObjInit);
+	MOVIESCENETRACKS_API UFloatPerlinNoiseChannelEvaluatorSystem(const FObjectInitializer& ObjInit);
 
-	virtual void OnSchedulePersistentTasks(UE::MovieScene::IEntitySystemScheduler* TaskScheduler) override;
-	virtual void OnRun(FSystemTaskPrerequisites& InPrerequisites, FSystemSubsequentTasks& Subsequents) override;
+	MOVIESCENETRACKS_API virtual void OnSchedulePersistentTasks(UE::MovieScene::IEntitySystemScheduler* TaskScheduler) override;
+	MOVIESCENETRACKS_API virtual void OnRun(FSystemTaskPrerequisites& InPrerequisites, FSystemSubsequentTasks& Subsequents) override;
 };
diff --git a/Engine/Source/Runtime/MovieSceneTracks/Public/Systems/IntegerChannelEvaluatorSystem.h b/Engine/Source/Runtime/MovieSceneTracks/Public/Systems/IntegerChannelEvaluatorSystem.h
--- a/Engine/Source/Runtime/MovieSceneTracks/Public/Systems/IntegerChannelEvaluatorSystem.h
+++ b/Engine/Source/Runtime/MovieSceneTracks/Public/Systems/IntegerChannelEvaluatorSystem.h
@@ -11,15 +11,15 @@
 /**
  * System that is responsible for evaluating integer channels.
  */
-UCLASS()
-class MOVIESCENETRACKS_API UIntegerChannelEvaluatorSystem : public UMovieSceneEntitySystem
+UCLASS(MinimalAPI)
+class UIntegerChannelEvaluatorSystem : public UMovieSceneEntitySystem
 {
 public:
 
 	GENERATED_BODY()
 
-	UIntegerChannelEvaluatorSystem(const FObjectInitializer& ObjInit);
+	MOVIESCENETRACKS_API UIntegerChannelEvaluatorSystem(const FObjectInitializer& ObjInit);
 
-	virtual void OnSchedulePersistentTasks(UE::MovieScene::IEntitySystemScheduler* TaskScheduler) override;
-	virtual void OnRun(FSystemTaskPrerequisites& InPrerequisites, FSystemSubsequentTasks& Subsequents) override;
+	MOVIESCENETRACKS_API virtual void OnSchedulePersistentTasks(UE::MovieScene::IEntitySystemScheduler* TaskScheduler) override;
+	MOVIESCENETRACKS_API virtual void OnRun(FSystemTaskPrerequisites& InPrerequisites, FSystemSubsequentTasks& Subsequents) override;
 };
diff --git a/Engine/Source/Runtime/MovieSceneTracks/Public/Systems/MovieScene3DTransformPropertySystem.h b/Engine/Source/Runtime/MovieSceneTracks/Public/Systems/MovieScene3DTransformPropertySystem.h
--- a/Engine/Source/Runtime/MovieSceneTracks/Public/Systems/MovieScene3DTransformPropertySystem.h
+++ b/Engine/Source/Runtime/MovieSceneTracks/Public/Systems/MovieScene3DTransformPropertySystem.h
@@ -6,14 +6,14 @@
 #include "MovieScene3DTransformPropertySystem.generated.h"
 
 
-UCLASS()
-class MOVIESCENETRACKS_API UMovieScene3DTransformPropertySystem : public UMovieScenePropertySystem
+UCLASS(MinimalAPI)
+class UMovieScene3DTransformPropertySystem : public UMovieScenePropertySystem
 {
 public:
 
 	GENERATED_BODY()
 
-	UMovieScene3DTransformPropertySystem(const FObjectInitializer& ObjInit);
+	MOVIESCENETRACKS_API UMovieScene3DTransformPropertySystem(const FObjectInitializer& ObjInit);
 
-	virtual void OnRun(FSystemTaskPrerequisites& InPrerequisites, FSystemSubsequentTasks& Subsequents) override;
+	MOVIESCENETRACKS_API virtual void OnRun(FSystemTaskPrerequisites& InPrerequisites, FSystemSubsequentTasks& Subsequents) override;
 };
diff --git a/Engine/Source/Runtime/MovieSceneTracks/Public/Systems/MovieSceneBaseValueEvaluatorSystem.h b/Engine/Source/Runtime/MovieSceneTracks/Public/Systems/MovieSceneBaseValueEvaluatorSystem.h
--- a/Engine/Source/Runtime/MovieSceneTracks/Public/Systems/MovieSceneBaseValueEvaluatorSystem.h
+++ b/Engine/Source/Runtime/MovieSceneTracks/Public/Systems/MovieSceneBaseValueEvaluatorSystem.h
@@ -13,15 +13,15 @@ class UObject;
 /**
  * System that is responsible for evaluating base values, for "additive from base" blending.
  */
-UCLASS()
-class MOVIESCENETRACKS_API UMovieSceneBaseValueEvaluatorSystem : public UMovieSceneEntitySystem
+UCLASS(MinimalAPI)
+class UMovieSceneBaseValueEvaluatorSystem : public UMovieSceneEntitySystem
 {
 public:
 
 	GENERATED_BODY()
 
-	UMovieSceneBaseValueEvaluatorSystem(const FObjectInitializer& ObjInit);
+	MOVIESCENETRACKS_API UMovieSceneBaseValueEvaluatorSystem(const FObjectInitializer& ObjInit);
 
-	virtual void OnRun(FSystemTaskPrerequisites& InPrerequisites, FSystemSubsequentTasks& Subsequents) override;
+	MOVIESCENETRACKS_API virtual void OnRun(FSystemTaskPrerequisites& InPrerequisites, FSystemSubsequentTasks& Subsequents) override;
 };
 
diff --git a/Engine/Source/Runtime/MovieSceneTracks/Public/Systems/MovieSceneBoolPropertySystem.h b/Engine/Source/Runtime/MovieSceneTracks/Public/Systems/MovieSceneBoolPropertySystem.h
--- a/Engine/Source/Runtime/MovieSceneTracks/Public/Systems/MovieSceneBoolPropertySystem.h
+++ b/Engine/Source/Runtime/MovieSceneTracks/Public/Systems/MovieSceneBoolPropertySystem.h
@@ -7,13 +7,13 @@
 
 
 
-UCLASS()
-class MOVIESCENETRACKS_API UMovieSceneBoolPropertySystem : public UMovieScenePropertySystem
+UCLASS(MinimalAPI)
+class UMovieSceneBoolPropertySystem : public UMovieScenePropertySystem
 {
 	GENERATED_BODY()
 
 	UMovieSceneBoolPropertySystem(const FObjectInitializer& ObjInit);
 
-	virtual void OnRun(FSystemTaskPrerequisites& InPrerequisites, FSystemSubsequentTasks& Subsequents) override;
+	MOVIESCENETRACKS_API virtual void OnRun(FSystemTaskPrerequisites& InPrerequisites, FSystemSubsequentTasks& Subsequents) override;
 };
 
diff --git a/Engine/Source/Runtime/MovieSceneTracks/Public/Systems/MovieSceneBytePropertySystem.h b/Engine/Source/Runtime/MovieSceneTracks/Public/Systems/MovieSceneBytePropertySystem.h
--- a/Engine/Source/Runtime/MovieSceneTracks/Public/Systems/MovieSceneBytePropertySystem.h
+++ b/Engine/Source/Runtime/MovieSceneTracks/Public/Systems/MovieSceneBytePropertySystem.h
@@ -5,12 +5,12 @@
 #include "Systems/MovieScenePropertySystem.h"
 #include "MovieSceneBytePropertySystem.generated.h"
 
-UCLASS()
-class MOVIESCENETRACKS_API UMovieSceneBytePropertySystem : public UMovieScenePropertySystem
+UCLASS(MinimalAPI)
+class UMovieSceneBytePropertySystem : public UMovieScenePropertySystem
 {
 	GENERATED_BODY()
 
 	UMovieSceneBytePropertySystem(const FObjectInitializer& ObjInit);
 
-	virtual void OnRun(FSystemTaskPrerequisites& InPrerequisites, FSystemSubsequentTasks& Subsequents) override;
+	MOVIESCENETRACKS_API virtual void OnRun(FSystemTaskPrerequisites& InPrerequisites, FSystemSubsequentTasks& Subsequents) override;
 };
diff --git a/Engine/Source/Runtime/MovieSceneTracks/Public/Systems/MovieSceneColorPropertySystem.h b/Engine/Source/Runtime/MovieSceneTracks/Public/Systems/MovieSceneColorPropertySystem.h
--- a/Engine/Source/Runtime/MovieSceneTracks/Public/Systems/MovieSceneColorPropertySystem.h
+++ b/Engine/Source/Runtime/MovieSceneTracks/Public/Systems/MovieSceneColorPropertySystem.h
@@ -6,13 +6,13 @@
 #include "Systems/MovieScenePropertySystem.h"
 #include "MovieSceneColorPropertySystem.generated.h"
 
-UCLASS()
-class MOVIESCENETRACKS_API UMovieSceneColorPropertySystem : public UMovieScenePropertySystem
+UCLASS(MinimalAPI)
+class UMovieSceneColorPropertySystem : public UMovieScenePropertySystem
 {
 	GENERATED_BODY()
 
 private:
 	UMovieSceneColorPropertySystem(const FObjectInitializer& ObjInit);
 
-	virtual void OnRun(FSystemTaskPrerequisites& InPrerequisites, FSystemSubsequentTasks& Subsequents) override;
+	MOVIESCENETRACKS_API virtual void OnRun(FSystemTaskPrerequisites& InPrerequisites, FSystemSubsequentTasks& Subsequents) override;
 };
diff --git a/Engine/Source/Runtime/MovieSceneTracks/Public/Systems/MovieSceneConstraintSystem.h b/Engine/Source/Runtime/MovieSceneTracks/Public/Systems/MovieSceneConstraintSystem.h
--- a/Engine/Source/Runtime/MovieSceneTracks/Public/Systems/MovieSceneConstraintSystem.h
+++ b/Engine/Source/Runtime/MovieSceneTracks/Public/Systems/MovieSceneConstraintSystem.h
@@ -16,8 +16,8 @@ namespace UE::MovieScene
 /**
  * System that is responsible for propagating constraints to a bound object's FConstraintsManagerController.
  */
-UCLASS()
-class MOVIESCENETRACKS_API UMovieSceneConstraintSystem : public UMovieSceneEntitySystem
+UCLASS(MinimalAPI)
+class UMovieSceneConstraintSystem : public UMovieSceneEntitySystem
 {
 public:
 
@@ -29,9 +29,9 @@ public:
 		TWeakObjectPtr<UTransformableComponentHandle> TransformHandle;
 	};
 
-	UMovieSceneConstraintSystem(const FObjectInitializer& ObjInit);
-	virtual void OnSchedulePersistentTasks(UE::MovieScene::IEntitySystemScheduler* TaskScheduler) override;
-	virtual void OnRun(FSystemTaskPrerequisites& InPrerequisites, FSystemSubsequentTasks& Subsequents) override;
+	MOVIESCENETRACKS_API UMovieSceneConstraintSystem(const FObjectInitializer& ObjInit);
+	MOVIESCENETRACKS_API virtual void OnSchedulePersistentTasks(UE::MovieScene::IEntitySystemScheduler* TaskScheduler) override;
+	MOVIESCENETRACKS_API virtual void OnRun(FSystemTaskPrerequisites& InPrerequisites, FSystemSubsequentTasks& Subsequents) override;
 
 protected:
 	friend UE::MovieScene::FEvaluateConstraintChannels;
diff --git a/Engine/Source/Runtime/MovieSceneTracks/Public/Systems/MovieSceneDoublePropertySystem.h b/Engine/Source/Runtime/MovieSceneTracks/Public/Systems/MovieSceneDoublePropertySystem.h
--- a/Engine/Source/Runtime/MovieSceneTracks/Public/Systems/MovieSceneDoublePropertySystem.h
+++ b/Engine/Source/Runtime/MovieSceneTracks/Public/Systems/MovieSceneDoublePropertySystem.h
@@ -7,12 +7,12 @@
 
 
 
-UCLASS()
-class MOVIESCENETRACKS_API UMovieSceneDoublePropertySystem : public UMovieScenePropertySystem
+UCLASS(MinimalAPI)
+class UMovieSceneDoublePropertySystem : public UMovieScenePropertySystem
 {
 	GENERATED_BODY()
 
 	UMovieSceneDoublePropertySystem(const FObjectInitializer& ObjInit);
 
-	virtual void OnRun(FSystemTaskPrerequisites& InPrerequisites, FSystemSubsequentTasks& Subsequents) override;
+	MOVIESCENETRACKS_API virtual void OnRun(FSystemTaskPrerequisites& InPrerequisites, FSystemSubsequentTasks& Subsequents) override;
 };
diff --git a/Engine/Source/Runtime/MovieSceneTracks/Public/Systems/MovieSceneEnumPropertySystem.h b/Engine/Source/Runtime/MovieSceneTracks/Public/Systems/MovieSceneEnumPropertySystem.h
--- a/Engine/Source/Runtime/MovieSceneTracks/Public/Systems/MovieSceneEnumPropertySystem.h
+++ b/Engine/Source/Runtime/MovieSceneTracks/Public/Systems/MovieSceneEnumPropertySystem.h
@@ -5,12 +5,12 @@
 #include "Systems/MovieScenePropertySystem.h"
 #include "MovieSceneEnumPropertySystem.generated.h"
 
-UCLASS()
-class MOVIESCENETRACKS_API UMovieSceneEnumPropertySystem : public UMovieScenePropertySystem
+UCLASS(MinimalAPI)
+class UMovieSceneEnumPropertySystem : public UMovieScenePropertySystem
 {
 	GENERATED_BODY()
 
 	UMovieSceneEnumPropertySystem(const FObjectInitializer& ObjInit);
 
-	virtual void OnRun(FSystemTaskPrerequisites& InPrerequisites, FSystemSubsequentTasks& Subsequents) override;
+	MOVIESCENETRACKS_API virtual void OnRun(FSystemTaskPrerequisites& InPrerequisites, FSystemSubsequentTasks& Subsequents) override;
 };
diff --git a/Engine/Source/Runtime/MovieSceneTracks/Public/Systems/MovieSceneEulerTransformPropertySystem.h b/Engine/Source/Runtime/MovieSceneTracks/Public/Systems/MovieSceneEulerTransformPropertySystem.h
--- a/Engine/Source/Runtime/MovieSceneTracks/Public/Systems/MovieSceneEulerTransformPropertySystem.h
+++ b/Engine/Source/Runtime/MovieSceneTracks/Public/Systems/MovieSceneEulerTransformPropertySystem.h
@@ -6,14 +6,14 @@
 #include "MovieSceneEulerTransformPropertySystem.generated.h"
 
 
-UCLASS()
-class MOVIESCENETRACKS_API UMovieSceneEulerTransformPropertySystem : public UMovieScenePropertySystem
+UCLASS(MinimalAPI)
+class UMovieSceneEulerTransformPropertySystem : public UMovieScenePropertySystem
 {
 public:
 
 	GENERATED_BODY()
 
-	UMovieSceneEulerTransformPropertySystem(const FObjectInitializer& ObjInit);
+	MOVIESCENETRACKS_API UMovieSceneEulerTransformPropertySystem(const FObjectInitializer& ObjInit);
 
-	virtual void OnRun(FSystemTaskPrerequisites& InPrerequisites, FSystemSubsequentTasks& Subsequents) override;
+	MOVIESCENETRACKS_API virtual void OnRun(FSystemTaskPrerequisites& InPrerequisites, FSystemSubsequentTasks& Subsequents) override;
 };
diff --git a/Engine/Source/Runtime/MovieSceneTracks/Public/Systems/MovieSceneEventSystems.h b/Engine/Source/Runtime/MovieSceneTracks/Public/Systems/MovieSceneEventSystems.h
--- a/Engine/Source/Runtime/MovieSceneTracks/Public/Systems/MovieSceneEventSystems.h
+++ b/Engine/Source/Runtime/MovieSceneTracks/Public/Systems/MovieSceneEventSystems.h
@@ -30,26 +30,26 @@ struct FMovieSceneEventTriggerData
  * Does not dispatch any async tasks
  */
 
-UCLASS()
-class MOVIESCENETRACKS_API UMovieSceneEventSystem : public UMovieSceneEntitySystem
+UCLASS(MinimalAPI)
+class UMovieSceneEventSystem : public UMovieSceneEntitySystem
 {
 public:
 	GENERATED_BODY()
 
-	UMovieSceneEventSystem(const FObjectInitializer& ObjInit);
+	MOVIESCENETRACKS_API UMovieSceneEventSystem(const FObjectInitializer& ObjInit);
 
-	void AddEvent(UE::MovieScene::FInstanceHandle RootInstance, const FMovieSceneEventTriggerData& TriggerData);
+	MOVIESCENETRACKS_API void AddEvent(UE::MovieScene::FInstanceHandle RootInstance, const FMovieSceneEventTriggerData& TriggerData);
 
 protected:
 
-	bool HasEvents() const;
-	void TriggerAllEvents();
+	MOVIESCENETRACKS_API bool HasEvents() const;
+	MOVIESCENETRACKS_API void TriggerAllEvents();
 
 private:
 
-	virtual void OnRun(FSystemTaskPrerequisites& InPrerequisites, FSystemSubsequentTasks& Subsequents) override;
-	virtual bool IsRelevantImpl(UMovieSceneEntitySystemLinker* InLinker) const override;
-	virtual void OnUnlink() override final;
+	MOVIESCENETRACKS_API virtual void OnRun(FSystemTaskPrerequisites& InPrerequisites, FSystemSubsequentTasks& Subsequents) override;
+	MOVIESCENETRACKS_API virtual bool IsRelevantImpl(UMovieSceneEntitySystemLinker* InLinker) const override;
+	MOVIESCENETRACKS_API virtual void OnUnlink() override final;
 
 	static void TriggerEvents(TArrayView<const FMovieSceneEventTriggerData> Events, IMovieScenePlayer* Player);
 	static void TriggerEventWithParameters(UObject* DirectorInstance, const FMovieSceneEventTriggerData& Event, TArrayView<UObject* const> GlobalContexts, IMovieScenePlayer* Player);
@@ -60,34 +60,34 @@ private:
 };
 
 /** System that triggers events before any spawnables */
-UCLASS()
-class MOVIESCENETRACKS_API UMovieScenePreSpawnEventSystem : public UMovieSceneEventSystem
+UCLASS(MinimalAPI)
+class UMovieScenePreSpawnEventSystem : public UMovieSceneEventSystem
 {
 public:
 	GENERATED_BODY()
 
-	UMovieScenePreSpawnEventSystem(const FObjectInitializer& ObjInit);
+	MOVIESCENETRACKS_API UMovieScenePreSpawnEventSystem(const FObjectInitializer& ObjInit);
 };
 
 
 /** System that triggers events after any spawnables */
-UCLASS()
-class MOVIESCENETRACKS_API UMovieScenePostSpawnEventSystem : public UMovieSceneEventSystem
+UCLASS(MinimalAPI)
+class UMovieScenePostSpawnEventSystem : public UMovieSceneEventSystem
 {
 public:
 	GENERATED_BODY()
 
-	UMovieScenePostSpawnEventSystem(const FObjectInitializer& ObjInit);
+	MOVIESCENETRACKS_API UMovieScenePostSpawnEventSystem(const FObjectInitializer& ObjInit);
 };
 
 /** System that triggers events right at the end of evaluation */
-UCLASS()
-class MOVIESCENETRACKS_API UMovieScenePostEvalEventSystem : public UMovieSceneEventSystem
+UCLASS(MinimalAPI)
+class UMovieScenePostEvalEventSystem : public UMovieSceneEventSystem
 {
 public:
 	GENERATED_BODY()
 
-	UMovieScenePostEvalEventSystem(const FObjectInitializer& ObjInit);
+	MOVIESCENETRACKS_API UMovieScenePostEvalEventSystem(const FObjectInitializer& ObjInit);
 
-	virtual void OnRun(FSystemTaskPrerequisites& InPrerequisites, FSystemSubsequentTasks& Subsequents)override;
+	MOVIESCENETRACKS_API virtual void OnRun(FSystemTaskPrerequisites& InPrerequisites, FSystemSubsequentTasks& Subsequents)override;
 };
diff --git a/Engine/Source/Runtime/MovieSceneTracks/Public/Systems/MovieSceneFloatPropertySystem.h b/Engine/Source/Runtime/MovieSceneTracks/Public/Systems/MovieSceneFloatPropertySystem.h
--- a/Engine/Source/Runtime/MovieSceneTracks/Public/Systems/MovieSceneFloatPropertySystem.h
+++ b/Engine/Source/Runtime/MovieSceneTracks/Public/Systems/MovieSceneFloatPropertySystem.h
@@ -7,12 +7,12 @@
 
 
 
-UCLASS()
-class MOVIESCENETRACKS_API UMovieSceneFloatPropertySystem : public UMovieScenePropertySystem
+UCLASS(MinimalAPI)
+class UMovieSceneFloatPropertySystem : public UMovieScenePropertySystem
 {
 	GENERATED_BODY()
 
 	UMovieSceneFloatPropertySystem(const FObjectInitializer& ObjInit);
 
-	virtual void OnRun(FSystemTaskPrerequisites& InPrerequisites, FSystemSubsequentTasks& Subsequents) override;
+	MOVIESCENETRACKS_API virtual void OnRun(FSystemTaskPrerequisites& InPrerequisites, FSystemSubsequentTasks& Subsequents) override;
 };
diff --git a/Engine/Source/Runtime/MovieSceneTracks/Public/Systems/MovieSceneInitialValueSystem.h b/Engine/Source/Runtime/MovieSceneTracks/Public/Systems/MovieSceneInitialValueSystem.h
--- a/Engine/Source/Runtime/MovieSceneTracks/Public/Systems/MovieSceneInitialValueSystem.h
+++ b/Engine/Source/Runtime/MovieSceneTracks/Public/Systems/MovieSceneInitialValueSystem.h
@@ -10,20 +10,20 @@
  * System responsible for initializing initial values for all property types
  * Will handle the presence of an FInitialValueCache extension on the linker
  */
-UCLASS()
-class MOVIESCENETRACKS_API UMovieSceneInitialValueSystem
+UCLASS(MinimalAPI)
+class UMovieSceneInitialValueSystem
 	: public UMovieSceneEntityInstantiatorSystem
 {
 public:
 
 	GENERATED_BODY()
 
-	UMovieSceneInitialValueSystem(const FObjectInitializer& ObjInit);
+	MOVIESCENETRACKS_API UMovieSceneInitialValueSystem(const FObjectInitializer& ObjInit);
 
 private:
 
-	virtual void OnRun(FSystemTaskPrerequisites& InPrerequisites, FSystemSubsequentTasks& Subsequents) override;
-	virtual bool IsRelevantImpl(UMovieSceneEntitySystemLinker* InLinker) const override;
-	virtual void OnLink() override;
-	virtual void OnUnlink() override;
-};
\ No newline at end of file
+	MOVIESCENETRACKS_API virtual void OnRun(FSystemTaskPrerequisites& InPrerequisites, FSystemSubsequentTasks& Subsequents) override;
+	MOVIESCENETRACKS_API virtual bool IsRelevantImpl(UMovieSceneEntitySystemLinker* InLinker) const override;
+	MOVIESCENETRACKS_API virtual void OnLink() override;
+	MOVIESCENETRACKS_API virtual void OnUnlink() override;
+};
diff --git a/Engine/Source/Runtime/MovieSceneTracks/Public/Systems/MovieSceneIntegerPropertySystem.h b/Engine/Source/Runtime/MovieSceneTracks/Public/Systems/MovieSceneIntegerPropertySystem.h
--- a/Engine/Source/Runtime/MovieSceneTracks/Public/Systems/MovieSceneIntegerPropertySystem.h
+++ b/Engine/Source/Runtime/MovieSceneTracks/Public/Systems/MovieSceneIntegerPropertySystem.h
@@ -5,12 +5,12 @@
 #include "Systems/MovieScenePropertySystem.h"
 #include "MovieSceneIntegerPropertySystem.generated.h"
 
-UCLASS()
-class MOVIESCENETRACKS_API UMovieSceneIntegerPropertySystem : public UMovieScenePropertySystem
+UCLASS(MinimalAPI)
+class UMovieSceneIntegerPropertySystem : public UMovieScenePropertySystem
 {
 	GENERATED_BODY()
 
 	UMovieSceneIntegerPropertySystem(const FObjectInitializer& ObjInit);
 
-	virtual void OnRun(FSystemTaskPrerequisites& InPrerequisites, FSystemSubsequentTasks& Subsequents) override;
+	MOVIESCENETRACKS_API virtual void OnRun(FSystemTaskPrerequisites& InPrerequisites, FSystemSubsequentTasks& Subsequents) override;
 };
