Sequencer: correctly cache initial material parameter value

When a material parameter only had one contributor it would only save the initial value if AlwaysCacheInitialValue was true (like for instance when a sequence has dynamic weighting). So in most cases it wouldn't do it until we initialized the blend path.
Now, the blend path gets initialized very easily: if there's more than one contributor, or if there's any non-absolute blending, or if there's any ease in/out. So we would often end up caching this initial value correctly.
But if the user goes through some steps in a very specific order, it could lead to caching the wrong value. The linked JIRA provides such steps:

1. A white cube has an MPC track animating its material color.
2. The track has two sections, both absolute, on different rows. There is not easing or anything so we (erroneously) skip caching the initial value. We animate the color parameter to "red".
3. The second section gets moved to the same row as the first one. This creates an overlap, and therefore some auto-easing is added. This suddently initializes the blend path, and we end up caching the initial value. But we cache "red", which isn't correct (it should be "white").

The bug materialized as a blending bug, because playing around with the easing of both sections' overlap easily leads to <100% weight. When we have "missing weight" like this, we "fill it up" with the initial value.
Now that "red" was cached as the initial value, we would use "red" for the remainder of the weight, but since the first section was also "red", it added up "red" with "red" and made it seem like the section's easing wasn't working.
Fun stuff.

diff --git a/Engine/Source/Runtime/MovieSceneTracks/Private/Systems/MovieSceneMaterialParameterSystem.cpp b/Engine/Source/Runtime/MovieSceneTracks/Private/Systems/MovieSceneMaterialParameterSystem.cpp
--- a/Engine/Source/Runtime/MovieSceneTracks/Private/Systems/MovieSceneMaterialParameterSystem.cpp
+++ b/Engine/Source/Runtime/MovieSceneTracks/Private/Systems/MovieSceneMaterialParameterSystem.cpp
@@ -179,10 +179,15 @@ struct FScalarMixin
 
 	void InitializeSoleInput(UMovieSceneEntitySystemLinker* Linker, const FObjectComponent& BoundMaterial, FName ParameterName, FMovieSceneEntityID SoleContributor, FAnimatedMaterialParameterInfo* Output)
 	{
-		if (Linker->EntityManager.HasComponent(SoleContributor, FBuiltInComponentTypes::Get()->Tags.AlwaysCacheInitialValue))
-		{
-			FMovieSceneTracksComponentTypes* TracksComponents = FMovieSceneTracksComponentTypes::Get();
+		FBuiltInComponentTypes* BuiltInComponents = FBuiltInComponentTypes::Get();
+		FMovieSceneTracksComponentTypes* TracksComponents = FMovieSceneTracksComponentTypes::Get();
 
+		const FComponentMask& SoleContributorType = Linker->EntityManager.GetEntityType(SoleContributor);
+		const bool bNeedsInitialValue = SoleContributorType.ContainsAny({ BuiltInComponents->Tags.RestoreState, BuiltInComponents->Tags.AlwaysCacheInitialValue });
+		const bool bAlreadyHasInitialValue = SoleContributorType.Contains(TracksComponents->FloatParameter.InitialValue);
+
+		if (bNeedsInitialValue && !bAlreadyHasInitialValue)
+		{
 			float InitialValue = 0.0;
 			if (UMaterialInstanceDynamic* MID = Cast<UMaterialInstanceDynamic>(BoundMaterial.GetObject()))
 			{
@@ -257,10 +262,15 @@ struct FVectorMixin
 
 	void InitializeSoleInput(UMovieSceneEntitySystemLinker* Linker, const FObjectComponent& BoundMaterial, FName ParameterName, FMovieSceneEntityID SoleContributor, FAnimatedMaterialParameterInfo* Output)
 	{
-		if (Linker->EntityManager.HasComponent(SoleContributor, FBuiltInComponentTypes::Get()->Tags.AlwaysCacheInitialValue))
-		{
-			FMovieSceneTracksComponentTypes* TracksComponents = FMovieSceneTracksComponentTypes::Get();
+		FBuiltInComponentTypes* BuiltInComponents = FBuiltInComponentTypes::Get();
+		FMovieSceneTracksComponentTypes* TracksComponents = FMovieSceneTracksComponentTypes::Get();
 
+		const FComponentMask& SoleContributorType = Linker->EntityManager.GetEntityType(SoleContributor);
+		const bool bNeedsInitialValue = SoleContributorType.ContainsAny({ BuiltInComponents->Tags.RestoreState, BuiltInComponents->Tags.AlwaysCacheInitialValue });
+		const bool bAlreadyHasInitialValue = SoleContributorType.Contains(TracksComponents->ColorParameter.InitialValue);
+
+		if (bNeedsInitialValue && !bAlreadyHasInitialValue)
+		{
 			FLinearColor ColorValue = FLinearColor::White;
 			if (UMaterialInstanceDynamic* MID = Cast<UMaterialInstanceDynamic>(BoundMaterial.GetObject()))
 			{
