Sequencer: optimize reinitialization of double blender accumulation buffers

diff --git a/Engine/Source/Runtime/MovieSceneTracks/Private/Systems/MovieScenePiecewiseDoubleBlenderSystem.cpp b/Engine/Source/Runtime/MovieSceneTracks/Private/Systems/MovieScenePiecewiseDoubleBlenderSystem.cpp
--- a/Engine/Source/Runtime/MovieSceneTracks/Private/Systems/MovieScenePiecewiseDoubleBlenderSystem.cpp
+++ b/Engine/Source/Runtime/MovieSceneTracks/Private/Systems/MovieScenePiecewiseDoubleBlenderSystem.cpp
@@ -816,14 +816,29 @@ void UMovieScenePiecewiseDoubleBlenderSystem::ReinitializeAccumulationBuffers()
 	const int32 NumResults = ResultComponents.Num();
 	const int32 MaximumNumBlends = AllocatedBlendChannels.Num();
 	FEntityManager& EntityManager = Linker->EntityManager;
+
+	FEntityComponentFilter AnyCompositeFilter;
+	FComponentMask AnyCompositeAllocationsMask;
+
 	for (int32 Index = 0; Index < NumResults; ++Index)
 	{
 		TComponentTypeID<double> Component = ResultComponents[Index];
 
-		const bool bHasAbsolutes         = EntityManager.Contains(FEntityComponentFilter().All({ Component, GetBlenderTypeTag(), BuiltInComponents->BlendChannelInput, BuiltInComponents->Tags.AbsoluteBlend }));
-		const bool bHasRelatives         = EntityManager.Contains(FEntityComponentFilter().All({ Component, GetBlenderTypeTag(), BuiltInComponents->BlendChannelInput, BuiltInComponents->Tags.RelativeBlend }));
-		const bool bHasAdditives         = EntityManager.Contains(FEntityComponentFilter().All({ Component, GetBlenderTypeTag(), BuiltInComponents->BlendChannelInput, BuiltInComponents->Tags.AdditiveBlend }));
-		const bool bHasAdditivesFromBase = EntityManager.Contains(FEntityComponentFilter().All({ Component, GetBlenderTypeTag(), BuiltInComponents->BlendChannelInput, BuiltInComponents->Tags.AdditiveFromBaseBlend }));
+		AnyCompositeFilter.Reset();
+		AnyCompositeFilter.All({ Component, GetBlenderTypeTag(), BuiltInComponents->BlendChannelInput });
+		AnyCompositeFilter.Any({
+				BuiltInComponents->Tags.AbsoluteBlend, 
+				BuiltInComponents->Tags.RelativeBlend, 
+				BuiltInComponents->Tags.AdditiveBlend, 
+				BuiltInComponents->Tags.AdditiveFromBaseBlend });
+
+		AnyCompositeAllocationsMask.Reset();
+		EntityManager.AccumulateMask(AnyCompositeFilter, AnyCompositeAllocationsMask);
+
+		const bool bHasAbsolutes         = AnyCompositeAllocationsMask.Contains(BuiltInComponents->Tags.AbsoluteBlend);
+		const bool bHasRelatives         = AnyCompositeAllocationsMask.Contains(BuiltInComponents->Tags.RelativeBlend);
+		const bool bHasAdditives         = AnyCompositeAllocationsMask.Contains(BuiltInComponents->Tags.AdditiveBlend);
+		const bool bHasAdditivesFromBase = AnyCompositeAllocationsMask.Contains(BuiltInComponents->Tags.AdditiveFromBaseBlend);
 
 		if (!(bHasAbsolutes || bHasRelatives || bHasAdditives || bHasAdditivesFromBase))
 		{
