Sequencer: patch up bound object pointers in the property instantiator system when an object gets replaced

Blueprint actors are recompiled very often, such as when they are moved. This caused pointers in the property instantiator to go stale, and the new object to now have any properties resolved. It was corrected on the next sequencer update, but it meant that keying the moved object would fail to recompose values.

diff --git a/Engine/Source/Runtime/MovieSceneTracks/Private/Systems/MovieScenePropertyInstantiator.cpp b/Engine/Source/Runtime/MovieSceneTracks/Private/Systems/MovieScenePropertyInstantiator.cpp
--- a/Engine/Source/Runtime/MovieSceneTracks/Private/Systems/MovieScenePropertyInstantiator.cpp
+++ b/Engine/Source/Runtime/MovieSceneTracks/Private/Systems/MovieScenePropertyInstantiator.cpp
@@ -82,6 +82,10 @@ void UMovieScenePropertyInstantiatorSystem::OnLink()
 	CleanFastPathMask.Reset();
 	CleanFastPathMask.SetAll({ BuiltInComponents->FastPropertyOffset, BuiltInComponents->SlowProperty, BuiltInComponents->CustomPropertyIndex });
 	CleanFastPathMask.CombineWithBitwiseOR(Linker->EntityManager.GetComponents()->GetMigrationMask(), EBitwiseOperatorFlags::MaxSize);
+
+#if WITH_EDITOR
+	FCoreUObjectDelegates::OnObjectsReplaced.AddUObject(this, &UMovieScenePropertyInstantiatorSystem::OnObjectsReplaced);
+#endif
 }
 
 void UMovieScenePropertyInstantiatorSystem::OnUnlink()
@@ -122,6 +126,10 @@ void UMovieScenePropertyInstantiatorSystem::OnUnlink()
 		InitializePropertyMetaDataTasks.Reset();
 		SaveGlobalStateTasks.Reset();
 	}
+
+#if WITH_EDITOR
+	FCoreUObjectDelegates::OnObjectsReplaced.RemoveAll(this);
+#endif
 }
 
 void UMovieScenePropertyInstantiatorSystem::OnCleanTaggedGarbage()
@@ -137,6 +145,33 @@ void UMovieScenePropertyInstantiatorSystem::OnCleanTaggedGarbage()
 	}
 }
 
+void UMovieScenePropertyInstantiatorSystem::OnObjectsReplaced(const TMap<UObject*, UObject*>& ReplacementMap)
+{
+#if WITH_EDITOR
+	TArray<TTuple<UObject*, FName>> ObjectPropertyKeys;
+	ObjectPropertyToResolvedIndex.GetKeys(ObjectPropertyKeys);
+	for (const TTuple<UObject*, FName>& Key : ObjectPropertyKeys)
+	{
+		if (UObject* const* NewObject = ReplacementMap.Find(Key.Key))
+		{
+			FName PropertyName = Key.Value;
+			int32 ResolvedIndex;
+			ObjectPropertyToResolvedIndex.RemoveAndCopyValue(Key, ResolvedIndex);
+			ObjectPropertyToResolvedIndex.Add(TTuple<UObject*, FName>(*NewObject, PropertyName), ResolvedIndex);
+		}
+	}
+
+	for (auto It = ResolvedProperties.CreateIterator(); It; ++It)
+	{
+		FObjectPropertyInfo& ResolvedProperty = (*It);
+		if (UObject* const* NewObject = ReplacementMap.Find(ResolvedProperty.BoundObject))
+		{
+			ResolvedProperty.BoundObject = *NewObject;
+		}
+	}
+#endif
+}
+
 void UMovieScenePropertyInstantiatorSystem::OnRun(FSystemTaskPrerequisites& InPrerequisites, FSystemSubsequentTasks& Subsequents)
 {
 	using namespace UE::MovieScene;
diff --git a/Engine/Source/Runtime/MovieSceneTracks/Public/Systems/MovieScenePropertyInstantiator.h b/Engine/Source/Runtime/MovieSceneTracks/Public/Systems/MovieScenePropertyInstantiator.h
--- a/Engine/Source/Runtime/MovieSceneTracks/Public/Systems/MovieScenePropertyInstantiator.h
+++ b/Engine/Source/Runtime/MovieSceneTracks/Public/Systems/MovieScenePropertyInstantiator.h
@@ -101,6 +101,10 @@ private:
 	MOVIESCENETRACKS_API virtual void OnUnlink() override;
 	MOVIESCENETRACKS_API virtual void OnCleanTaggedGarbage() override;
 
+private:
+
+	void OnObjectsReplaced(const TMap<UObject*, UObject*>& ReplacementMap);
+
 private:
 
 	using FChannelMask     = TBitArray<TFixedAllocator< 1 >>;
