MovieScene: Added custom task scheduler support for MovieScenePiecewiseBoolBlenderSystem.

This fixes an issue with PPV EditConditionToggle override flags being set before the bool blender system evaluated with the new custom task scheduler which had the visual effect of flickering on activating PPV tracks.

diff --git a/Engine/Source/Runtime/MovieSceneTracks/Private/Systems/MovieScenePiecewiseBoolBlenderSystem.cpp b/Engine/Source/Runtime/MovieSceneTracks/Private/Systems/MovieScenePiecewiseBoolBlenderSystem.cpp
--- a/Engine/Source/Runtime/MovieSceneTracks/Private/Systems/MovieScenePiecewiseBoolBlenderSystem.cpp
+++ b/Engine/Source/Runtime/MovieSceneTracks/Private/Systems/MovieScenePiecewiseBoolBlenderSystem.cpp
@@ -44,12 +44,25 @@ struct TSimpleBlendResultTraits<bool>
 UMovieScenePiecewiseBoolBlenderSystem::UMovieScenePiecewiseBoolBlenderSystem(const FObjectInitializer& ObjInit)
 	: Super(ObjInit)
 {
+	using namespace UE::MovieScene;
+
+	Phase = ESystemPhase::Scheduling;
+	
 	Impl.Setup(
 			this, 
 			UE::MovieScene::FBuiltInComponentTypes::Get()->BoolResult,
 			nullptr);
 }
 
+void UMovieScenePiecewiseBoolBlenderSystem::OnSchedulePersistentTasks(UE::MovieScene::IEntitySystemScheduler* TaskScheduler)
+{
+	using namespace UE::MovieScene;
+
+	CompactBlendChannels();
+
+	Impl.Schedule(Linker, AllocatedBlendChannels, TaskScheduler);
+}
+
 void UMovieScenePiecewiseBoolBlenderSystem::OnRun(FSystemTaskPrerequisites& InPrerequisites, FSystemSubsequentTasks& Subsequents)
 {
 	using namespace UE::MovieScene;
diff --git a/Engine/Source/Runtime/MovieSceneTracks/Public/Systems/MovieSceneBlenderSystemHelper.h b/Engine/Source/Runtime/MovieSceneTracks/Public/Systems/MovieSceneBlenderSystemHelper.h
--- a/Engine/Source/Runtime/MovieSceneTracks/Public/Systems/MovieSceneBlenderSystemHelper.h
+++ b/Engine/Source/Runtime/MovieSceneTracks/Public/Systems/MovieSceneBlenderSystemHelper.h
@@ -66,7 +66,7 @@ struct TSimpleBlenderGatherResults
 		, BlendChannelResults(InBlendChannelResults)
 	{}
 
-	void ForEachEntity(FMovieSceneBlendChannelID BlendChannelInput, PropertyType Value)
+	void ForEachEntity(FMovieSceneBlendChannelID BlendChannelInput, PropertyType Value) const
 	{
 		using FResultTraits = TSimpleBlendResultTraits<PropertyType>;
 
@@ -89,7 +89,7 @@ struct TSimpleBlenderCombineResults
 		, BlendChannelResults(InBlendChannelResults)
 	{}
 
-	void ForEachEntity(FMovieSceneBlendChannelID BlendChannelOutput, PropertyType& OutValue)
+	void ForEachEntity(FMovieSceneBlendChannelID BlendChannelOutput, PropertyType& OutValue) const
 	{
 		using FResultTraits = TSimpleBlendResultTraits<PropertyType>;
 
@@ -133,6 +133,58 @@ public:
 		}
 	}
 
+	void Schedule(UMovieSceneEntitySystemLinker* Linker, TBitArray<>& AllocatedBlendChannels, IEntitySystemScheduler* TaskScheduler)
+	{
+		using namespace UE::MovieScene;
+
+		// If we have no blend channels, we're done.
+		const int32 MaximumNumBlends = AllocatedBlendChannels.Num();
+		if (MaximumNumBlends == 0)
+		{
+			return;
+		}
+
+		const FBuiltInComponentTypes* BuiltInComponents = FBuiltInComponentTypes::Get();
+
+		BlendChannelResults.Reset();
+
+		const bool bHasChannels = Linker->EntityManager.Contains(FEntityComponentFilter().All({
+					ResultComponentID, 
+					BlenderSystem->GetBlenderTypeTag(),
+					BuiltInComponents->BlendChannelInput }));
+		if (bHasChannels)
+		{
+			BlendChannelResults.SetNum(MaximumNumBlends);
+		}
+
+		// Reset the result buffer to the default values, as defined by the value type's traits.
+		FTaskID ResetWeightsTask = TaskScheduler->AddMemberFunctionTask(
+			FTaskParams(TEXT("Reset Simple Blender Weights")),
+			this, &TSimpleBlenderSystemImpl::ZeroAccumulationBuffers);
+
+		// Accumulate all contributors into the result buffer.
+		FTaskID AccumulateTask = FEntityTaskBuilder()
+			.Read(BuiltInComponents->BlendChannelInput)
+			.Read(ResultComponentID)
+			.FilterAll({ BlenderSystem->GetBlenderTypeTag() })
+			.FilterNone({ BuiltInComponents->Tags.Ignored })
+			.template Schedule_PerEntity<TSimpleBlenderGatherResults<PropertyType>>(
+					&Linker->EntityManager, TaskScheduler, BlenderSystem->GetBlenderSystemID(), BlendChannelResults);
+
+		TaskScheduler->AddPrerequisite(ResetWeightsTask, AccumulateTask);
+
+		// Compute the final blended values and assign them to the blend channel output entities.
+		FTaskID CombineTask = FEntityTaskBuilder()
+			.Read(BuiltInComponents->BlendChannelOutput)
+			.Write(ResultComponentID)
+			.FilterAll({ BlenderSystem->GetBlenderTypeTag() })
+			.FilterNone({ BuiltInComponents->Tags.Ignored })
+			.template Fork_PerEntity<TSimpleBlenderCombineResults<PropertyType>>(
+					&Linker->EntityManager, TaskScheduler, BlenderSystem->GetBlenderSystemID(), BlendChannelResults);
+
+		TaskScheduler->AddPrerequisite(AccumulateTask, CombineTask);
+	}
+
 	/**
 	 * Runs the blender system
 	 */
@@ -194,6 +246,12 @@ public:
 					&Linker->EntityManager, Prereqs, &Subsequents, BlenderSystem->GetBlenderSystemID(), BlendChannelResults);
 	}
 
+	void ZeroAccumulationBuffers()
+	{
+		using FResultTraits = TSimpleBlendResultTraits<PropertyType>;
+		FResultTraits::ZeroAccumulationBuffer(MakeArrayView(BlendChannelResults));
+	}
+
 private:
 	UMovieSceneBlenderSystem* BlenderSystem;
 	TComponentTypeID<PropertyType> ResultComponentID;
diff --git a/Engine/Source/Runtime/MovieSceneTracks/Public/Systems/MovieScenePiecewiseBoolBlenderSystem.h b/Engine/Source/Runtime/MovieSceneTracks/Public/Systems/MovieScenePiecewiseBoolBlenderSystem.h
--- a/Engine/Source/Runtime/MovieSceneTracks/Public/Systems/MovieScenePiecewiseBoolBlenderSystem.h
+++ b/Engine/Source/Runtime/MovieSceneTracks/Public/Systems/MovieScenePiecewiseBoolBlenderSystem.h
@@ -23,6 +23,7 @@ public:
 
 	MOVIESCENETRACKS_API UMovieScenePiecewiseBoolBlenderSystem(const FObjectInitializer& ObjInit);
 
+	MOVIESCENETRACKS_API virtual void OnSchedulePersistentTasks(UE::MovieScene::IEntitySystemScheduler* TaskScheduler) override;
 	MOVIESCENETRACKS_API virtual void OnRun(FSystemTaskPrerequisites& InPrerequisites, FSystemSubsequentTasks& Subsequents) override;
 
 private:
