Sequencer: Add the ability to animate layer and function material parameters in Sequencer.
* Changed the material track to interact directly with the UMaterialInterface on the track to find all parameters of a given type rather than the UMaterial assigned to the track- this allowed 'nested' parameters to be found in layers/functions within a UMaterialInterface that had layers.
* Added a new UMovieSceneComponentMaterialParameterSection to hold the keys rather than the more generic UMovieSceneParameterSection. This was necessary because in order to fully address layer/function material parameters you need to use an FMaterialParameterInfo struct rather than just an FName, and UMovieSceneParameterSection only holds FName parameter names, and is being used more generically across different types of parameters depending on the type.
* We do still support old UMovieSceneParameterSections for material parameter animations. These still use material parameter names at runtime, as we didn't want to pollute what has become a generic parameter section with material specific code.
*  MaterialParameterCollection tracks use the regular parameter sections as they have no need to animate layer/function parameters.
* Added an automation test using layered material parameter animation
* Added the ability to have tooltips and different group names vs. group display text on Categories/groups. This is so we can have nice tooltips for the material parameter groups, which need to convey more information about where exactly the parameter comes from within the material structure.

diff --git a/Engine/Source/Editor/MovieSceneTools/Private/Sections/ComponentMaterialParameterSection.cpp b/Engine/Source/Editor/MovieSceneTools/Private/Sections/ComponentMaterialParameterSection.cpp
new file mode 100644
--- /dev/null
+++ b/Engine/Source/Editor/MovieSceneTools/Private/Sections/ComponentMaterialParameterSection.cpp
@@ -0,0 +1,279 @@
+// Copyright Epic Games, Inc. All Rights Reserved.
+
+#include "Sections/ComponentMaterialParameterSection.h"
+
+#include "Channels/MovieSceneChannelEditorData.h"
+#include "Channels/MovieSceneChannelHandle.h"
+#include "Channels/MovieSceneFloatChannel.h"
+#include "HAL/PlatformCrt.h"
+#include "IKeyArea.h"
+#include "Internationalization/Internationalization.h"
+#include "Internationalization/Text.h"
+#include "Layout/Geometry.h"
+#include "Layout/PaintGeometry.h"
+#include "Math/Color.h"
+#include "Math/Vector2D.h"
+#include "Misc/Optional.h"
+#include "MovieSceneSection.h"
+#include "MVVM/ViewModels/CategoryModel.h"
+#include "MVVM/ViewModels/ChannelModel.h"
+#include "MVVM/Views/SChannelView.h"
+#include "MVVM/Views/STrackLane.h"
+#include "MVVM/Views/STrackAreaView.h"
+#include "Rendering/DrawElementPayloads.h"
+#include "Rendering/DrawElements.h"
+#include "Rendering/RenderingCommon.h"
+#include "ScopedTransaction.h"
+#include "Sections/MovieSceneComponentMaterialParameterSection.h"
+#include "Sections/MovieSceneSectionHelpers.h"
+#include "SequencerSectionPainter.h"
+#include "Styling/AppStyle.h"
+#include "Templates/Casts.h"
+#include "Templates/SharedPointer.h"
+#include "Templates/Tuple.h"
+#include "TimeToPixel.h"
+#include "Types/SlateEnums.h"
+#include "UObject/WeakObjectPtrTemplates.h"
+
+class FName;
+struct FKeyHandle;
+
+#define LOCTEXT_NAMESPACE "ParameterSection"
+
+namespace UE::Sequencer
+{
+	class SMaterialColorStripView : public SChannelView
+	{
+		SLATE_BEGIN_ARGS(SMaterialColorStripView) {}
+		SLATE_END_ARGS()
+
+			void Construct(const FArguments& InArgs, const FViewModelPtr& InViewModel, TSharedPtr<STrackAreaView> InTrackAreaView)
+		{
+			WeakModel = InViewModel;
+			WeakTrackAreaView = InTrackAreaView;
+
+			SChannelView::Construct(SChannelView::FArguments(), InViewModel, InTrackAreaView);
+		}
+
+		int32 OnPaint(const FPaintArgs& Args, const FGeometry& AllottedGeometry, const FSlateRect& MyCullingRect, FSlateWindowElementList& OutDrawElements, int32 LayerId, const FWidgetStyle& InWidgetStyle, bool bParentEnabled) const override
+		{
+			FViewModelPtr OwningModel = WeakModel.Pin();
+			if (!OwningModel)
+			{
+				return LayerId;
+			}
+
+			bool bIsExpanded = false;
+			if (TSharedPtr<FLinkedOutlinerExtension> LinkedOutlinerExtension = OwningModel.ImplicitCast())
+			{
+				TSharedPtr<IOutlinerExtension> OutlinerItem = LinkedOutlinerExtension ? LinkedOutlinerExtension->GetLinkedOutlinerItem() : nullptr;
+				if (OutlinerItem)
+				{
+					bIsExpanded = OutlinerItem->IsExpanded();
+				}
+			}
+
+			const FMovieSceneFloatChannel* ColorChannels[4] = { };
+			for (TViewModelPtr<FChannelModel> ChildChannel : OwningModel->GetChildrenOfType<FChannelModel>())
+			{
+				FMovieSceneChannelHandle ChannelHandle = ChildChannel->GetKeyArea()->GetChannel();
+
+				FText DisplayText = ChannelHandle.GetMetaData()->DisplayText;
+
+				if (ChannelHandle.GetMetaData()->DisplayText.EqualTo(FCommonChannelData::ChannelR))
+				{
+					ColorChannels[0] = ChannelHandle.Cast<FMovieSceneFloatChannel>().Get();
+				}
+				if (ChannelHandle.GetMetaData()->DisplayText.EqualTo(FCommonChannelData::ChannelG))
+				{
+					ColorChannels[1] = ChannelHandle.Cast<FMovieSceneFloatChannel>().Get();
+				}
+				if (ChannelHandle.GetMetaData()->DisplayText.EqualTo(FCommonChannelData::ChannelB))
+				{
+					ColorChannels[2] = ChannelHandle.Cast<FMovieSceneFloatChannel>().Get();
+				}
+				if (ChannelHandle.GetMetaData()->DisplayText.EqualTo(FCommonChannelData::ChannelA))
+				{
+					ColorChannels[3] = ChannelHandle.Cast<FMovieSceneFloatChannel>().Get();
+				}
+			}
+
+			if (ColorChannels[0] && ColorChannels[1] && ColorChannels[2] && ColorChannels[3])
+			{
+				const ESlateDrawEffect DrawEffects = bParentEnabled ? ESlateDrawEffect::None : ESlateDrawEffect::DisabledEffect;
+
+				FTimeToPixel RelativeTimeToPixel = *WeakTrackAreaView->GetTimeToPixel();
+
+				TSharedPtr<ITrackLaneExtension> TrackLaneExtension = WeakModel.ImplicitPin();
+				if (TrackLaneExtension)
+				{
+					FTrackLaneVirtualAlignment VirtualAlignment = TrackLaneExtension->ArrangeVirtualTrackLaneView();
+					if (VirtualAlignment.Range.GetLowerBound().IsClosed())
+					{
+						RelativeTimeToPixel = RelativeTimeToPixel.RelativeTo(VirtualAlignment.Range.GetLowerBoundValue());
+					}
+				}
+
+				const float StartTime = RelativeTimeToPixel.PixelToSeconds(0.f);
+				const float EndTime = RelativeTimeToPixel.PixelToSeconds(AllottedGeometry.GetLocalSize().X);
+				const float SectionDuration = EndTime - StartTime;
+
+				FVector2D GradientSize = AllottedGeometry.GetLocalSize() - FVector2D(2.f, 0.f);
+				if (GradientSize.X >= 1.f)
+				{
+					FPaintGeometry PaintGeometry = AllottedGeometry.ToPaintGeometry(GradientSize, FSlateLayoutTransform(FVector2D(1.f, 1.f)));
+
+					// If we are showing a background pattern and the colors is transparent, draw a checker pattern
+					FSlateDrawElement::MakeBox(
+						OutDrawElements,
+						LayerId,
+						PaintGeometry,
+						FAppStyle::GetBrush("Checker"),
+						DrawEffects);
+
+					FLinearColor DefaultColor = FLinearColor::Black;
+					DefaultColor.R = ColorChannels[0]->GetDefault().IsSet() ? ColorChannels[0]->GetDefault().GetValue() : 0.f;
+					DefaultColor.G = ColorChannels[1]->GetDefault().IsSet() ? ColorChannels[1]->GetDefault().GetValue() : 0.f;
+					DefaultColor.B = ColorChannels[2]->GetDefault().IsSet() ? ColorChannels[2]->GetDefault().GetValue() : 0.f;
+					DefaultColor.A = ColorChannels[3]->GetDefault().IsSet() ? ColorChannels[3]->GetDefault().GetValue() : 0.f;
+
+					TArray< TTuple<float, FLinearColor> > ColorKeys;
+					MovieSceneSectionHelpers::ConsolidateColorCurves(ColorKeys, DefaultColor, ColorChannels, RelativeTimeToPixel);
+
+					TArray<FSlateGradientStop> GradientStops;
+
+					for (const TTuple<float, FLinearColor>& ColorStop : ColorKeys)
+					{
+						const float Time = ColorStop.Get<0>();
+
+						// HACK: The color is converted to SRgb and then reinterpreted as linear here because gradients are converted to FColor
+						// without the SRgb conversion before being passed to the renderer for some reason.
+						const FLinearColor Color = ColorStop.Get<1>().ToFColor(true).ReinterpretAsLinear();
+
+						float TimeFraction = (Time - StartTime) / SectionDuration;
+						GradientStops.Add(FSlateGradientStop(FVector2D(TimeFraction * GradientSize.X, 0), Color));
+					}
+
+					if (GradientStops.Num() > 0)
+					{
+						FSlateDrawElement::MakeGradient(
+							OutDrawElements,
+							LayerId + 1,
+							PaintGeometry,
+							GradientStops,
+							Orient_Vertical,
+							DrawEffects
+						);
+					}
+				}
+			}
+
+			LayerId += 2;
+			return SChannelView::OnPaint(Args, AllottedGeometry, MyCullingRect, OutDrawElements, LayerId, InWidgetStyle, bParentEnabled);
+		}
+
+		FVector2D ComputeDesiredSize(float) const override
+		{
+			return FVector2D();
+		}
+
+		FWeakViewModelPtr WeakModel;
+		TSharedPtr<STrackAreaView> WeakTrackAreaView;
+	};
+}
+
+FReply FComponentMaterialParameterSection::OnKeyDoubleClicked(const TArray<FKeyHandle>& KeyHandles)
+{
+	UMovieSceneComponentMaterialParameterSection* ParameterSection = Cast<UMovieSceneComponentMaterialParameterSection>(WeakSection.Get());
+	if (!ParameterSection)
+	{
+		return FReply::Handled();
+	}
+
+	for (FColorMaterialParameterInfoAndCurves& NameAndCurve : ParameterSection->ColorParameterInfosAndCurves)
+	{
+		FMovieSceneKeyColorPicker KeyColorPicker(ParameterSection, &NameAndCurve.RedCurve, &NameAndCurve.GreenCurve, &NameAndCurve.BlueCurve, &NameAndCurve.AlphaCurve, KeyHandles);
+	}
+
+	return FReply::Handled();
+}
+
+TSharedPtr<UE::Sequencer::FCategoryModel> FComponentMaterialParameterSection::ConstructCategoryModel(FName InCategoryName, const FText& InDisplayText, TArrayView<const FChannelData> Channels) const
+{
+	using namespace UE::Sequencer;
+
+	// Only construct the color category if it has all the channels
+	uint8 ColorChannelMask = 0;
+
+	for (const FChannelData& Channel : Channels)
+	{
+		if (Channel.MetaData.DisplayText.EqualTo(FCommonChannelData::ChannelR))
+		{
+			ColorChannelMask |= 1 << 0;
+		}
+		else if (Channel.MetaData.DisplayText.EqualTo(FCommonChannelData::ChannelG))
+		{
+			ColorChannelMask |= 1 << 1;
+		}
+		else if (Channel.MetaData.DisplayText.EqualTo(FCommonChannelData::ChannelB))
+		{
+			ColorChannelMask |= 1 << 2;
+		}
+		else if (Channel.MetaData.DisplayText.EqualTo(FCommonChannelData::ChannelA))
+		{
+			ColorChannelMask |= 1 << 3;
+		}
+	}
+
+	if (ColorChannelMask == 0b1111)
+	{
+		class FMaterialColorCategory : public FCategoryModel
+		{
+		public:
+			FMaterialColorCategory(FName InCategoryName, const FText& InDisplayText)
+				: FCategoryModel(InCategoryName)
+			{}
+
+			TSharedPtr<ITrackLaneWidget> CreateTrackLaneView(const FCreateTrackLaneViewParams& InParams) override
+			{
+				return SNew(SMaterialColorStripView, SharedThis(this), InParams.OwningTrackLane->GetTrackAreaView());
+			}
+		};
+
+		return MakeShared<FMaterialColorCategory>(InCategoryName, InDisplayText);
+	}
+
+	return nullptr;
+}
+
+bool FComponentMaterialParameterSection::RequestDeleteCategory(const TArray<FName>& CategoryNamePath)
+{
+	const FScopedTransaction Transaction(LOCTEXT("DeleteParameter", "Delete parameter"));
+	UMovieSceneComponentMaterialParameterSection* ParameterSection = Cast<UMovieSceneComponentMaterialParameterSection>(WeakSection.Get());
+	if (ParameterSection->Modify())
+	{
+		const bool bColorParameterDeleted = ParameterSection->RemoveColorParameter(CategoryNamePath[0]);
+		const bool bScalarParameterDeleted = ParameterSection->RemoveScalarParameter(CategoryNamePath[0]);
+		return bColorParameterDeleted || bScalarParameterDeleted;
+	}
+	return false;
+}
+
+
+bool FComponentMaterialParameterSection::RequestDeleteKeyArea(const TArray<FName>& KeyAreaNamePath)
+{
+	// Only handle paths with a single name, in all other cases the user is deleting a component of a vector parameter.
+	if (KeyAreaNamePath.Num() == 1)
+	{
+		const FScopedTransaction Transaction(LOCTEXT("DeleteParameter", "Delete parameter"));
+		UMovieSceneComponentMaterialParameterSection* ParameterSection = Cast<UMovieSceneComponentMaterialParameterSection>(WeakSection.Get());
+		if (ParameterSection->TryModify())
+		{
+			const bool bColorParameterDeleted = ParameterSection->RemoveColorParameter(KeyAreaNamePath[0]);
+			const bool bScalarParameterDeleted = ParameterSection->RemoveScalarParameter(KeyAreaNamePath[0]);
+			return bColorParameterDeleted || bScalarParameterDeleted;
+		}
+	}
+	return false;
+}
+#undef LOCTEXT_NAMESPACE
diff --git a/Engine/Source/Editor/MovieSceneTools/Private/Sections/ComponentMaterialParameterSection.h b/Engine/Source/Editor/MovieSceneTools/Private/Sections/ComponentMaterialParameterSection.h
new file mode 100644
--- /dev/null
+++ b/Engine/Source/Editor/MovieSceneTools/Private/Sections/ComponentMaterialParameterSection.h
@@ -0,0 +1,35 @@
+// Copyright Epic Games, Inc. All Rights Reserved.
+
+#pragma once
+
+#include "Containers/Array.h"
+#include "HAL/Platform.h"
+#include "ISequencerSection.h"
+#include "Input/Reply.h"
+
+class FName;
+class FSequencerSectionPainter;
+class UMovieSceneSection;
+struct FKeyHandle;
+
+/**
+ * A movie scene section for material parameters.
+ */
+class FComponentMaterialParameterSection
+	: public FSequencerSection
+{
+public:
+
+	FComponentMaterialParameterSection(UMovieSceneSection& InSectionObject)
+		: FSequencerSection(InSectionObject)
+	{ }
+
+public:
+
+	//~ ISequencerSection interface
+	virtual FReply OnKeyDoubleClicked(const TArray<FKeyHandle>& KeyHandles) override;
+	virtual bool RequestDeleteCategory(const TArray<FName>& CategoryNamePath) override;
+	virtual bool RequestDeleteKeyArea(const TArray<FName>& KeyAreaNamePath) override;
+
+	virtual TSharedPtr<UE::Sequencer::FCategoryModel> ConstructCategoryModel(FName InCategoryName, const FText& InDisplayText, TArrayView<const FChannelData> Channels) const;
+};
diff --git a/Engine/Source/Editor/MovieSceneTools/Private/Sections/ParameterSection.h b/Engine/Source/Editor/MovieSceneTools/Private/Sections/ParameterSection.h
--- a/Engine/Source/Editor/MovieSceneTools/Private/Sections/ParameterSection.h
+++ b/Engine/Source/Editor/MovieSceneTools/Private/Sections/ParameterSection.h
@@ -13,7 +13,7 @@ class UMovieSceneSection;
 struct FKeyHandle;
 
 /**
- * A movie scene section for material parameters.
+ * A movie scene section for generic parameters.
  */
 class FParameterSection
 	: public FSequencerSection
diff --git a/Engine/Source/Editor/MovieSceneTools/Private/TrackEditors/MaterialTrackEditor.cpp b/Engine/Source/Editor/MovieSceneTools/Private/TrackEditors/MaterialTrackEditor.cpp
--- a/Engine/Source/Editor/MovieSceneTools/Private/TrackEditors/MaterialTrackEditor.cpp
+++ b/Engine/Source/Editor/MovieSceneTools/Private/TrackEditors/MaterialTrackEditor.cpp
@@ -6,8 +6,9 @@
 #include "Components/PrimitiveComponent.h"
 #include "Materials/Material.h"
 #include "Materials/MaterialInstance.h"
-#include "Sections/MovieSceneParameterSection.h"
+#include "Sections/MovieSceneComponentMaterialParameterSection.h"
 #include "Tracks/MovieSceneMaterialTrack.h"
+#include "Sections/ComponentMaterialParameterSection.h"
 #include "Sections/ParameterSection.h"
 #include "SequencerUtilities.h"
 #include "Modules/ModuleManager.h"
@@ -27,10 +28,18 @@ FMaterialTrackEditor::FMaterialTrackEditor( TSharedRef<ISequencer> InSequencer )
 
 TSharedRef<ISequencerSection> FMaterialTrackEditor::MakeSectionInterface( UMovieSceneSection& SectionObject, UMovieSceneTrack& Track, FGuid ObjectBinding )
 {
+	UMovieSceneComponentMaterialParameterSection* ComponentMaterialParameterSection = Cast<UMovieSceneComponentMaterialParameterSection>(&SectionObject);
 	UMovieSceneParameterSection* ParameterSection = Cast<UMovieSceneParameterSection>(&SectionObject);
-	checkf( ParameterSection != nullptr, TEXT("Unsupported section type.") );
+	checkf( ComponentMaterialParameterSection != nullptr || ParameterSection != nullptr, TEXT("Unsupported section type.") );
 
-	return MakeShareable(new FParameterSection( *ParameterSection ));
+	if (ComponentMaterialParameterSection)
+	{
+		return MakeShareable(new FComponentMaterialParameterSection(*ComponentMaterialParameterSection));
+	}
+	else
+	{
+		return MakeShareable(new FParameterSection(*ParameterSection));
+	}
 }
 
 
@@ -68,20 +77,30 @@ TSharedRef<SWidget> FMaterialTrackEditor::OnGetAddMenuContent( FGuid ObjectBindi
 }
 
 
-struct FParameterNameAndAction
+struct FParameterInfoAndAction
 {
-	FName ParameterName;
+	FMaterialParameterInfo ParameterInfo;
+	FText ParameterDisplayName;
 	FUIAction Action;
 
-	FParameterNameAndAction( FName InParameterName, FUIAction InAction )
+	FParameterInfoAndAction(const FMaterialParameterInfo& InParameterInfo, FText InParameterDisplayName, FUIAction InAction )
 	{
-		ParameterName = InParameterName;
+		ParameterInfo = InParameterInfo;
+		ParameterDisplayName = InParameterDisplayName;
 		Action = InAction;
 	}
 
-	bool operator<(FParameterNameAndAction const& Other) const
+	bool operator<(FParameterInfoAndAction const& Other) const
 	{
-		return ParameterName.LexicalLess(Other.ParameterName);
+		if (ParameterInfo.Index == Other.ParameterInfo.Index)
+		{
+			if (ParameterInfo.Association == Other.ParameterInfo.Association)
+			{
+				return ParameterInfo.Name.LexicalLess(Other.ParameterInfo.Name);
+			}
+			return ParameterInfo.Association < Other.ParameterInfo.Association;
+		}
+		return ParameterInfo.Index < Other.ParameterInfo.Index;
 	}
 };
 
@@ -112,43 +131,87 @@ void FMaterialTrackEditor::OnBuildAddParameterMenu( FMenuBuilder& MenuBuilder, F
 			bCollectedVisibleParameters = true;
 		}
 
-		TArray<FParameterNameAndAction> ParameterNamesAndActions;
+		TArray<FParameterInfoAndAction> ParameterInfosAndActions;
 
 		// Collect scalar parameters.
-		TArray<FMaterialParameterInfo> ScalarParameterInfo;
+		TArray<FMaterialParameterInfo> ScalarParameterInfos;
 		TArray<FGuid> ScalarParameterGuids;
-		Material->GetAllScalarParameterInfo(ScalarParameterInfo, ScalarParameterGuids );
-		for (int32 ScalarParameterIndex = 0; ScalarParameterIndex < ScalarParameterInfo.Num(); ++ScalarParameterIndex)
+		MaterialInterface->GetAllScalarParameterInfo(ScalarParameterInfos, ScalarParameterGuids );
+		// In case we need to grab layer names.
+		FMaterialLayersFunctions Layers;
+		MaterialInterface->GetMaterialLayers(Layers);
+
+		auto GetMaterialParameterLayerName = [&Layers](const FMaterialParameterInfo& InParameterInfo)
+		{
+			FString LayerName;
+			if (Layers.EditorOnly.LayerNames.IsValidIndex(InParameterInfo.Index))
+			{
+				LayerName = Layers.GetLayerName(InParameterInfo.Index).ToString();
+			}
+			return LayerName;
+		};
+		auto GetMaterialParameterAssetName = [&Layers](const FMaterialParameterInfo& InParameterInfo)
+		{
+			FString AssetName;
+			if (InParameterInfo.Association == EMaterialParameterAssociation::LayerParameter && Layers.Layers.IsValidIndex(InParameterInfo.Index))
+			{
+				AssetName = Layers.Layers[InParameterInfo.Index]->GetName();
+			}
+			else if (InParameterInfo.Association == EMaterialParameterAssociation::BlendParameter && Layers.Blends.IsValidIndex(InParameterInfo.Index))
+			{
+				AssetName = Layers.Blends[InParameterInfo.Index]->GetName();
+			}
+			return AssetName;
+		};
+
+		auto GetMaterialParameterDisplayName = [](const FMaterialParameterInfo& InParameterInfo, const FString& InLayerName, const FString& InAssetName)
+		{
+			FText DisplayName = FText::FromName(InParameterInfo.Name);
+			if (!InLayerName.IsEmpty() && !InAssetName.IsEmpty())
+			{
+				DisplayName = FText::Format(LOCTEXT("MaterialParameterDisplayName", "{0} ({1}.{2})"), DisplayName, FText::FromString(InLayerName), FText::FromString(InAssetName));
+			}
+			return DisplayName;
+		};
+
+		for (int32 ScalarParameterIndex = 0; ScalarParameterIndex < ScalarParameterInfos.Num(); ++ScalarParameterIndex)
 		{
-			if (!bCollectedVisibleParameters || VisibleExpressions.Contains(FMaterialParameterInfo(ScalarParameterInfo[ScalarParameterIndex].Name)))
+			FMaterialParameterInfo ScalarParameterInfo = ScalarParameterInfos[ScalarParameterIndex];
+			if (!bCollectedVisibleParameters || VisibleExpressions.Contains(ScalarParameterInfo))
 			{
-				FName ScalarParameterName = ScalarParameterInfo[ScalarParameterIndex].Name;
-				FUIAction AddParameterMenuAction( FExecuteAction::CreateSP( this, &FMaterialTrackEditor::AddScalarParameter, ObjectBinding, MaterialTrack, ScalarParameterName ) );
-				FParameterNameAndAction NameAndAction( ScalarParameterName, AddParameterMenuAction );
-				ParameterNamesAndActions.Add(NameAndAction);
+				FString LayerName = GetMaterialParameterLayerName(ScalarParameterInfo);
+				FString AssetName = GetMaterialParameterAssetName(ScalarParameterInfo);
+				FText ParameterDisplayName = GetMaterialParameterDisplayName(ScalarParameterInfo, LayerName, AssetName);
+				FUIAction AddParameterMenuAction( FExecuteAction::CreateSP( this, &FMaterialTrackEditor::AddScalarParameter, ObjectBinding, MaterialTrack, ScalarParameterInfo, LayerName, AssetName) );
+				FParameterInfoAndAction InfoAndAction(ScalarParameterInfo, ParameterDisplayName, AddParameterMenuAction );
+				ParameterInfosAndActions.Add(InfoAndAction);
 			}
 		}
 
 		// Collect color parameters.
-		TArray<FMaterialParameterInfo> ColorParameterInfo;
+		TArray<FMaterialParameterInfo> ColorParameterInfos;
 		TArray<FGuid> ColorParameterGuids;
-		Material->GetAllVectorParameterInfo( ColorParameterInfo, ColorParameterGuids );
-		for (int32 ColorParameterIndex = 0; ColorParameterIndex < ColorParameterInfo.Num(); ++ColorParameterIndex)
+		MaterialInterface->GetAllVectorParameterInfo(ColorParameterInfos, ColorParameterGuids );
+		for (int32 ColorParameterIndex = 0; ColorParameterIndex < ColorParameterInfos.Num(); ++ColorParameterIndex)
 		{
-			if (!bCollectedVisibleParameters || VisibleExpressions.Contains(FMaterialParameterInfo(ColorParameterInfo[ColorParameterIndex].Name)))
+			FMaterialParameterInfo ColorParameterInfo = ColorParameterInfos[ColorParameterIndex];
+			if (!bCollectedVisibleParameters || VisibleExpressions.Contains(ColorParameterInfo))
 			{
-				FName ColorParameterName = ColorParameterInfo[ColorParameterIndex].Name;
-				FUIAction AddParameterMenuAction( FExecuteAction::CreateSP( this, &FMaterialTrackEditor::AddColorParameter, ObjectBinding, MaterialTrack, ColorParameterName ) );
-				FParameterNameAndAction NameAndAction( ColorParameterName, AddParameterMenuAction );
-				ParameterNamesAndActions.Add( NameAndAction );
+				FString LayerName = GetMaterialParameterLayerName(ColorParameterInfo);
+				FString AssetName = GetMaterialParameterAssetName(ColorParameterInfo);
+				FText ParameterDisplayName = GetMaterialParameterDisplayName(ColorParameterInfo, LayerName, AssetName);
+				FUIAction AddParameterMenuAction( FExecuteAction::CreateSP( this, &FMaterialTrackEditor::AddColorParameter, ObjectBinding, MaterialTrack, ColorParameterInfo, LayerName, AssetName ) );
+				FParameterInfoAndAction InfoAndAction(ColorParameterInfo, ParameterDisplayName, AddParameterMenuAction );
+				ParameterInfosAndActions.Add(InfoAndAction);
 			}
 		}
 
 		// Sort and generate menu.
-		ParameterNamesAndActions.Sort();
-		for ( FParameterNameAndAction NameAndAction : ParameterNamesAndActions )
+		ParameterInfosAndActions.Sort();
+
+		for (FParameterInfoAndAction InfoAndAction : ParameterInfosAndActions)
 		{
-			MenuBuilder.AddMenuEntry( FText::FromName( NameAndAction.ParameterName ), FText(), FSlateIcon(), NameAndAction.Action );
+			MenuBuilder.AddMenuEntry(InfoAndAction.ParameterDisplayName, FText(), FSlateIcon(), InfoAndAction.Action );
 		}
 	}
 }
@@ -177,7 +240,7 @@ UMaterial* FMaterialTrackEditor::GetMaterialForTrack( FGuid ObjectBinding, UMovi
 }
 
 
-void FMaterialTrackEditor::AddScalarParameter( FGuid ObjectBinding, UMovieSceneMaterialTrack* MaterialTrack, FName ParameterName )
+void FMaterialTrackEditor::AddScalarParameter( FGuid ObjectBinding, UMovieSceneMaterialTrack* MaterialTrack, FMaterialParameterInfo ParameterInfo, FString InLayerName, FString InAssetName)
 {
 	FFrameNumber KeyTime = GetTimeForKey();
 
@@ -186,15 +249,15 @@ void FMaterialTrackEditor::AddScalarParameter( FGuid ObjectBinding, UMovieSceneM
 	{
 		const FScopedTransaction Transaction( LOCTEXT( "AddScalarParameter", "Add scalar parameter" ) );
 		float ParameterValue;
-		Material->GetScalarParameterValue(ParameterName, ParameterValue);
+		Material->GetScalarParameterValue(ParameterInfo, ParameterValue);
 		MaterialTrack->Modify();
-		MaterialTrack->AddScalarParameterKey(ParameterName, KeyTime, ParameterValue);
+		MaterialTrack->AddScalarParameterKey(ParameterInfo, KeyTime, ParameterValue, InLayerName, InAssetName);
 	}
 	GetSequencer()->NotifyMovieSceneDataChanged( EMovieSceneDataChangeType::MovieSceneStructureItemAdded );
 }
 
 
-void FMaterialTrackEditor::AddColorParameter( FGuid ObjectBinding, UMovieSceneMaterialTrack* MaterialTrack, FName ParameterName )
+void FMaterialTrackEditor::AddColorParameter( FGuid ObjectBinding, UMovieSceneMaterialTrack* MaterialTrack, FMaterialParameterInfo ParameterInfo, FString InLayerName, FString InAssetName)
 {
 	FFrameNumber KeyTime = GetTimeForKey();
 
@@ -203,9 +266,9 @@ void FMaterialTrackEditor::AddColorParameter( FGuid ObjectBinding, UMovieSceneMa
 	{
 		const FScopedTransaction Transaction( LOCTEXT( "AddVectorParameter", "Add vector parameter" ) );
 		FLinearColor ParameterValue;
-		Material->GetVectorParameterValue( ParameterName, ParameterValue );
+		Material->GetVectorParameterValue(ParameterInfo, ParameterValue );
 		MaterialTrack->Modify();
-		MaterialTrack->AddColorParameterKey( ParameterName, KeyTime, ParameterValue );
+		MaterialTrack->AddColorParameterKey(ParameterInfo, KeyTime, ParameterValue, InLayerName, InAssetName);
 	}
 	GetSequencer()->NotifyMovieSceneDataChanged( EMovieSceneDataChangeType::MovieSceneStructureItemAdded );
 }
diff --git a/Engine/Source/Editor/MovieSceneTools/Public/TrackEditors/MaterialTrackEditor.h b/Engine/Source/Editor/MovieSceneTools/Public/TrackEditors/MaterialTrackEditor.h
--- a/Engine/Source/Editor/MovieSceneTools/Public/TrackEditors/MaterialTrackEditor.h
+++ b/Engine/Source/Editor/MovieSceneTools/Public/TrackEditors/MaterialTrackEditor.h
@@ -10,6 +10,7 @@
 #include "MovieSceneTrack.h"
 #include "ISequencerSection.h"
 #include "ISequencerTrackEditor.h"
+#include "MaterialTypes.h"
 #include "MovieSceneTrackEditor.h"
 
 class UMaterial;
@@ -62,14 +63,14 @@ private:
 	 * @param MaterialTrack The track in which to look for sections to add the parameter to.
 	 * @param ParameterName The name of the parameter to add an initial key for.
 	 */
-	void AddScalarParameter( FGuid ObjectBinding, UMovieSceneMaterialTrack* MaterialTrack, FName ParameterName );
+	void AddScalarParameter( FGuid ObjectBinding, UMovieSceneMaterialTrack* MaterialTrack, FMaterialParameterInfo ParameterInfo, FString InLayerName, FString InAssetName );
 
 	/** Adds a color parameter and initial key to a material track.
 	* @param ObjectBinding The object binding which owns the material track.
 	 * @param MaterialTrack The track in which to look for sections to add the parameter to.
 	* @param ParameterName The name of the parameter to add an initial key for.
 	*/
-	void AddColorParameter( FGuid ObjectBinding, UMovieSceneMaterialTrack* MaterialTrack, FName ParameterName );
+	void AddColorParameter( FGuid ObjectBinding, UMovieSceneMaterialTrack* MaterialTrack, FMaterialParameterInfo ParameterInfo, FString InLayerName, FString InAssetName);
 };
 
 
diff --git a/Engine/Source/Editor/Sequencer/Private/ISequencerSection.cpp b/Engine/Source/Editor/Sequencer/Private/ISequencerSection.cpp
--- a/Engine/Source/Editor/Sequencer/Private/ISequencerSection.cpp
+++ b/Engine/Source/Editor/Sequencer/Private/ISequencerSection.cpp
@@ -29,8 +29,9 @@ struct FMovieSceneChannel;
 /** Data pertaining to a group of channels */
 struct FGroupData
 {
-	FGroupData(FText InGroupText)
+	FGroupData(FText InGroupText, FText InGroupTooltipText)
 		: GroupText(InGroupText)
+		, GroupTooltipText(InGroupTooltipText)
 		, SortOrder(-1)
 	{}
 
@@ -46,6 +47,9 @@ struct FGroupData
 
 	/** Text to display for the group */
 	FText GroupText;
+	
+	/** Text to display for the group tooltip */
+	FText GroupTooltipText;
 
 	/** Sort order of the group */
 	uint32 SortOrder;
@@ -88,7 +92,12 @@ void ISequencerSection::GenerateSectionLayout( ISectionLayoutBuilder& LayoutBuil
 				FGroupData* ExistingGroup = GroupToChannelsMap.Find(GroupName);
 				if (!ExistingGroup)
 				{
-					ExistingGroup = &GroupToChannelsMap.Add(GroupName, FGroupData(MetaData.Group));
+					FText GroupDisplayName = FText::FromString(MetaData.GetPropertyMetaData(FCommonChannelData::GroupDisplayName));
+					if (GroupDisplayName.IsEmpty())
+					{
+						GroupDisplayName = FText::FromName(GroupName);
+					}
+					ExistingGroup = &GroupToChannelsMap.Add(GroupName, FGroupData(GroupDisplayName, FText::FromString(MetaData.GetPropertyMetaData(FCommonChannelData::GroupTooltipText))));
 				}
 
 				ExistingGroup->AddChannel(FChannelData{ Channel, MetaData });
@@ -176,7 +185,7 @@ void ISequencerSection::GenerateSectionLayout( ISectionLayoutBuilder& LayoutBuil
 				return this->ConstructCategoryModel(InCategoryName, InDisplayText, ChannelData.Channels);
 			};
 
-			LayoutBuilder.PushCategory(GroupName, ChannelData.GroupText, Factory);
+			LayoutBuilder.PushCategory(GroupName, ChannelData.GroupText, ChannelData.GroupTooltipText, Factory);
 		}
 
 		for (const FChannelData& ChannelAndData : ChannelData.Channels)
diff --git a/Engine/Source/Editor/Sequencer/Private/MVVM/ViewModels/CategoryModel.cpp b/Engine/Source/Editor/Sequencer/Private/MVVM/ViewModels/CategoryModel.cpp
--- a/Engine/Source/Editor/Sequencer/Private/MVVM/ViewModels/CategoryModel.cpp
+++ b/Engine/Source/Editor/Sequencer/Private/MVVM/ViewModels/CategoryModel.cpp
@@ -85,9 +85,10 @@ FLinearColor FCategoryModel::GetKeyBarColor() const
 	return FColor(160, 160, 160);
 }
 
-FCategoryGroupModel::FCategoryGroupModel(FName InCategoryName, const FText& InDisplayText)
+FCategoryGroupModel::FCategoryGroupModel(FName InCategoryName, const FText& InDisplayText, const FText& InTooltipText)
 	: CategoryName(InCategoryName)
 	, DisplayText(InDisplayText)
+	, TooltipText(InTooltipText)
 {
 	SetIdentifier(InCategoryName);
 }
@@ -165,6 +166,11 @@ FText FCategoryGroupModel::GetLabel() const
 	return GetDisplayText();
 }
 
+FText FCategoryGroupModel::GetLabelToolTipText() const
+{
+	return GetTooltipText();
+}
+
 FSlateFontInfo FCategoryGroupModel::GetLabelFont() const
 {
 	return IsAnimated()
diff --git a/Engine/Source/Editor/Sequencer/Private/MVVM/ViewModels/ChannelModel.cpp b/Engine/Source/Editor/Sequencer/Private/MVVM/ViewModels/ChannelModel.cpp
--- a/Engine/Source/Editor/Sequencer/Private/MVVM/ViewModels/ChannelModel.cpp
+++ b/Engine/Source/Editor/Sequencer/Private/MVVM/ViewModels/ChannelModel.cpp
@@ -261,10 +261,12 @@ FLinearColor FChannelModel::GetKeyBarColor() const
 	return FColor(160, 160, 160);
 }
 
-FChannelGroupModel::FChannelGroupModel(FName InChannelName, const FText& InDisplayText)
+
+FChannelGroupModel::FChannelGroupModel(FName InChannelName, const FText& InDisplayText, const FText& InTooltipText)
 	: ChannelsSerialNumber(0)
 	, ChannelName(InChannelName)
 	, DisplayText(InDisplayText)
+	, TooltipText(InTooltipText)
 {
 }
 
@@ -718,8 +720,8 @@ void FChannelGroupOverrideHelper::RemoveChannelOverrides(TSharedPtr<FSequenceMod
 	SequenceModel->GetSequencer()->NotifyMovieSceneDataChanged(EMovieSceneDataChangeType::MovieSceneStructureItemsChanged);
 }
 
-FChannelGroupOutlinerModel::FChannelGroupOutlinerModel(FName InChannelName, const FText& InDisplayText)
-	: TOutlinerModelMixin<FChannelGroupModel>(InChannelName, InDisplayText)
+FChannelGroupOutlinerModel::FChannelGroupOutlinerModel(FName InChannelName, const FText& InDisplayText, const FText& InTooltipText)
+	: TOutlinerModelMixin<FChannelGroupModel>(InChannelName, InDisplayText, InTooltipText)
 {
 	SetIdentifier(InChannelName);
 }
@@ -810,6 +812,11 @@ FSlateFontInfo FChannelGroupOutlinerModel::GetLabelFont() const
 		: FOutlinerItemModelMixin::GetLabelFont();
 }
 
+FText FChannelGroupOutlinerModel::GetLabelToolTipText() const
+{
+	return GetTooltipText();
+}
+
 bool FChannelGroupOutlinerModel::HasCurves() const
 {
 	return FChannelGroupModel::HasCurves();
diff --git a/Engine/Source/Editor/Sequencer/Private/MVVM/ViewModels/TrackModelLayoutBuilder.cpp b/Engine/Source/Editor/Sequencer/Private/MVVM/ViewModels/TrackModelLayoutBuilder.cpp
--- a/Engine/Source/Editor/Sequencer/Private/MVVM/ViewModels/TrackModelLayoutBuilder.cpp
+++ b/Engine/Source/Editor/Sequencer/Private/MVVM/ViewModels/TrackModelLayoutBuilder.cpp
@@ -164,7 +164,7 @@ void FTrackModelLayoutBuilder::RefreshLayout(TSharedPtr<FSectionModel> InSection
 	TrackAreaList = FHierarchicalModelListRefresher();
 }
 
-void FTrackModelLayoutBuilder::PushCategory(FName CategoryName, const FText& DisplayLabel, TFunction<TSharedPtr<FCategoryModel>(FName, const FText&)> OptionalFactory)
+void FTrackModelLayoutBuilder::PushCategory(FName CategoryName, const FText& DisplayLabel, const FText& TooltipText, TFunction<TSharedPtr<FCategoryModel>(FName, const FText&)> OptionalFactory)
 {
 	check(TrackAreaList.IsValid() && SequencerSection);
 
@@ -179,7 +179,7 @@ void FTrackModelLayoutBuilder::PushCategory(FName CategoryName, const FText& Dis
 	TSharedPtr<FCategoryGroupModel> OutlinerModel = OutlinerList.FindItem<FCategoryGroupModel>(CategoryNamePredicate);
 	if (!OutlinerModel)
 	{
-		OutlinerModel = MakeShared<FCategoryGroupModel>(CategoryName, DisplayLabel);
+		OutlinerModel = MakeShared<FCategoryGroupModel>(CategoryName, DisplayLabel, TooltipText);
 	}
 
 	OutlinerList.Link(OutlinerModel);
@@ -280,13 +280,14 @@ void FTrackModelLayoutBuilder::AddChannel(const FMovieSceneChannelHandle& Channe
 
 	if (!OutlinerModel)
 	{
+		FString TooltipText = MetaData->PropertyMetaData.FindRef(FCommonChannelData::TooltipText, FString());
 		if (bIsTopLevel)
 		{
-			OutlinerModel = MakeShared<FChannelGroupModel>(ChannelName, MetaData->DisplayText);
+			OutlinerModel = MakeShared<FChannelGroupModel>(ChannelName, MetaData->DisplayText, FText::FromString(TooltipText));
 		}
 		else
 		{
-			OutlinerModel = MakeShared<FChannelGroupOutlinerModel>(ChannelName, MetaData->DisplayText);
+			OutlinerModel = MakeShared<FChannelGroupOutlinerModel>(ChannelName, MetaData->DisplayText, FText::FromString(TooltipText));
 		}
 	}
 
diff --git a/Engine/Source/Editor/Sequencer/Private/MVVM/ViewModels/TrackModelLayoutBuilder.h b/Engine/Source/Editor/Sequencer/Private/MVVM/ViewModels/TrackModelLayoutBuilder.h
--- a/Engine/Source/Editor/Sequencer/Private/MVVM/ViewModels/TrackModelLayoutBuilder.h
+++ b/Engine/Source/Editor/Sequencer/Private/MVVM/ViewModels/TrackModelLayoutBuilder.h
@@ -112,7 +112,7 @@ public:
 
 	// ISectionLayoutBuilder interface
 
-	virtual void PushCategory( FName CategoryName, const FText& DisplayLabel, TFunction<TSharedPtr<UE::Sequencer::FCategoryModel>(FName, const FText&)> OptionalFactory ) override;
+	virtual void PushCategory( FName CategoryName, const FText& DisplayLabel, const FText& TooltipText, TFunction<TSharedPtr<UE::Sequencer::FCategoryModel>(FName, const FText&)> OptionalFactory ) override;
 	virtual void SetTopLevelChannel( const FMovieSceneChannelHandle& Channel, TFunction<TSharedPtr<UE::Sequencer::FChannelModel>(FName, const FMovieSceneChannelHandle&)> OptionalFactory ) override;
 	virtual void AddChannel( const FMovieSceneChannelHandle& Channel, TFunction<TSharedPtr<UE::Sequencer::FChannelModel>(FName, const FMovieSceneChannelHandle&)> OptionalFactory ) override;
 	virtual void PopCategory() override;
diff --git a/Engine/Source/Editor/Sequencer/Public/ISectionLayoutBuilder.h b/Engine/Source/Editor/Sequencer/Public/ISectionLayoutBuilder.h
--- a/Engine/Source/Editor/Sequencer/Public/ISectionLayoutBuilder.h
+++ b/Engine/Source/Editor/Sequencer/Public/ISectionLayoutBuilder.h
@@ -29,7 +29,7 @@ public:
 	 * @param CategoryName	The name of the category
 	 * @param DisplayLabel	The localized display label for the category
 	 */
-	virtual void PushCategory( FName CategoryName, const FText& DisplayLabel, TFunction<TSharedPtr<UE::Sequencer::FCategoryModel>(FName, const FText&)> OptionalFactory) = 0;
+	virtual void PushCategory( FName CategoryName, const FText& DisplayLabel, const FText& TooltipText, TFunction<TSharedPtr<UE::Sequencer::FCategoryModel>(FName, const FText&)> OptionalFactory) = 0;
 
 	/**
 	 * Sets the section as a key area itself
diff --git a/Engine/Source/Editor/Sequencer/Public/MVVM/ViewModels/CategoryModel.h b/Engine/Source/Editor/Sequencer/Public/MVVM/ViewModels/CategoryModel.h
--- a/Engine/Source/Editor/Sequencer/Public/MVVM/ViewModels/CategoryModel.h
+++ b/Engine/Source/Editor/Sequencer/Public/MVVM/ViewModels/CategoryModel.h
@@ -87,7 +87,7 @@ class SEQUENCER_API FCategoryGroupModel
 public:
 	UE_SEQUENCER_DECLARE_CASTABLE(FCategoryGroupModel, FOutlinerItemModel, ITrackAreaExtension, ICompoundOutlinerExtension, IDeletableExtension, IRecyclableExtension);
 
-	explicit FCategoryGroupModel(FName InCategoryName, const FText& InDisplayText);
+	explicit FCategoryGroupModel(FName InCategoryName, const FText& InDisplayText, const FText& InTooltipText);
 
 	~FCategoryGroupModel();
 
@@ -106,6 +106,11 @@ public:
 		return DisplayText;
 	}
 
+	FText GetTooltipText() const
+	{
+		return TooltipText;
+	}
+
 	void AddCategory(TWeakViewModelPtr<FCategoryModel> InCategory);
 
 	TArrayView<const TWeakViewModelPtr<FCategoryModel>> GetCategories() const;
@@ -119,6 +124,7 @@ public:
 	FOutlinerSizing GetOutlinerSizing() const override;
 	FText GetLabel() const override;
 	FSlateFontInfo GetLabelFont() const override;
+	FText GetLabelToolTipText() const override;
 	TSharedRef<SWidget> CreateOutlinerView(const FCreateOutlinerViewParams& InParams) override;
 
 	/*~ ITrackAreaExtension */
@@ -137,6 +143,7 @@ private:
 	TArray<TWeakViewModelPtr<FCategoryModel>> Categories;
 	FName CategoryName;
 	FText DisplayText;
+	FText TooltipText;
 	FOutlinerSizing ComputedSizing;
 };
 
diff --git a/Engine/Source/Editor/Sequencer/Public/MVVM/ViewModels/ChannelModel.h b/Engine/Source/Editor/Sequencer/Public/MVVM/ViewModels/ChannelModel.h
--- a/Engine/Source/Editor/Sequencer/Public/MVVM/ViewModels/ChannelModel.h
+++ b/Engine/Source/Editor/Sequencer/Public/MVVM/ViewModels/ChannelModel.h
@@ -127,7 +127,7 @@ public:
 
 	UE_SEQUENCER_DECLARE_CASTABLE(FChannelGroupModel, FViewModel, ITrackAreaExtension, IRecyclableExtension);
 
-	FChannelGroupModel(FName InChannelName, const FText& InDisplayText);
+	FChannelGroupModel(FName InChannelName, const FText& InDisplayText, const FText& InTooltipText);
 	~FChannelGroupModel();
 
 	/** Returns whether any of the channels within this group have any keyframes on them */
@@ -139,6 +139,9 @@ public:
 	/** Returns the label for this group */
 	FText GetDisplayText() const { return DisplayText; }
 
+	/** Returns the tooltip for this group */
+	FText GetTooltipText() const { return TooltipText; }
+
 	/** Gets all the channel models in this group */
 	TArrayView<const TWeakViewModelPtr<FChannelModel>> GetChannels() const;
 	/** Adds a channel model to this group */
@@ -190,6 +193,7 @@ protected:
 	uint32 ChannelsSerialNumber;
 	FName ChannelName;
 	FText DisplayText;
+	FText TooltipText;
 };
 
 
@@ -206,7 +210,7 @@ public:
 
 	UE_SEQUENCER_DECLARE_CASTABLE(FChannelGroupOutlinerModel, FChannelGroupModel, FOutlinerItemModelMixin, ICompoundOutlinerExtension, IDeletableExtension);
 
-	FChannelGroupOutlinerModel(FName InChannelName, const FText& InDisplayText);
+	FChannelGroupOutlinerModel(FName InChannelName, const FText& InDisplayText, const FText& InTooltipText);
 	~FChannelGroupOutlinerModel();
 
 public:
@@ -218,6 +222,7 @@ public:
 	FOutlinerSizing GetOutlinerSizing() const override;
 	FText GetLabel() const override;
 	FSlateFontInfo GetLabelFont() const override;
+	FText GetLabelToolTipText() const override;
 	TSharedRef<SWidget> CreateOutlinerView(const FCreateOutlinerViewParams& InParams) override;
 
 	/*~ ICurveEditorTreeItem */
diff --git a/Engine/Source/Runtime/MovieScene/Private/Channels/MovieSceneChannelEditorData.cpp b/Engine/Source/Runtime/MovieScene/Private/Channels/MovieSceneChannelEditorData.cpp
--- a/Engine/Source/Runtime/MovieScene/Private/Channels/MovieSceneChannelEditorData.cpp
+++ b/Engine/Source/Runtime/MovieScene/Private/Channels/MovieSceneChannelEditorData.cpp
@@ -18,6 +18,10 @@ const FLinearColor FCommonChannelData::RedChannelColor(1.0f, 0.05f, 0.05f, 0.9f)
 const FLinearColor FCommonChannelData::GreenChannelColor(0.05f, 1.0f, 0.05f, 0.9f);
 const FLinearColor FCommonChannelData::BlueChannelColor(0.1f, 0.2f, 1.0f, 0.9f);
 
+const FName FCommonChannelData::TooltipText = TEXT("TooltipText");
+const FName FCommonChannelData::GroupDisplayName = TEXT("GroupDisplayName");
+const FName FCommonChannelData::GroupTooltipText = TEXT("GroupTooltipText");
+
 
 FMovieSceneChannelMetaData::FMovieSceneChannelMetaData()
 	: bEnabled(true)
diff --git a/Engine/Source/Runtime/MovieScene/Public/Channels/MovieSceneChannelEditorData.h b/Engine/Source/Runtime/MovieScene/Public/Channels/MovieSceneChannelEditorData.h
--- a/Engine/Source/Runtime/MovieScene/Public/Channels/MovieSceneChannelEditorData.h
+++ b/Engine/Source/Runtime/MovieScene/Public/Channels/MovieSceneChannelEditorData.h
@@ -142,6 +142,10 @@ struct FCommonChannelData
 	static MOVIESCENE_API const FLinearColor RedChannelColor;
 	static MOVIESCENE_API const FLinearColor GreenChannelColor;
 	static MOVIESCENE_API const FLinearColor BlueChannelColor;
+
+	static MOVIESCENE_API const FName TooltipText;
+	static MOVIESCENE_API const FName GroupDisplayName;
+	static MOVIESCENE_API const FName GroupTooltipText;
 };
 
 
diff --git a/Engine/Source/Runtime/MovieSceneTracks/Private/MovieSceneTracksComponentTypes.cpp b/Engine/Source/Runtime/MovieSceneTracks/Private/MovieSceneTracksComponentTypes.cpp
--- a/Engine/Source/Runtime/MovieSceneTracks/Private/MovieSceneTracksComponentTypes.cpp
+++ b/Engine/Source/Runtime/MovieSceneTracks/Private/MovieSceneTracksComponentTypes.cpp
@@ -551,6 +551,10 @@ FMovieSceneTracksComponentTypes::FMovieSceneTracksComponentTypes()
 	ComponentRegistry->NewComponentType(&ColorParameterName,     TEXT("Color Parameter Name"), EComponentTypeFlags::CopyToChildren | EComponentTypeFlags::CopyToOutput);
 	ComponentRegistry->NewComponentType(&TransformParameterName, TEXT("Transform Parameter Name"), EComponentTypeFlags::CopyToChildren | EComponentTypeFlags::CopyToOutput);
 
+	ComponentRegistry->NewComponentType(&ScalarMaterialParameterInfo, TEXT("Scalar Material Parameter Info"), EComponentTypeFlags::CopyToChildren | EComponentTypeFlags::CopyToOutput);
+	ComponentRegistry->NewComponentType(&ColorMaterialParameterInfo, TEXT("Color Material Parameter Info"), EComponentTypeFlags::CopyToChildren | EComponentTypeFlags::CopyToOutput);
+	ComponentRegistry->NewComponentType(&VectorMaterialParameterInfo, TEXT("Vector Material Parameter Info"), EComponentTypeFlags::CopyToChildren | EComponentTypeFlags::CopyToOutput);
+
 	ComponentRegistry->NewComponentType(&Fade,                   TEXT("Fade"), EComponentTypeFlags::CopyToChildren);
 
 	ComponentRegistry->NewComponentType(&Audio,                  TEXT("Audio"), EComponentTypeFlags::CopyToChildren);
diff --git a/Engine/Source/Runtime/MovieSceneTracks/Private/Sections/MovieSceneComponentMaterialParameterSection.cpp b/Engine/Source/Runtime/MovieSceneTracks/Private/Sections/MovieSceneComponentMaterialParameterSection.cpp
new file mode 100644
--- /dev/null
+++ b/Engine/Source/Runtime/MovieSceneTracks/Private/Sections/MovieSceneComponentMaterialParameterSection.cpp
@@ -0,0 +1,414 @@
+// Copyright Epic Games, Inc. All Rights Reserved.
+
+#include "Sections/MovieSceneComponentMaterialParameterSection.h"
+#include "Channels/MovieSceneChannelProxy.h"
+
+#include "Evaluation/MovieSceneEvaluationField.h"
+#include "EntitySystem/BuiltInComponentTypes.h"
+#include "MovieSceneTracksComponentTypes.h"
+#include "Tracks/MovieSceneMaterialTrack.h"
+
+#include UE_INLINE_GENERATED_CPP_BY_NAME(MovieSceneComponentMaterialParameterSection)
+
+#if WITH_EDITOR
+#define LOCTEXT_NAMESPACE "MaterialTrackEditor"
+#endif
+
+namespace UE::MovieScene
+{
+
+	/* Entity IDs are an encoded type and index, with the upper 8 bits being the type, and the lower 24 bits as the index */
+	uint32 EncodeMaterialParameterEntityID(int32 InIndex, uint8 InType)
+	{
+		check(InIndex >= 0 && InIndex < int32(0x00FFFFFF));
+		return static_cast<uint32>(InIndex) | (uint32(InType) << 24);
+	}
+	void DecodeMaterialParameterEntityID(uint32 InEntityID, int32& OutIndex, uint8& OutType)
+	{
+		// Mask out the type to get the index
+		OutIndex = static_cast<int32>(InEntityID & 0x00FFFFFF);
+		OutType = InEntityID >> 24;
+	}
+
+
+}// namespace UE::MovieScene
+
+FScalarMaterialParameterInfoAndCurve::FScalarMaterialParameterInfoAndCurve(const FMaterialParameterInfo& InParameterInfo)
+{
+	ParameterInfo = InParameterInfo;
+}
+
+FColorMaterialParameterInfoAndCurves::FColorMaterialParameterInfoAndCurves(const FMaterialParameterInfo& InParameterInfo)
+{
+	ParameterInfo = InParameterInfo;
+}
+
+UMovieSceneComponentMaterialParameterSection::UMovieSceneComponentMaterialParameterSection(const FObjectInitializer& ObjectInitializer)
+	: Super(ObjectInitializer)
+{
+	bSupportsInfiniteRange = true;
+}
+
+EMovieSceneChannelProxyType UMovieSceneComponentMaterialParameterSection::CacheChannelProxy()
+{
+	FMovieSceneChannelProxyData Channels;
+
+#if WITH_EDITOR
+
+	auto GetMaterialParameterDisplayText = [](const FText& ParameterName, const FString& InLayerName, const FString& InAssetName)
+	{
+		FText DisplayName = ParameterName;
+		if (!InLayerName.IsEmpty() && !InAssetName.IsEmpty())
+		{
+			DisplayName = FText::Format(LOCTEXT("MaterialParameterDisplayText", "{0} ({1}.{2})"), ParameterName, FText::FromString(InLayerName), FText::FromString(InAssetName));
+		}
+		return DisplayName;
+	};
+	auto GetMaterialParameterPath = [](const FName& ParameterName, const FString& InLayerName, const FString& InAssetName)
+	{
+		FString Path = ParameterName.ToString();
+		if (!InLayerName.IsEmpty() && !InAssetName.IsEmpty())
+		{
+			Path = FString::Printf(TEXT("%s.%s.%s"), *InLayerName, *InAssetName, *ParameterName.ToString());
+		}
+		return Path;
+	};
+	auto GetMaterialParameterTooltipText = [](const FString& ParameterPath)
+	{
+		return FString::Printf(TEXT("Path: %s"), *ParameterPath);
+	};
+	int32 SortOrder = 0;
+	for (FScalarMaterialParameterInfoAndCurve& Scalar : ScalarParameterInfosAndCurves)
+	{
+		FString ParameterPath = *GetMaterialParameterPath(Scalar.ParameterInfo.Name, Scalar.ParameterLayerName, Scalar.ParameterAssetName);
+		FText ParameterDisplayName = GetMaterialParameterDisplayText(FText::FromName(Scalar.ParameterInfo.Name), Scalar.ParameterLayerName, Scalar.ParameterAssetName);
+		FMovieSceneChannelMetaData MetaData(*ParameterPath, ParameterDisplayName);
+		// Prevent single channels from collapsing to the track node
+		MetaData.PropertyMetaData.Add(FCommonChannelData::TooltipText, GetMaterialParameterTooltipText(ParameterPath));
+		MetaData.bCanCollapseToTrack = false;
+		MetaData.SortOrder = SortOrder++;
+		Channels.Add(Scalar.ParameterCurve, MetaData, TMovieSceneExternalValue<float>());
+	}
+	for (FColorMaterialParameterInfoAndCurves& Color : ColorParameterInfosAndCurves)
+	{
+		FString ParameterPath = GetMaterialParameterPath(Color.ParameterInfo.Name, Color.ParameterLayerName, Color.ParameterAssetName);
+		FString TooltipText = GetMaterialParameterTooltipText(ParameterPath);
+		FText ParameterDisplayName = GetMaterialParameterDisplayText(FText::FromName(Color.ParameterInfo.Name), Color.ParameterLayerName, Color.ParameterAssetName);
+		FText Group = FText::FromString(ParameterPath);
+
+		FMovieSceneChannelMetaData MetaData_R(*(ParameterPath + TEXT("R")), FCommonChannelData::ChannelR, Group);
+		MetaData_R.PropertyMetaData.Add(FCommonChannelData::TooltipText, FString::Printf(TEXT("%s.%s"), *TooltipText, TEXT("R")));
+		MetaData_R.SortOrder = SortOrder++;
+		MetaData_R.Color = FCommonChannelData::RedChannelColor;
+		MetaData_R.PropertyMetaData.Add(FCommonChannelData::GroupDisplayName, ParameterDisplayName.ToString());
+		MetaData_R.PropertyMetaData.Add(FCommonChannelData::GroupTooltipText, TooltipText);
+
+		FMovieSceneChannelMetaData MetaData_G(*(ParameterPath + TEXT("G")), FCommonChannelData::ChannelG, Group);
+		MetaData_G.PropertyMetaData.Add(FCommonChannelData::TooltipText, FString::Printf(TEXT("%s.%s"), *TooltipText, TEXT("G")));
+		MetaData_G.SortOrder = SortOrder++;
+		MetaData_G.Color = FCommonChannelData::GreenChannelColor;
+		MetaData_G.PropertyMetaData.Add(FCommonChannelData::GroupDisplayName, ParameterDisplayName.ToString());
+		MetaData_G.PropertyMetaData.Add(FCommonChannelData::GroupTooltipText, TooltipText);
+
+		FMovieSceneChannelMetaData MetaData_B(*(ParameterPath + TEXT("B")), FCommonChannelData::ChannelB, Group);
+		MetaData_B.PropertyMetaData.Add(FCommonChannelData::TooltipText, FString::Printf(TEXT("%s.%s"), *TooltipText, TEXT("B")));
+		MetaData_B.SortOrder = SortOrder++;
+		MetaData_B.Color = FCommonChannelData::BlueChannelColor;
+		MetaData_B.PropertyMetaData.Add(FCommonChannelData::GroupDisplayName, ParameterDisplayName.ToString());
+		MetaData_B.PropertyMetaData.Add(FCommonChannelData::GroupTooltipText, TooltipText);
+
+		FMovieSceneChannelMetaData MetaData_A(*(ParameterPath + TEXT("A")), FCommonChannelData::ChannelA, Group);
+		MetaData_A.PropertyMetaData.Add(FCommonChannelData::TooltipText, FString::Printf(TEXT("%s.%s"), *TooltipText, TEXT("A")));
+		MetaData_A.SortOrder = SortOrder++;
+		MetaData_A.PropertyMetaData.Add(FCommonChannelData::GroupDisplayName, ParameterDisplayName.ToString());
+		MetaData_A.PropertyMetaData.Add(FCommonChannelData::GroupTooltipText, TooltipText);
+
+		Channels.Add(Color.RedCurve, MetaData_R, TMovieSceneExternalValue<float>());
+		Channels.Add(Color.GreenCurve, MetaData_G, TMovieSceneExternalValue<float>());
+		Channels.Add(Color.BlueCurve, MetaData_B, TMovieSceneExternalValue<float>());
+		Channels.Add(Color.AlphaCurve, MetaData_A, TMovieSceneExternalValue<float>());
+	}
+#else
+
+	for (FScalarMaterialParameterInfoAndCurve& Scalar : ScalarParameterInfosAndCurves)
+	{
+		Channels.Add(Scalar.ParameterCurve);
+	}
+
+	for (FColorMaterialParameterInfoAndCurves& Color : ColorParameterInfosAndCurves)
+	{
+		Channels.Add(Color.RedCurve);
+		Channels.Add(Color.GreenCurve);
+		Channels.Add(Color.BlueCurve);
+		Channels.Add(Color.AlphaCurve);
+	}
+#endif
+
+	ChannelProxy = MakeShared<FMovieSceneChannelProxy>(MoveTemp(Channels));
+
+	return EMovieSceneChannelProxyType::Dynamic;
+}
+
+void UMovieSceneComponentMaterialParameterSection::PostEditImport()
+{
+	Super::PostEditImport();
+
+	CacheChannelProxy();
+}
+#if WITH_EDITOR
+void UMovieSceneComponentMaterialParameterSection::PostEditUndo()
+{
+	Super::PostEditUndo();
+
+	CacheChannelProxy();
+}
+#endif
+void UMovieSceneComponentMaterialParameterSection::AddScalarParameterKey(const FMaterialParameterInfo& InParameterInfo, FFrameNumber InTime, float InValue, const FString& InLayerName, const FString& InAssetName)
+{
+	FMovieSceneFloatChannel* ExistingChannel = nullptr;
+	for (FScalarMaterialParameterInfoAndCurve& ScalarParameterInfoAndCurve : ScalarParameterInfosAndCurves)
+	{
+		if (ScalarParameterInfoAndCurve.ParameterInfo == InParameterInfo)
+		{
+			ExistingChannel = &ScalarParameterInfoAndCurve.ParameterCurve;
+			break;
+		}
+	}
+	if (ExistingChannel == nullptr)
+	{
+		const int32 NewIndex = ScalarParameterInfosAndCurves.Add(FScalarMaterialParameterInfoAndCurve(InParameterInfo));
+#if WITH_EDITOR
+		ScalarParameterInfosAndCurves[NewIndex].ParameterLayerName = InLayerName;
+		ScalarParameterInfosAndCurves[NewIndex].ParameterAssetName = InAssetName;
+#endif
+		ExistingChannel = &ScalarParameterInfosAndCurves[NewIndex].ParameterCurve;
+		CacheChannelProxy();
+	}
+
+	ExistingChannel->GetData().UpdateOrAddKey(InTime, FMovieSceneFloatValue(InValue));
+
+	if (TryModify())
+	{
+		SetRange(TRange<FFrameNumber>::Hull(TRange<FFrameNumber>(InTime), GetRange()));
+	}
+}
+
+void UMovieSceneComponentMaterialParameterSection::AddColorParameterKey(const FMaterialParameterInfo& InParameterInfo, FFrameNumber InTime, FLinearColor InValue, const FString& InLayerName, const FString& InAssetName)
+{
+	FColorMaterialParameterInfoAndCurves* ExistingCurves = nullptr;
+	for (FColorMaterialParameterInfoAndCurves& ColorParameterInfoAndCurve : ColorParameterInfosAndCurves)
+	{
+		if (ColorParameterInfoAndCurve.ParameterInfo == InParameterInfo)
+		{
+			ExistingCurves = &ColorParameterInfoAndCurve;
+			break;
+		}
+	}
+	if (ExistingCurves == nullptr)
+	{
+		int32 NewIndex = ColorParameterInfosAndCurves.Add(FColorMaterialParameterInfoAndCurves(InParameterInfo));
+		ExistingCurves = &ColorParameterInfosAndCurves[NewIndex];
+#if WITH_EDITOR
+		ColorParameterInfosAndCurves[NewIndex].ParameterLayerName = InLayerName;
+		ColorParameterInfosAndCurves[NewIndex].ParameterAssetName = InAssetName;
+#endif
+		CacheChannelProxy();
+	}
+
+	ExistingCurves->RedCurve.GetData().UpdateOrAddKey(InTime, FMovieSceneFloatValue(InValue.R));
+	ExistingCurves->GreenCurve.GetData().UpdateOrAddKey(InTime, FMovieSceneFloatValue(InValue.G));
+	ExistingCurves->BlueCurve.GetData().UpdateOrAddKey(InTime, FMovieSceneFloatValue(InValue.B));
+	ExistingCurves->AlphaCurve.GetData().UpdateOrAddKey(InTime, FMovieSceneFloatValue(InValue.A));
+
+	if (TryModify())
+	{
+		SetRange(TRange<FFrameNumber>::Hull(TRange<FFrameNumber>(InTime), GetRange()));
+	}
+}
+
+bool UMovieSceneComponentMaterialParameterSection::RemoveScalarParameter(const FMaterialParameterInfo& InParameterInfo)
+{
+	for (int32 i = 0; i < ScalarParameterInfosAndCurves.Num(); i++)
+	{
+		if (ScalarParameterInfosAndCurves[i].ParameterInfo == InParameterInfo)
+		{
+			ScalarParameterInfosAndCurves.RemoveAt(i);
+			CacheChannelProxy();
+			return true;
+		}
+	}
+	return false;
+}
+
+bool UMovieSceneComponentMaterialParameterSection::RemoveColorParameter(const FMaterialParameterInfo& InParameterInfo)
+{
+	for (int32 i = 0; i < ColorParameterInfosAndCurves.Num(); i++)
+	{
+		if (ColorParameterInfosAndCurves[i].ParameterInfo == InParameterInfo)
+		{
+			ColorParameterInfosAndCurves.RemoveAt(i);
+			CacheChannelProxy();
+			return true;
+		}
+	}
+	return false;
+}
+
+#if WITH_EDITOR
+bool UMovieSceneComponentMaterialParameterSection::RemoveScalarParameter(FName InParameterPath)
+{
+	TArray<FString> ParameterPathArray;
+	InParameterPath.ToString().ParseIntoArray(ParameterPathArray, TEXT("."));
+	for (int32 i = 0; i < ScalarParameterInfosAndCurves.Num(); i++)
+	{
+		bool bFound = false;
+		if (ParameterPathArray.Num() == 3)
+		{
+			bFound = ScalarParameterInfosAndCurves[i].ParameterLayerName == *ParameterPathArray[0]
+				&& ScalarParameterInfosAndCurves[i].ParameterAssetName == ParameterPathArray[1]
+				&& ScalarParameterInfosAndCurves[i].ParameterInfo.Name == ParameterPathArray[2];
+		}
+		else if (ParameterPathArray.Num() == 1)
+		{
+			bFound = ScalarParameterInfosAndCurves[i].ParameterInfo.Name == *ParameterPathArray[0];
+		}
+		if (bFound)
+		{
+			ScalarParameterInfosAndCurves.RemoveAt(i);
+			CacheChannelProxy();
+			return true;
+		}
+	}
+	return false;
+}
+
+bool UMovieSceneComponentMaterialParameterSection::RemoveColorParameter(FName InParameterPath)
+{
+	TArray<FString> ParameterPathArray;
+	InParameterPath.ToString().ParseIntoArray(ParameterPathArray, TEXT("."));
+	for (int32 i = 0; i < ColorParameterInfosAndCurves.Num(); i++)
+	{
+		bool bFound = false;
+		if (ParameterPathArray.Num() == 3)
+		{
+			bFound = ColorParameterInfosAndCurves[i].ParameterLayerName == *ParameterPathArray[0]
+				&& ColorParameterInfosAndCurves[i].ParameterAssetName == ParameterPathArray[1]
+				&& ColorParameterInfosAndCurves[i].ParameterInfo.Name == ParameterPathArray[2];
+		}
+		else if (ParameterPathArray.Num() == 1)
+		{
+			bFound = ColorParameterInfosAndCurves[i].ParameterInfo.Name == *ParameterPathArray[0];
+		}
+		if (bFound)
+		{
+			ColorParameterInfosAndCurves.RemoveAt(i);
+			CacheChannelProxy();
+			return true;
+		}
+	}
+	return false;
+}
+#endif
+
+void UMovieSceneComponentMaterialParameterSection::ImportEntityImpl(UMovieSceneEntitySystemLinker* EntityLinker, const FEntityImportParams& Params, FImportedEntity* OutImportedEntity)
+{
+	using namespace UE::MovieScene;
+
+	uint8 ParameterType = 0;
+	int32 EntityIndex = 0;
+	DecodeMaterialParameterEntityID(Params.EntityID, EntityIndex, ParameterType);
+
+	FBuiltInComponentTypes* BuiltInComponentTypes = FBuiltInComponentTypes::Get();
+	FMovieSceneTracksComponentTypes* TracksComponentTypes = FMovieSceneTracksComponentTypes::Get();
+
+	FGuid ObjectBindingID = Params.GetObjectBindingID();
+	
+	// Find material index from the outer track
+	int32 MaterialIndex = 0;
+	if (UMovieSceneComponentMaterialTrack* MaterialTrack = GetTypedOuter<UMovieSceneComponentMaterialTrack>())
+	{
+		MaterialIndex = MaterialTrack->GetMaterialIndex();
+	}
+
+	TEntityBuilder<TAddConditional<FGuid>> BaseBuilder = FEntityBuilder()
+		.AddConditional(BuiltInComponentTypes->GenericObjectBinding, ObjectBindingID, ObjectBindingID.IsValid());
+
+	switch (ParameterType)
+	{
+	case 0:
+	{
+		const FScalarMaterialParameterInfoAndCurve& Scalar = ScalarParameterInfosAndCurves[EntityIndex];
+
+		if (Scalar.ParameterCurve.HasAnyData())
+		{
+			OutImportedEntity->AddBuilder(
+				BaseBuilder
+				.Add(TracksComponentTypes->ScalarMaterialParameterInfo, Scalar.ParameterInfo)
+				.Add(BuiltInComponentTypes->FloatChannel[0], &Scalar.ParameterCurve)
+				.Add(TracksComponentTypes->ComponentMaterialIndex, MaterialIndex)
+				// If the section has no valid blend type (legacy data), make it use absolute blending.
+				// Otherwise, the base section class will add the appropriate blend type tag in BuildDefaultComponents.
+				.AddTagConditional(BuiltInComponentTypes->Tags.AbsoluteBlend, !GetBlendType().IsValid())
+			);
+		}
+		break;
+	}
+	case 1:
+	{
+		const FColorMaterialParameterInfoAndCurves& Color = ColorParameterInfosAndCurves[EntityIndex];
+		if (Color.RedCurve.HasAnyData() || Color.GreenCurve.HasAnyData() || Color.BlueCurve.HasAnyData() || Color.AlphaCurve.HasAnyData())
+		{
+			OutImportedEntity->AddBuilder(
+				BaseBuilder
+				.Add(TracksComponentTypes->ColorMaterialParameterInfo, Color.ParameterInfo)
+				.AddConditional(BuiltInComponentTypes->FloatChannel[0], &Color.RedCurve, Color.RedCurve.HasAnyData())
+				.AddConditional(BuiltInComponentTypes->FloatChannel[1], &Color.GreenCurve, Color.GreenCurve.HasAnyData())
+				.AddConditional(BuiltInComponentTypes->FloatChannel[2], &Color.BlueCurve, Color.BlueCurve.HasAnyData())
+				.AddConditional(BuiltInComponentTypes->FloatChannel[3], &Color.AlphaCurve, Color.AlphaCurve.HasAnyData())
+				.Add(TracksComponentTypes->ComponentMaterialIndex, MaterialIndex)
+				// If the section has no valid blend type (legacy data), make it use absolute blending.
+				// Otherwise, the base section class will add the appropriate blend type tag in BuildDefaultComponents.
+				.AddTagConditional(BuiltInComponentTypes->Tags.AbsoluteBlend, !GetBlendType().IsValid())
+			);
+		}
+		break;
+	}
+	}
+}
+
+bool UMovieSceneComponentMaterialParameterSection::PopulateEvaluationFieldImpl(const TRange<FFrameNumber>& EffectiveRange, const FMovieSceneEvaluationFieldEntityMetaData& InMetaData, FMovieSceneEntityComponentFieldBuilder* OutFieldBuilder)
+{
+	// By default, material parameter sections do not populate any evaluation field entries
+	// that is the job of its outer UMovieSceneTrack through a call to ExternalPopulateEvaluationField.
+
+	// Ideally we wouldn't have this exception for MaterialParameterSections, but deprecated behavior from MaterialTracks using ParameterSections have this behavior, and when a track implements PopulateEvaluationFieldImpl, the version
+	// on Sections is not called.
+
+	return true;
+}
+
+void UMovieSceneComponentMaterialParameterSection::ExternalPopulateEvaluationField(const TRange<FFrameNumber>& EffectiveRange, const FMovieSceneEvaluationFieldEntityMetaData& InMetaData, FMovieSceneEntityComponentFieldBuilder* OutFieldBuilder)
+{
+	using namespace UE::MovieScene;
+
+	const int32 MetaDataIndex = OutFieldBuilder->AddMetaData(InMetaData);
+
+	// We use the top 8 bits of EntityID to encode the type of parameter
+	const int32 NumScalarID = ScalarParameterInfosAndCurves.Num();
+	const int32 NumColorID = ColorParameterInfosAndCurves.Num();
+
+	for (int32 Index = 0; Index < NumScalarID; ++Index)
+	{
+		const int32 EntityIndex = OutFieldBuilder->FindOrAddEntity(this, EncodeMaterialParameterEntityID(Index, 0));
+		OutFieldBuilder->AddPersistentEntity(EffectiveRange, EntityIndex, MetaDataIndex);
+	}
+	for (int32 Index = 0; Index < NumColorID; ++Index)
+	{
+		const int32 EntityIndex = OutFieldBuilder->FindOrAddEntity(this, EncodeMaterialParameterEntityID(Index, 1));
+		OutFieldBuilder->AddPersistentEntity(EffectiveRange, EntityIndex, MetaDataIndex);
+	}
+}
+
+#if WITH_EDITOR
+#undef LOCTEXT_NAMESPACE
+#endif
\ No newline at end of file
diff --git a/Engine/Source/Runtime/MovieSceneTracks/Private/Systems/MovieSceneMaterialParameterCollectionSystem.cpp b/Engine/Source/Runtime/MovieSceneTracks/Private/Systems/MovieSceneMaterialParameterCollectionSystem.cpp
--- a/Engine/Source/Runtime/MovieSceneTracks/Private/Systems/MovieSceneMaterialParameterCollectionSystem.cpp
+++ b/Engine/Source/Runtime/MovieSceneTracks/Private/Systems/MovieSceneMaterialParameterCollectionSystem.cpp
@@ -78,10 +78,15 @@ void UMovieSceneMaterialParameterCollectionSystem::OnRun(FSystemTaskPrerequisite
 			TComponentReader<TWeakObjectPtr<UMaterialParameterCollection>> MPCs = Allocation->ReadComponents(TracksComponents->MPC);
 			TComponentReader<FInstanceHandle> InstanceHandles = Allocation->ReadComponents(BuiltInComponents->InstanceHandle);
 
+			// Support older sections that still use ParameterName rather than ParameterInfo
 			TOptionalComponentReader<FName> ScalarParameterNames = Allocation->TryReadComponents(TracksComponents->ScalarParameterName);
 			TOptionalComponentReader<FName> VectorParameterNames = Allocation->TryReadComponents(TracksComponents->VectorParameterName);
 			TOptionalComponentReader<FName> ColorParameterNames = Allocation->TryReadComponents(TracksComponents->ColorParameterName);
 
+			TOptionalComponentReader<FMaterialParameterInfo> ScalarParameterInfos = Allocation->TryReadComponents(TracksComponents->ScalarMaterialParameterInfo);
+			TOptionalComponentReader<FMaterialParameterInfo> VectorParameterInfos = Allocation->TryReadComponents(TracksComponents->VectorMaterialParameterInfo);
+			TOptionalComponentReader<FMaterialParameterInfo> ColorParameterInfos = Allocation->TryReadComponents(TracksComponents->ColorMaterialParameterInfo);
+
 			const int32 Num = Allocation->Num();
 			for (int32 Index = 0; Index < Num; ++Index)
 			{
@@ -101,9 +106,10 @@ void UMovieSceneMaterialParameterCollectionSystem::OnRun(FSystemTaskPrerequisite
 						TEXT("Unable to create MPC instance for %s with World %s. Material parameter collection tracks will not function."),
 						*Collection->GetName(), *World->GetName()))
 					{
-						if (ScalarParameterNames)
+						if (ScalarParameterInfos || ScalarParameterNames)
 						{
-							FName Name = ScalarParameterNames[Index];
+							// Support older sections that still use ParameterName rather than ParameterInfo
+							FName Name = ScalarParameterInfos ? ScalarParameterInfos[Index].Name : ScalarParameterNames[Index];
 							if (Collection->GetScalarParameterByName(Name) == nullptr)
 							{
 								if (!Instance->bLoggedMissingParameterWarning)
@@ -112,9 +118,14 @@ void UMovieSceneMaterialParameterCollectionSystem::OnRun(FSystemTaskPrerequisite
 								}
 							}
 						}
-						else if (VectorParameterNames || ColorParameterNames)
+						else if (VectorParameterInfos || ColorParameterInfos || VectorParameterNames || ColorParameterNames)
 						{
-							FName Name = VectorParameterNames ? VectorParameterNames[Index] : ColorParameterNames[Index];
+							// Support older sections that still use ParameterName rather than ParameterInfo
+							FName Name = VectorParameterInfos ? VectorParameterInfos[Index].Name :
+								ColorParameterInfos ? ColorParameterInfos[Index].Name :
+								VectorParameterNames ? VectorParameterNames[Index] :
+								ColorParameterNames[Index];
+
 							if (Collection->GetVectorParameterByName(Name) == nullptr)
 							{
 								if (!Instance->bLoggedMissingParameterWarning)
@@ -172,8 +183,8 @@ void UMovieSceneMaterialParameterCollectionSystem::OnRun(FSystemTaskPrerequisite
 
 	Params.AdditionalFilter.None({ BuiltInComponents->BlendChannelOutput });
 	Params.AdditionalFilter.All({ TracksComponents->MPC });
-	ScalarParameterStorage->BeginTrackingAndCachePreAnimatedValuesTask(Linker, Params, TracksComponents->BoundMaterial, TracksComponents->ScalarParameterName);
-	VectorParameterStorage->BeginTrackingAndCachePreAnimatedValuesTask(Linker, Params, TracksComponents->BoundMaterial, TracksComponents->VectorParameterName);
-	VectorParameterStorage->BeginTrackingAndCachePreAnimatedValuesTask(Linker, Params, TracksComponents->BoundMaterial, TracksComponents->ColorParameterName);
+	ScalarParameterStorage->BeginTrackingAndCachePreAnimatedValuesTask(Linker, Params, TracksComponents->BoundMaterial, TracksComponents->ScalarMaterialParameterInfo);
+	VectorParameterStorage->BeginTrackingAndCachePreAnimatedValuesTask(Linker, Params, TracksComponents->BoundMaterial, TracksComponents->VectorMaterialParameterInfo);
+	VectorParameterStorage->BeginTrackingAndCachePreAnimatedValuesTask(Linker, Params, TracksComponents->BoundMaterial, TracksComponents->ColorMaterialParameterInfo);
 }
 
diff --git a/Engine/Source/Runtime/MovieSceneTracks/Private/Systems/MovieSceneMaterialParameterSystem.cpp b/Engine/Source/Runtime/MovieSceneTracks/Private/Systems/MovieSceneMaterialParameterSystem.cpp
--- a/Engine/Source/Runtime/MovieSceneTracks/Private/Systems/MovieSceneMaterialParameterSystem.cpp
+++ b/Engine/Source/Runtime/MovieSceneTracks/Private/Systems/MovieSceneMaterialParameterSystem.cpp
@@ -74,16 +74,47 @@ FAnimatedMaterialParameterInfo::~FAnimatedMaterialParameterInfo()
 /** Apply scalar material parameters */
 struct FApplyScalarParameters
 {
-	static void ForEachEntity(const FObjectComponent& BoundMaterial, FName ParameterName, double InScalarValue)
+	static void ForEachAllocation(FEntityAllocationIteratorItem Item,
+		TRead<FObjectComponent> BoundMaterials,
+		TReadOneOrMoreOf<FMaterialParameterInfo, FName> ScalarParameterInfosOrNames,
+		TRead<double> ScalarValues)
 	{
-		// WARNING: BoundMaterial may be nullptr here
-		if (UMaterialInstanceDynamic* MID = Cast<UMaterialInstanceDynamic>(BoundMaterial.GetObject()))
-		{
-			MID->SetScalarParameterValue(ParameterName, (float)InScalarValue);
-		}
-		else if (UMaterialParameterCollectionInstance* MPCI = Cast<UMaterialParameterCollectionInstance>(BoundMaterial.GetObject()))
+		const int32 Num = Item.GetAllocation()->Num();
+		const FMaterialParameterInfo* ParameterInfos = ScalarParameterInfosOrNames.Get<0>();
+		const FName* ParameterNames = ScalarParameterInfosOrNames.Get<1>();
+		for (int32 Index = 0; Index < Num; ++Index)
 		{
-			MPCI->SetScalarParameterValue(ParameterName, (float)InScalarValue);
+			const FObjectComponent& BoundMaterial = BoundMaterials[Index];
+			if (!BoundMaterial)
+			{
+				continue;
+			}
+
+			// WARNING: BoundMaterial may be nullptr here
+			if (UMaterialInstanceDynamic* MID = Cast<UMaterialInstanceDynamic>(BoundMaterial.GetObject()))
+			{
+				if (ParameterInfos)
+				{
+					MID->SetScalarParameterValueByInfo(ParameterInfos[Index], (float)ScalarValues[Index]);
+				}
+				else
+				{
+					MID->SetScalarParameterValue(ParameterNames[Index], (float)ScalarValues[Index]);
+				}
+			}
+			else if (UMaterialParameterCollectionInstance* MPCI = Cast<UMaterialParameterCollectionInstance>(BoundMaterial.GetObject()))
+			{
+				if (ParameterInfos)
+				{
+					MPCI->SetScalarParameterValue(ParameterInfos[Index].Name, (float)ScalarValues[Index]);
+				}
+				else
+				{
+					MPCI->SetScalarParameterValue(ParameterNames[Index], (float)ScalarValues[Index]);
+				}
+			}
+
+
 		}
 	}
 };
@@ -93,12 +124,13 @@ struct FApplyVectorParameters
 {
 	static void ForEachAllocation(FEntityAllocationIteratorItem Item,
 		TRead<FObjectComponent> BoundMaterials,
-		TReadOneOrMoreOf<FName, FName> VectorOrColorParameterNames,
+		TReadOneOrMoreOf<FMaterialParameterInfo, FMaterialParameterInfo, FName, FName> VectorOrColorParameterInfosOrNames,
 		TReadOneOrMoreOf<double, double, double, double> VectorChannels)
 	{
 		const int32 Num = Item.GetAllocation()->Num();
 		// Use either the vector parameter name, or the color parameter name
-		const FName* ParameterNames = VectorOrColorParameterNames.Get<0>() ? VectorOrColorParameterNames.Get<0>() : VectorOrColorParameterNames.Get<1>();
+		const FMaterialParameterInfo* ParameterInfos = VectorOrColorParameterInfosOrNames.Get<0>() ? VectorOrColorParameterInfosOrNames.Get<0>() : VectorOrColorParameterInfosOrNames.Get<1>();
+		const FName* ParameterNames = VectorOrColorParameterInfosOrNames.Get<2>() ? VectorOrColorParameterInfosOrNames.Get<2>() : VectorOrColorParameterInfosOrNames.Get<3>();
 		const double* RESTRICT R = VectorChannels.Get<0>();
 		const double* RESTRICT G = VectorChannels.Get<1>();
 		const double* RESTRICT B = VectorChannels.Get<2>();
@@ -121,11 +153,25 @@ struct FApplyVectorParameters
 			);
 			if (UMaterialInstanceDynamic* MID = Cast<UMaterialInstanceDynamic>(BoundMaterial.GetObject()))
 			{
-				MID->SetVectorParameterValue(ParameterNames[Index], Color);
+				if (ParameterInfos)
+				{
+					MID->SetVectorParameterValueByInfo(ParameterInfos[Index], Color);
+				}
+				else if (ParameterNames)
+				{
+					MID->SetVectorParameterValue(ParameterNames[Index], Color);
+				}
 			}
 			else if (UMaterialParameterCollectionInstance* MPCI = Cast<UMaterialParameterCollectionInstance>(BoundMaterial.GetObject()))
 			{
-				MPCI->SetVectorParameterValue(ParameterNames[Index], Color);
+				if (ParameterInfos)
+				{
+					MPCI->SetVectorParameterValue(ParameterInfos[Index].Name, Color);
+				}
+				else if (ParameterNames)
+				{
+					MPCI->SetVectorParameterValue(ParameterNames[Index], Color);
+				}
 			}
 		}
 	}
@@ -133,7 +179,7 @@ struct FApplyVectorParameters
 
 struct FScalarMixin
 {
-	void CreateEntity(UMovieSceneEntitySystemLinker* Linker, const FObjectComponent& BoundMaterial, FName ParameterName, FComponentTypeID BlenderTypeTag, TArrayView<const FMovieSceneEntityID> Inputs, FAnimatedMaterialParameterInfo* Output)
+	void CreateEntity(UMovieSceneEntitySystemLinker* Linker, const FObjectComponent& BoundMaterial, const FMaterialParameterInfo& ParameterInfo, FComponentTypeID BlenderTypeTag, TArrayView<const FMovieSceneEntityID> Inputs, FAnimatedMaterialParameterInfo* Output)
 	{
 		FBuiltInComponentTypes* BuiltInComponents = FBuiltInComponentTypes::Get();
 		FMovieSceneTracksComponentTypes* TracksComponents = FMovieSceneTracksComponentTypes::Get();
@@ -155,11 +201,11 @@ struct FScalarMixin
 		{
 			if (UMaterialInstanceDynamic* MID = Cast<UMaterialInstanceDynamic>(BoundMaterial.GetObject()))
 			{
-				bHasInitialValue = MID->GetScalarParameterValue(ParameterName, InitialValue);
+				bHasInitialValue = MID->GetScalarParameterValue(ParameterInfo, InitialValue);
 			}
 			else if (UMaterialParameterCollectionInstance* MPCI = Cast<UMaterialParameterCollectionInstance>(BoundMaterial.GetObject()))
 			{
-				bHasInitialValue = MPCI->GetScalarParameterValue(ParameterName, InitialValue);
+				bHasInitialValue = MPCI->GetScalarParameterValue(ParameterInfo.Name, InitialValue);
 			}
 		}
 
@@ -177,7 +223,7 @@ struct FScalarMixin
 		Linker->EntityManager.CopyComponents(Inputs[0], Output->OutputEntityID, Linker->EntityManager.GetComponents()->GetCopyAndMigrationMask());
 	}
 
-	void InitializeSoleInput(UMovieSceneEntitySystemLinker* Linker, const FObjectComponent& BoundMaterial, FName ParameterName, FMovieSceneEntityID SoleContributor, FAnimatedMaterialParameterInfo* Output)
+	void InitializeSoleInput(UMovieSceneEntitySystemLinker* Linker, const FObjectComponent& BoundMaterial, const FMaterialParameterInfo& ParameterInfo, FMovieSceneEntityID SoleContributor, FAnimatedMaterialParameterInfo* Output)
 	{
 		FBuiltInComponentTypes* BuiltInComponents = FBuiltInComponentTypes::Get();
 		FMovieSceneTracksComponentTypes* TracksComponents = FMovieSceneTracksComponentTypes::Get();
@@ -191,12 +237,12 @@ struct FScalarMixin
 			float InitialValue = 0.0;
 			if (UMaterialInstanceDynamic* MID = Cast<UMaterialInstanceDynamic>(BoundMaterial.GetObject()))
 			{
-				MID->GetScalarParameterValue(ParameterName, InitialValue);
+				MID->GetScalarParameterValue(ParameterInfo, InitialValue);
 				Linker->EntityManager.AddComponent(SoleContributor, TracksComponents->FloatParameter.InitialValue, InitialValue);
 			}
 			else if (UMaterialParameterCollectionInstance* MPCI = Cast<UMaterialParameterCollectionInstance>(BoundMaterial.GetObject()))
 			{
-				MPCI->GetScalarParameterValue(ParameterName, InitialValue);
+				MPCI->GetScalarParameterValue(ParameterInfo.Name, InitialValue);
 				Linker->EntityManager.AddComponent(SoleContributor, TracksComponents->FloatParameter.InitialValue, InitialValue);
 			}
 		}
@@ -207,7 +253,7 @@ struct FScalarMixin
 
 struct FVectorMixin
 {
-	void CreateEntity(UMovieSceneEntitySystemLinker* Linker, const FObjectComponent& BoundMaterial, FName ParameterName, FComponentTypeID BlenderTypeTag, TArrayView<const FMovieSceneEntityID> Inputs, FAnimatedMaterialParameterInfo* Output)
+	void CreateEntity(UMovieSceneEntitySystemLinker* Linker, const FObjectComponent& BoundMaterial, const FMaterialParameterInfo& ParameterInfo, FComponentTypeID BlenderTypeTag, TArrayView<const FMovieSceneEntityID> Inputs, FAnimatedMaterialParameterInfo* Output)
 	{
 		FBuiltInComponentTypes* BuiltInComponents = FBuiltInComponentTypes::Get();
 		FMovieSceneTracksComponentTypes* TracksComponents = FMovieSceneTracksComponentTypes::Get();
@@ -230,11 +276,11 @@ struct FVectorMixin
 			FLinearColor ColorValue;
 			if (UMaterialInstanceDynamic* MID = Cast<UMaterialInstanceDynamic>(BoundMaterial.GetObject()))
 			{
-				bHasInitialValue = MID->GetVectorParameterValue(ParameterName, ColorValue);
+				bHasInitialValue = MID->GetVectorParameterValue(ParameterInfo, ColorValue);
 			}
 			else if (UMaterialParameterCollectionInstance* MPCI = Cast<UMaterialParameterCollectionInstance>(BoundMaterial.GetObject()))
 			{
-				bHasInitialValue = MPCI->GetVectorParameterValue(ParameterName, ColorValue);
+				bHasInitialValue = MPCI->GetVectorParameterValue(ParameterInfo.Name, ColorValue);
 			}
 
 			if (bHasInitialValue)
@@ -260,7 +306,7 @@ struct FVectorMixin
 		Linker->EntityManager.CopyComponents(Inputs[0], Output->OutputEntityID, Linker->EntityManager.GetComponents()->GetCopyAndMigrationMask());
 	}
 
-	void InitializeSoleInput(UMovieSceneEntitySystemLinker* Linker, const FObjectComponent& BoundMaterial, FName ParameterName, FMovieSceneEntityID SoleContributor, FAnimatedMaterialParameterInfo* Output)
+	void InitializeSoleInput(UMovieSceneEntitySystemLinker* Linker, const FObjectComponent& BoundMaterial, const FMaterialParameterInfo& ParameterInfo, FMovieSceneEntityID SoleContributor, FAnimatedMaterialParameterInfo* Output)
 	{
 		FBuiltInComponentTypes* BuiltInComponents = FBuiltInComponentTypes::Get();
 		FMovieSceneTracksComponentTypes* TracksComponents = FMovieSceneTracksComponentTypes::Get();
@@ -274,14 +320,14 @@ struct FVectorMixin
 			FLinearColor ColorValue = FLinearColor::White;
 			if (UMaterialInstanceDynamic* MID = Cast<UMaterialInstanceDynamic>(BoundMaterial.GetObject()))
 			{
-				if (MID->GetVectorParameterValue(ParameterName, ColorValue))
+				if (MID->GetVectorParameterValue(ParameterInfo, ColorValue))
 				{
 					Linker->EntityManager.AddComponent(SoleContributor, TracksComponents->ColorParameter.InitialValue, FIntermediateColor(ColorValue));
 				}
 			}
 			else if (UMaterialParameterCollectionInstance* MPCI = Cast<UMaterialParameterCollectionInstance>(BoundMaterial.GetObject()))
 			{
-				if (MPCI->GetVectorParameterValue(ParameterName, ColorValue))
+				if (MPCI->GetVectorParameterValue(ParameterInfo.Name, ColorValue))
 				{
 					Linker->EntityManager.AddComponent(SoleContributor, TracksComponents->ColorParameter.InitialValue, FIntermediateColor(ColorValue));
 				}
@@ -302,12 +348,12 @@ struct TOverlappingMaterialParameterHandler : Mixin
 		, System(InSystem)
 	{}
 
-	void InitializeOutput(const FObjectComponent& BoundMaterial, FName ParameterName, TArrayView<const FMovieSceneEntityID> Inputs, FAnimatedMaterialParameterInfo* Output, FEntityOutputAggregate Aggregate)
+	void InitializeOutput(const FObjectComponent& BoundMaterial, const FMaterialParameterInfo& ParameterInfo, TArrayView<const FMovieSceneEntityID> Inputs, FAnimatedMaterialParameterInfo* Output, FEntityOutputAggregate Aggregate)
 	{
-		UpdateOutput(BoundMaterial, ParameterName, Inputs, Output, Aggregate);
+		UpdateOutput(BoundMaterial, ParameterInfo, Inputs, Output, Aggregate);
 	}
 
-	void UpdateOutput(const FObjectComponent& BoundMaterial, FName ParameterName, TArrayView<const FMovieSceneEntityID> Inputs, FAnimatedMaterialParameterInfo* Output, FEntityOutputAggregate Aggregate)
+	void UpdateOutput(const FObjectComponent& BoundMaterial, const FMaterialParameterInfo& ParameterInfo, TArrayView<const FMovieSceneEntityID> Inputs, FAnimatedMaterialParameterInfo* Output, FEntityOutputAggregate Aggregate)
 	{
 		using namespace UE::MovieScene;
 
@@ -339,7 +385,7 @@ struct TOverlappingMaterialParameterHandler : Mixin
 
 				// Needs blending
 				FComponentTypeID BlenderTypeTag = System->DoubleBlenderSystem->GetBlenderTypeTag();
-				Mixin::CreateEntity(Linker, BoundMaterial, ParameterName, BlenderTypeTag, Inputs, Output);
+				Mixin::CreateEntity(Linker, BoundMaterial, ParameterInfo, BlenderTypeTag, Inputs, Output);
 			}
 
 			const FComponentTypeID BlenderTypeTag = System->DoubleBlenderSystem->GetBlenderTypeTag();
@@ -410,13 +456,13 @@ struct TOverlappingMaterialParameterHandler : Mixin
 		{
 			Linker->EntityManager.RemoveComponent(Inputs[0], BuiltInComponents->BlendChannelInput);
 
-			Mixin::InitializeSoleInput(Linker, BoundMaterial, ParameterName, Inputs[0], Output);
+			Mixin::InitializeSoleInput(Linker, BoundMaterial, ParameterInfo, Inputs[0], Output);
 		}
 
 		Output->NumContributors = NumContributors;
 	}
 
-	void DestroyOutput(const FObjectComponent& BoundMaterial, FName ParameterName, FAnimatedMaterialParameterInfo* Output, FEntityOutputAggregate Aggregate)
+	void DestroyOutput(const FObjectComponent& BoundMaterial, const FMaterialParameterInfo& ParameterInfo, FAnimatedMaterialParameterInfo* Output, FEntityOutputAggregate Aggregate)
 	{
 		if (Output->OutputEntityID)
 		{
@@ -527,17 +573,29 @@ void UMovieSceneMaterialParameterSystem::OnInstantiation()
 			this->VectorParameterTracker.VisitUnlinkedAllocation(Allocation);
 		};
 
-		auto HandleUpdatedAllocation = [this](const FEntityAllocation* Allocation, TRead<FObjectComponent> InObjects, TReadOneOf<FName, FName, FName> ScalarVectorOrColorParameterNames)
+		auto HandleUpdatedAllocation = [this](const FEntityAllocation* Allocation, TRead<FObjectComponent> InObjects, TReadOneOf<FMaterialParameterInfo, FMaterialParameterInfo, FMaterialParameterInfo, FName, FName, FName> ScalarVectorOrColorParameterInfosOrNames)
 		{
-			if (TComponentPtr<const FName> ScalarParameterNames = TComponentPtr<const FName>(ScalarVectorOrColorParameterNames.Get<0>()))
+			if (TComponentPtr<const FMaterialParameterInfo> ScalarParameterInfos = TComponentPtr<const FMaterialParameterInfo>(ScalarVectorOrColorParameterInfosOrNames.Get<0>()))
+			{
+				this->ScalarParameterTracker.VisitActiveAllocation(Allocation, InObjects, ScalarParameterInfos);
+			}
+			else if (TComponentPtr<const FMaterialParameterInfo> VectorParameterInfos = TComponentPtr<const FMaterialParameterInfo>(ScalarVectorOrColorParameterInfosOrNames.Get<1>()))
+			{
+				this->VectorParameterTracker.VisitActiveAllocation(Allocation, InObjects, VectorParameterInfos);
+			}
+			else if (TComponentPtr<const FMaterialParameterInfo> ColorParameterInfos = TComponentPtr<const FMaterialParameterInfo>(ScalarVectorOrColorParameterInfosOrNames.Get<2>()))
+			{
+				this->VectorParameterTracker.VisitActiveAllocation(Allocation, InObjects, ColorParameterInfos);
+			}
+			else if (TComponentPtr<const FName> ScalarParameterNames = TComponentPtr<const FName>(ScalarVectorOrColorParameterInfosOrNames.Get<3>()))
 			{
 				this->ScalarParameterTracker.VisitActiveAllocation(Allocation, InObjects, ScalarParameterNames);
 			}
-			else if (TComponentPtr<const FName> VectorParameterNames = TComponentPtr<const FName>(ScalarVectorOrColorParameterNames.Get<1>()))
+			else if (TComponentPtr<const FName> VectorParameterNames = TComponentPtr<const FName>(ScalarVectorOrColorParameterInfosOrNames.Get<4>()))
 			{
 				this->VectorParameterTracker.VisitActiveAllocation(Allocation, InObjects, VectorParameterNames);
 			}
-			else if (TComponentPtr<const FName> ColorParameterNames = TComponentPtr<const FName>(ScalarVectorOrColorParameterNames.Get<2>()))
+			else if (TComponentPtr<const FName> ColorParameterNames = TComponentPtr<const FName>(ScalarVectorOrColorParameterInfosOrNames.Get<5>()))
 			{
 				this->VectorParameterTracker.VisitActiveAllocation(Allocation, InObjects, ColorParameterNames);
 			}
@@ -546,7 +604,13 @@ void UMovieSceneMaterialParameterSystem::OnInstantiation()
 		// First step handle any new or updated bound materials
 		FEntityTaskBuilder()
 		.Read(TracksComponents->BoundMaterial)
-		.ReadOneOf(TracksComponents->ScalarParameterName, TracksComponents->VectorParameterName, TracksComponents->ColorParameterName)
+		// Handle both new parameter info and old parameter names for backwards compatibility with MovieSceneParameterSection
+		.ReadOneOf(	TracksComponents->ScalarMaterialParameterInfo, 
+					TracksComponents->VectorMaterialParameterInfo, 
+					TracksComponents->ColorMaterialParameterInfo,
+					TracksComponents->ScalarParameterName,
+					TracksComponents->VectorParameterName,
+					TracksComponents->ColorParameterName)
 		.FilterAny({ BuiltInComponents->Tags.NeedsLink, TracksComponents->Tags.BoundMaterialChanged })
 		.FilterNone({ BuiltInComponents->Tags.NeedsUnlink })
 		.Iterate_PerAllocation(&Linker->EntityManager, HandleUpdatedAllocation);
@@ -565,12 +629,16 @@ void UMovieSceneMaterialParameterSystem::OnInstantiation()
 		VectorParameterTracker.ProcessInvalidatedOutputs(Linker, VectorHandler);
 
 		// Gather inputs that contribute to the material parameter by excluding outputs (which will not have an instance handle)
-		TPreAnimatedStateTaskParams<FObjectComponent, FName> Params;
+		TPreAnimatedStateTaskParams<FObjectComponent, FMaterialParameterInfo> Params;
 		Params.AdditionalFilter.Reset();
 		Params.AdditionalFilter.None({ TracksComponents->MPC, BuiltInComponents->BlendChannelOutput });
 		Params.AdditionalFilter.Any({ BuiltInComponents->Tags.NeedsLink, TracksComponents->Tags.BoundMaterialChanged });
 		
+		// Handle both new parameter info and old parameter names for backwards compatibility with MovieSceneParameterSection
+		ScalarParameterStorage->BeginTrackingAndCachePreAnimatedValuesTask(Linker, Params, TracksComponents->BoundMaterial, TracksComponents->ScalarMaterialParameterInfo);
 		ScalarParameterStorage->BeginTrackingAndCachePreAnimatedValuesTask(Linker, Params, TracksComponents->BoundMaterial, TracksComponents->ScalarParameterName);
+		VectorParameterStorage->BeginTrackingAndCachePreAnimatedValuesTask(Linker, Params, TracksComponents->BoundMaterial, TracksComponents->VectorMaterialParameterInfo);
+		VectorParameterStorage->BeginTrackingAndCachePreAnimatedValuesTask(Linker, Params, TracksComponents->BoundMaterial, TracksComponents->ColorMaterialParameterInfo);
 		VectorParameterStorage->BeginTrackingAndCachePreAnimatedValuesTask(Linker, Params, TracksComponents->BoundMaterial, TracksComponents->VectorParameterName);
 		VectorParameterStorage->BeginTrackingAndCachePreAnimatedValuesTask(Linker, Params, TracksComponents->BoundMaterial, TracksComponents->ColorParameterName);
 	}
@@ -585,16 +653,18 @@ void UMovieSceneMaterialParameterSystem::OnSchedulePersistentTasks(UE::MovieScen
 
 	FEntityTaskBuilder()
 	.Read(TracksComponents->BoundMaterial)
-	.Read(TracksComponents->ScalarParameterName)
+	// Handle both new parameter info and old parameter names for backwards compatibility with MovieSceneParameterSection
+	.ReadOneOrMoreOf(TracksComponents->ScalarMaterialParameterInfo, TracksComponents->ScalarParameterName)
 	.Read(BuiltInComponents->DoubleResult[0])
 	.FilterNone({ BuiltInComponents->BlendChannelInput })
 	.SetDesiredThread(Linker->EntityManager.GetDispatchThread())
-	.Fork_PerEntity<FApplyScalarParameters>(&Linker->EntityManager, TaskScheduler);
+	.Fork_PerAllocation<FApplyScalarParameters>(&Linker->EntityManager, TaskScheduler);
 
 	// Vectors and colors use the same API
 	FEntityTaskBuilder()
 	.Read(TracksComponents->BoundMaterial)
-	.ReadOneOrMoreOf(TracksComponents->VectorParameterName, TracksComponents->ColorParameterName)
+	// Handle both new parameter info and old parameter names for backwards compatibility with MovieSceneParameterSection
+	.ReadOneOrMoreOf(TracksComponents->VectorMaterialParameterInfo, TracksComponents->ColorMaterialParameterInfo, TracksComponents->VectorParameterName, TracksComponents->ColorParameterName)
 	.ReadOneOrMoreOf(BuiltInComponents->DoubleResult[0], BuiltInComponents->DoubleResult[1], BuiltInComponents->DoubleResult[2], BuiltInComponents->DoubleResult[3])
 	.FilterNone({ BuiltInComponents->BlendChannelInput })
 	.SetDesiredThread(Linker->EntityManager.GetDispatchThread())
@@ -608,23 +678,27 @@ void UMovieSceneMaterialParameterSystem::OnEvaluation(FSystemTaskPrerequisites&
 	FBuiltInComponentTypes* BuiltInComponents = FBuiltInComponentTypes::Get();
 	FMovieSceneTracksComponentTypes* TracksComponents = FMovieSceneTracksComponentTypes::Get();
 
-	if (Linker->EntityManager.ContainsComponent(TracksComponents->ScalarParameterName))
+	if (Linker->EntityManager.ContainsComponent(TracksComponents->ScalarMaterialParameterInfo) || Linker->EntityManager.ContainsComponent(TracksComponents->ScalarParameterName))
 	{
 		FEntityTaskBuilder()
-		.Read(TracksComponents->BoundMaterial)
-		.Read(TracksComponents->ScalarParameterName)
+		.Read(TracksComponents->BoundMaterial)	
+		// Handle both new parameter info and old parameter names for backwards compatibility with MovieSceneParameterSection
+		.ReadOneOrMoreOf(TracksComponents->ScalarMaterialParameterInfo, TracksComponents->ScalarParameterName)
 		.Read(BuiltInComponents->DoubleResult[0])
 		.FilterNone({ BuiltInComponents->BlendChannelInput })
 		.SetDesiredThread(Linker->EntityManager.GetDispatchThread())
-		.Dispatch_PerEntity<FApplyScalarParameters>(&Linker->EntityManager, InPrerequisites, &Subsequents);
+		.Dispatch_PerAllocation<FApplyScalarParameters>(&Linker->EntityManager, InPrerequisites, &Subsequents);
 	}
 
 	// Vectors and colors use the same API
-	if (Linker->EntityManager.ContainsComponent(TracksComponents->VectorParameterName) || Linker->EntityManager.ContainsComponent(TracksComponents->ColorParameterName))
+	if (Linker->EntityManager.ContainsComponent(TracksComponents->VectorMaterialParameterInfo) 
+		|| Linker->EntityManager.ContainsComponent(TracksComponents->ColorMaterialParameterInfo)
+		|| Linker->EntityManager.ContainsComponent(TracksComponents->VectorParameterName)
+		|| Linker->EntityManager.ContainsComponent(TracksComponents->ColorParameterName))
 	{
 		FEntityTaskBuilder()
 		.Read(TracksComponents->BoundMaterial)
-		.ReadOneOrMoreOf(TracksComponents->VectorParameterName, TracksComponents->ColorParameterName)
+		.ReadOneOrMoreOf(TracksComponents->VectorMaterialParameterInfo, TracksComponents->ColorMaterialParameterInfo, TracksComponents->VectorParameterName, TracksComponents->ColorParameterName)
 		.ReadOneOrMoreOf(BuiltInComponents->DoubleResult[0], BuiltInComponents->DoubleResult[1], BuiltInComponents->DoubleResult[2], BuiltInComponents->DoubleResult[3])
 		.FilterNone({ BuiltInComponents->BlendChannelInput })
 		.SetDesiredThread(Linker->EntityManager.GetDispatchThread())
diff --git a/Engine/Source/Runtime/MovieSceneTracks/Private/Systems/MovieScenePreAnimatedMaterialParameters.cpp b/Engine/Source/Runtime/MovieSceneTracks/Private/Systems/MovieScenePreAnimatedMaterialParameters.cpp
--- a/Engine/Source/Runtime/MovieSceneTracks/Private/Systems/MovieScenePreAnimatedMaterialParameters.cpp
+++ b/Engine/Source/Runtime/MovieSceneTracks/Private/Systems/MovieScenePreAnimatedMaterialParameters.cpp
@@ -14,17 +14,22 @@ TAutoRegisterPreAnimatedStorageID<FPreAnimatedVectorMaterialParameterStorage> FP
 
 FMaterialParameterKey::FMaterialParameterKey(const FObjectComponent& InBoundMaterial, const FName& InParameterName)
 	: BoundMaterial(InBoundMaterial.GetObject())
-	, ParameterName(InParameterName)
+	, ParameterInfo(InParameterName)
+{}
+
+FMaterialParameterKey::FMaterialParameterKey(const FObjectComponent& InBoundMaterial, const FMaterialParameterInfo& InParameterInfo)
+	: BoundMaterial(InBoundMaterial.GetObject())
+	, ParameterInfo(InParameterInfo)
 {}
 
 uint32 GetTypeHash(const FMaterialParameterKey& InKey)
 {
-	return GetTypeHash(InKey.BoundMaterial) ^ GetTypeHash(InKey.ParameterName);
+	return GetTypeHash(InKey.BoundMaterial) ^ GetTypeHash(InKey.ParameterInfo);
 }
 
 bool operator==(const FMaterialParameterKey& A, const FMaterialParameterKey& B)
 {
-	return A.BoundMaterial == B.BoundMaterial && A.ParameterName == B.ParameterName;
+	return A.BoundMaterial == B.BoundMaterial && A.ParameterInfo == B.ParameterInfo;
 }
 
 void FMaterialParameterCollectionScalarTraits::ReplaceObject(FMaterialParameterKey& InOutKey, const FObjectKey& NewObject)
@@ -33,17 +38,22 @@ void FMaterialParameterCollectionScalarTraits::ReplaceObject(FMaterialParameterK
 }
 
 float FMaterialParameterCollectionScalarTraits::CachePreAnimatedValue(const FObjectComponent& InBoundMaterial, const FName& ParameterName)
+{
+	return CachePreAnimatedValue(InBoundMaterial, FMaterialParameterInfo(ParameterName));
+}
+
+float FMaterialParameterCollectionScalarTraits::CachePreAnimatedValue(const FObjectComponent& InBoundMaterial, const FMaterialParameterInfo& ParameterInfo)
 {
 	UObject* BoundMaterial = InBoundMaterial.GetObject();
 
 	float ParameterValue = 0.f;
 	if (UMaterialParameterCollectionInstance* MPCI = Cast<UMaterialParameterCollectionInstance>(BoundMaterial))
 	{
-		MPCI->GetScalarParameterValue(ParameterName, ParameterValue);
+		MPCI->GetScalarParameterValue(ParameterInfo.Name, ParameterValue);
 	}
 	else if (UMaterialInterface* MaterialInterface = Cast<UMaterialInterface>(BoundMaterial))
 	{
-		MaterialInterface->GetScalarParameterValue(FMaterialParameterInfo(ParameterName), ParameterValue);
+		MaterialInterface->GetScalarParameterValue(ParameterInfo, ParameterValue);
 	}
 	return ParameterValue;
 }
@@ -54,11 +64,11 @@ void FMaterialParameterCollectionScalarTraits::RestorePreAnimatedValue(const FMa
 
 	if (UMaterialParameterCollectionInstance* MPCI = Cast<UMaterialParameterCollectionInstance>(BoundObject))
 	{
-		MPCI->SetScalarParameterValue(InKey.ParameterName, OldValue);
+		MPCI->SetScalarParameterValue(InKey.ParameterInfo.Name, OldValue);
 	}
 	else if (UMaterialInstanceDynamic* MID = Cast<UMaterialInstanceDynamic>(BoundObject))
 	{
-		MID->SetScalarParameterValue(InKey.ParameterName, OldValue);
+		MID->SetScalarParameterValueByInfo(InKey.ParameterInfo, OldValue);
 	}
 }
 
@@ -68,17 +78,22 @@ void FMaterialParameterCollectionVectorTraits::ReplaceObject(FMaterialParameterK
 }
 
 FLinearColor FMaterialParameterCollectionVectorTraits::CachePreAnimatedValue(const FObjectComponent& InBoundMaterial, const FName& ParameterName)
+{
+	return CachePreAnimatedValue(InBoundMaterial, FMaterialParameterInfo(ParameterName));
+}
+
+FLinearColor FMaterialParameterCollectionVectorTraits::CachePreAnimatedValue(const FObjectComponent& InBoundMaterial, const FMaterialParameterInfo& ParameterInfo)
 {
 	UObject* BoundMaterial = InBoundMaterial.GetObject();
 
 	FLinearColor ParameterValue = FLinearColor::White;
 	if (UMaterialParameterCollectionInstance* MPCI = Cast<UMaterialParameterCollectionInstance>(BoundMaterial))
 	{
-		MPCI->GetVectorParameterValue(ParameterName, ParameterValue);
+		MPCI->GetVectorParameterValue(ParameterInfo.Name, ParameterValue);
 	}
 	else if (UMaterialInterface* MaterialInterface = Cast<UMaterialInterface>(BoundMaterial))
 	{
-		MaterialInterface->GetVectorParameterValue(FMaterialParameterInfo(ParameterName), ParameterValue);
+		MaterialInterface->GetVectorParameterValue(ParameterInfo, ParameterValue);
 	}
 	return ParameterValue;
 }
@@ -89,11 +104,11 @@ void FMaterialParameterCollectionVectorTraits::RestorePreAnimatedValue(const FMa
 
 	if (UMaterialParameterCollectionInstance* MPCI = Cast<UMaterialParameterCollectionInstance>(BoundObject))
 	{
-		MPCI->SetVectorParameterValue(InKey.ParameterName, OldValue);
+		MPCI->SetVectorParameterValue(InKey.ParameterInfo.Name, OldValue);
 	}
 	else if (UMaterialInstanceDynamic* MID = Cast<UMaterialInstanceDynamic>(BoundObject))
 	{
-		MID->SetVectorParameterValue(InKey.ParameterName, OldValue);
+		MID->SetVectorParameterValueByInfo(InKey.ParameterInfo, OldValue);
 	}
 }
 
diff --git a/Engine/Source/Runtime/MovieSceneTracks/Private/Tracks/MovieSceneMaterialParameterCollectionTrack.cpp b/Engine/Source/Runtime/MovieSceneTracks/Private/Tracks/MovieSceneMaterialParameterCollectionTrack.cpp
--- a/Engine/Source/Runtime/MovieSceneTracks/Private/Tracks/MovieSceneMaterialParameterCollectionTrack.cpp
+++ b/Engine/Source/Runtime/MovieSceneTracks/Private/Tracks/MovieSceneMaterialParameterCollectionTrack.cpp
@@ -6,6 +6,7 @@
 #include "Evaluation/MovieSceneEvaluationField.h"
 #include "EntitySystem/BuiltInComponentTypes.h"
 #include "MovieSceneTracksComponentTypes.h"
+#include "Sections/MovieSceneComponentMaterialParameterSection.h"
 
 #include UE_INLINE_GENERATED_CPP_BY_NAME(MovieSceneMaterialParameterCollectionTrack)
 
@@ -54,22 +55,32 @@ bool UMovieSceneMaterialParameterCollectionTrack::PopulateEvaluationFieldImpl(co
 {
 	const FMovieSceneTrackEvaluationField& LocalEvaluationField = GetEvaluationField();
 
-	// Define entities for every entry in our evaluation field
+	// Define entities for the old style parameter sections. ComponentMaterialParameterSections define their own.
 	for (const FMovieSceneTrackEvaluationFieldEntry& Entry : LocalEvaluationField.Entries)
 	{
 		UMovieSceneParameterSection* ParameterSection = Cast<UMovieSceneParameterSection>(Entry.Section);
-		if (!ParameterSection || IsRowEvalDisabled(ParameterSection->GetRowIndex()))
+		UMovieSceneComponentMaterialParameterSection* ComponentMaterialParameterSection = Cast<UMovieSceneComponentMaterialParameterSection>(Entry.Section);
+		if (ParameterSection || ComponentMaterialParameterSection)
 		{
-			continue;
-		}
-
-		TRange<FFrameNumber> SectionEffectiveRange = TRange<FFrameNumber>::Intersection(EffectiveRange, Entry.Range);
-		if (!SectionEffectiveRange.IsEmpty())
-		{
-			FMovieSceneEvaluationFieldEntityMetaData SectionMetaData = InMetaData;
-			SectionMetaData.Flags = Entry.Flags;
-
-			ParameterSection->ExternalPopulateEvaluationField(SectionEffectiveRange, SectionMetaData, OutFieldBuilder);
+			if (IsRowEvalDisabled(Entry.Section->GetRowIndex()))
+			{
+				continue;
+			}
+
+			TRange<FFrameNumber> SectionEffectiveRange = TRange<FFrameNumber>::Intersection(EffectiveRange, Entry.Range);
+			if (!SectionEffectiveRange.IsEmpty())
+			{
+				FMovieSceneEvaluationFieldEntityMetaData SectionMetaData = InMetaData;
+				SectionMetaData.Flags = Entry.Flags;
+				if (ParameterSection)
+				{
+					ParameterSection->ExternalPopulateEvaluationField(SectionEffectiveRange, SectionMetaData, OutFieldBuilder);
+				}
+				else if (ComponentMaterialParameterSection)
+				{
+					ComponentMaterialParameterSection->ExternalPopulateEvaluationField(SectionEffectiveRange, SectionMetaData, OutFieldBuilder);
+				}
+			}
 		}
 	}
 
@@ -78,7 +89,7 @@ bool UMovieSceneMaterialParameterCollectionTrack::PopulateEvaluationFieldImpl(co
 
 bool UMovieSceneMaterialParameterCollectionTrack::SupportsType(TSubclassOf<UMovieSceneSection> SectionClass) const
 {
-	return SectionClass == UMovieSceneParameterSection::StaticClass();
+	return SectionClass == UMovieSceneComponentMaterialParameterSection::StaticClass() || SectionClass == UMovieSceneParameterSection::StaticClass();
 }
 
 #if WITH_EDITORONLY_DATA
diff --git a/Engine/Source/Runtime/MovieSceneTracks/Private/Tracks/MovieSceneMaterialTrack.cpp b/Engine/Source/Runtime/MovieSceneTracks/Private/Tracks/MovieSceneMaterialTrack.cpp
--- a/Engine/Source/Runtime/MovieSceneTracks/Private/Tracks/MovieSceneMaterialTrack.cpp
+++ b/Engine/Source/Runtime/MovieSceneTracks/Private/Tracks/MovieSceneMaterialTrack.cpp
@@ -7,6 +7,7 @@
 #include "EntitySystem/MovieSceneEntityBuilder.h"
 #include "EntitySystem/BuiltInComponentTypes.h"
 #include "MovieSceneTracksComponentTypes.h"
+#include "Sections/MovieSceneComponentMaterialParameterSection.h"
 
 #include UE_INLINE_GENERATED_CPP_BY_NAME(MovieSceneMaterialTrack)
 
@@ -27,13 +28,13 @@ UMovieSceneMaterialTrack::UMovieSceneMaterialTrack(const FObjectInitializer& Obj
 
 bool UMovieSceneMaterialTrack::SupportsType(TSubclassOf<UMovieSceneSection> SectionClass) const
 {
-	return SectionClass == UMovieSceneParameterSection::StaticClass();
+	return SectionClass == UMovieSceneComponentMaterialParameterSection::StaticClass() || SectionClass == UMovieSceneParameterSection::StaticClass();
 }
 
 
 UMovieSceneSection* UMovieSceneMaterialTrack::CreateNewSection()
 {
-	UMovieSceneSection* NewSection = NewObject<UMovieSceneParameterSection>(this, NAME_None, RF_Transactional);
+	UMovieSceneSection* NewSection = NewObject<UMovieSceneComponentMaterialParameterSection>(this, NAME_None, RF_Transactional);
 	NewSection->SetBlendType(EMovieSceneBlendType::Absolute);
 	return NewSection;
 }
@@ -120,20 +121,29 @@ const TArray<UMovieSceneSection*>& UMovieSceneMaterialTrack::GetAllSections() co
 
 void UMovieSceneMaterialTrack::AddScalarParameterKey(FName ParameterName, FFrameNumber Time, float Value)
 {
-	AddScalarParameterKey(ParameterName, Time, INDEX_NONE, Value);
+	AddScalarParameterKey(FMaterialParameterInfo(ParameterName), Time, INDEX_NONE, Value, FString(), FString());
 }
 
-
 void UMovieSceneMaterialTrack::AddScalarParameterKey(FName ParameterName, FFrameNumber Time, int32 RowIndex, float Value)
 {
-	UMovieSceneParameterSection* NearestSection = Cast<UMovieSceneParameterSection>(SectionToKey);
+	AddScalarParameterKey(FMaterialParameterInfo(ParameterName), Time, RowIndex, Value, FString(), FString());
+}
+
+void UMovieSceneMaterialTrack::AddScalarParameterKey(const FMaterialParameterInfo& ParameterInfo, FFrameNumber Time, float Value, const FString& InLayerName, const FString& InAssetName)
+{
+	AddScalarParameterKey(ParameterInfo, Time, INDEX_NONE, Value, InLayerName, InAssetName);
+}
+
+void UMovieSceneMaterialTrack::AddScalarParameterKey(const FMaterialParameterInfo& ParameterInfo, FFrameNumber Time, int32 RowIndex, float Value, const FString& InLayerName, const FString& InAssetName)
+{
+	UMovieSceneSection* NearestSection = SectionToKey;
 	if (NearestSection == nullptr || (RowIndex != INDEX_NONE && NearestSection->GetRowIndex() != RowIndex))
 	{
-		NearestSection = Cast<UMovieSceneParameterSection>(MovieSceneHelpers::FindNearestSectionAtTime(Sections, Time, RowIndex));
+		NearestSection = MovieSceneHelpers::FindNearestSectionAtTime(Sections, Time, RowIndex);
 	}
 	if (NearestSection == nullptr)
 	{
-		NearestSection = Cast<UMovieSceneParameterSection>(CreateNewSection());
+		NearestSection = Cast<UMovieSceneComponentMaterialParameterSection>(CreateNewSection());
 
 		UMovieScene* MovieScene = GetTypedOuter<UMovieScene>();
 		check(MovieScene);
@@ -141,40 +151,65 @@ void UMovieSceneMaterialTrack::AddScalarParameterKey(FName ParameterName, FFrame
 		NearestSection->SetRange(MovieScene->GetPlaybackRange());
 		Sections.Add(NearestSection);
 	}
-	if (NearestSection->TryModify())
+	if (NearestSection != nullptr && NearestSection->TryModify())
 	{
-		NearestSection->AddScalarParameterKey(ParameterName, Time, Value);
+		// If we have an old parameter section already, fall back to the old style section
+		if (UMovieSceneParameterSection* NearestParameterSection = Cast<UMovieSceneParameterSection>(NearestSection))
+		{
+			NearestParameterSection->AddScalarParameterKey(ParameterInfo.Name, Time, Value);
+		}
+		else if (UMovieSceneComponentMaterialParameterSection* NearestComponentMaterialParameterSection = Cast<UMovieSceneComponentMaterialParameterSection>(NearestSection))
+		{
+			NearestComponentMaterialParameterSection->AddScalarParameterKey(ParameterInfo, Time, Value, InLayerName, InAssetName);
+		}
 	}
 }
 
 
 void UMovieSceneMaterialTrack::AddColorParameterKey(FName ParameterName, FFrameNumber Time, FLinearColor Value)
 {
-	AddColorParameterKey(ParameterName, Time, INDEX_NONE, Value);
+	AddColorParameterKey(FMaterialParameterInfo(ParameterName), Time, INDEX_NONE, Value, FString(), FString());
 }
 
 
 void UMovieSceneMaterialTrack::AddColorParameterKey(FName ParameterName, FFrameNumber Time, int32 RowIndex, FLinearColor Value)
 {
-	UMovieSceneParameterSection* NearestSection = Cast<UMovieSceneParameterSection>(SectionToKey);
+	AddColorParameterKey(FMaterialParameterInfo(ParameterName), Time, RowIndex, Value, FString(), FString());
+}
+
+void UMovieSceneMaterialTrack::AddColorParameterKey(const FMaterialParameterInfo& ParameterInfo, FFrameNumber Time, FLinearColor Value, const FString& InLayerName, const FString& InAssetName)
+{
+	AddColorParameterKey(ParameterInfo, Time, INDEX_NONE, Value, InLayerName, InAssetName);
+}
+
+void UMovieSceneMaterialTrack::AddColorParameterKey(const FMaterialParameterInfo& ParameterInfo, FFrameNumber Time, int32 RowIndex, FLinearColor Value, const FString& InLayerName, const FString& InAssetName)
+{
+	UMovieSceneSection* NearestSection = SectionToKey;
 	if (NearestSection == nullptr || (RowIndex != INDEX_NONE && NearestSection->GetRowIndex() != RowIndex))
 	{
-		NearestSection = Cast<UMovieSceneParameterSection>(MovieSceneHelpers::FindNearestSectionAtTime(Sections, Time));
+		NearestSection = MovieSceneHelpers::FindNearestSectionAtTime(Sections, Time, RowIndex);
 	}
 	if (NearestSection == nullptr)
 	{
-		NearestSection = Cast<UMovieSceneParameterSection>(CreateNewSection());
+		NearestSection = Cast<UMovieSceneComponentMaterialParameterSection>(CreateNewSection());
 
 		UMovieScene* MovieScene = GetTypedOuter<UMovieScene>();
 		check(MovieScene);
 
 		NearestSection->SetRange(MovieScene->GetPlaybackRange());
-
 		Sections.Add(NearestSection);
 	}
-	if (NearestSection->TryModify())
+	if (NearestSection != nullptr && NearestSection->TryModify())
 	{
-		NearestSection->AddColorParameterKey(ParameterName, Time, Value);
+		// If we have an old parameter section already, fall back to the old style section
+		if (UMovieSceneParameterSection* NearestParameterSection = Cast<UMovieSceneParameterSection>(NearestSection))
+		{
+			NearestParameterSection->AddColorParameterKey(ParameterInfo.Name, Time, Value);
+		}
+		else if (UMovieSceneComponentMaterialParameterSection* NearestComponentMaterialParameterSection = Cast<UMovieSceneComponentMaterialParameterSection>(NearestSection))
+		{
+			NearestComponentMaterialParameterSection->AddColorParameterKey(ParameterInfo, Time, Value, InLayerName, InAssetName);
+		}
 	}
 }
 
@@ -212,22 +247,32 @@ bool UMovieSceneComponentMaterialTrack::PopulateEvaluationFieldImpl(const TRange
 {
 	const FMovieSceneTrackEvaluationField& LocalEvaluationField = GetEvaluationField();
 
-	// Define entities for every entry in our evaluation field
+	// Define entities for the old style parameter sections. ComponentMaterialParameterSections define their own.
 	for (const FMovieSceneTrackEvaluationFieldEntry& Entry : LocalEvaluationField.Entries)
 	{
 		UMovieSceneParameterSection* ParameterSection = Cast<UMovieSceneParameterSection>(Entry.Section);
-		if (!ParameterSection || IsRowEvalDisabled(ParameterSection->GetRowIndex()))
-		{
-			continue;
-		}
-
-		TRange<FFrameNumber> SectionEffectiveRange = TRange<FFrameNumber>::Intersection(EffectiveRange, Entry.Range);
-		if (!SectionEffectiveRange.IsEmpty())
+		UMovieSceneComponentMaterialParameterSection* ComponentMaterialParameterSection = Cast<UMovieSceneComponentMaterialParameterSection>(Entry.Section);
+		if (ParameterSection || ComponentMaterialParameterSection)
 		{
-			FMovieSceneEvaluationFieldEntityMetaData SectionMetaData = InMetaData;
-			SectionMetaData.Flags = Entry.Flags;
-
-			ParameterSection->ExternalPopulateEvaluationField(SectionEffectiveRange, SectionMetaData, OutFieldBuilder);
+			if (IsRowEvalDisabled(Entry.Section->GetRowIndex()))
+			{
+				continue;
+			}
+
+			TRange<FFrameNumber> SectionEffectiveRange = TRange<FFrameNumber>::Intersection(EffectiveRange, Entry.Range);
+			if (!SectionEffectiveRange.IsEmpty())
+			{
+				FMovieSceneEvaluationFieldEntityMetaData SectionMetaData = InMetaData;
+				SectionMetaData.Flags = Entry.Flags;
+				if (ParameterSection)
+				{
+					ParameterSection->ExternalPopulateEvaluationField(SectionEffectiveRange, SectionMetaData, OutFieldBuilder);
+				}
+				else if (ComponentMaterialParameterSection)
+				{
+					ComponentMaterialParameterSection->ExternalPopulateEvaluationField(SectionEffectiveRange, SectionMetaData, OutFieldBuilder);
+				}
+			}
 		}
 	}
 
diff --git a/Engine/Source/Runtime/MovieSceneTracks/Public/MovieSceneTracksComponentTypes.h b/Engine/Source/Runtime/MovieSceneTracks/Public/MovieSceneTracksComponentTypes.h
--- a/Engine/Source/Runtime/MovieSceneTracks/Public/MovieSceneTracksComponentTypes.h
+++ b/Engine/Source/Runtime/MovieSceneTracks/Public/MovieSceneTracksComponentTypes.h
@@ -14,6 +14,7 @@
 #include "MovieSceneTracksPropertyTypes.h"
 #include "Styling/SlateColor.h"
 #include "ConstraintChannel.h"
+#include "MaterialTypes.h"
 #include "MovieSceneTracksComponentTypes.generated.h"
 
 class UMaterialParameterCollection;
@@ -565,6 +566,10 @@ struct FMovieSceneTracksComponentTypes
 	TComponentTypeID<FName> ColorParameterName;
 	TComponentTypeID<FName> TransformParameterName;
 
+	TComponentTypeID<FMaterialParameterInfo> ScalarMaterialParameterInfo;
+	TComponentTypeID<FMaterialParameterInfo> ColorMaterialParameterInfo;
+	TComponentTypeID<FMaterialParameterInfo> VectorMaterialParameterInfo;
+
 	TComponentTypeID<FObjectComponent> BoundMaterial;
 	TComponentTypeID<TWeakObjectPtr<UMaterialParameterCollection>> MPC;
 
diff --git a/Engine/Source/Runtime/MovieSceneTracks/Public/Sections/MovieSceneComponentMaterialParameterSection.h b/Engine/Source/Runtime/MovieSceneTracks/Public/Sections/MovieSceneComponentMaterialParameterSection.h
new file mode 100644
--- /dev/null
+++ b/Engine/Source/Runtime/MovieSceneTracks/Public/Sections/MovieSceneComponentMaterialParameterSection.h
@@ -0,0 +1,173 @@
+// Copyright Epic Games, Inc. All Rights Reserved.
+
+#pragma once
+
+#include "CoreTypes.h"
+#include "UObject/ObjectMacros.h"
+#include "Curves/KeyHandle.h"
+#include "Curves/RichCurve.h"
+#include "MovieSceneSection.h"
+#include "Channels/MovieSceneFloatChannel.h"
+#include "EntitySystem/IMovieSceneEntityProvider.h"
+#include "MaterialTypes.h"
+#include "MovieSceneComponentMaterialParameterSection.generated.h"
+
+/**
+ * Structure representing an animated scalar parameter and it's associated animation curve.
+ */
+USTRUCT()
+struct FScalarMaterialParameterInfoAndCurve
+{
+	GENERATED_USTRUCT_BODY()
+
+		FScalarMaterialParameterInfoAndCurve()
+	{}
+
+	/** Creates a new FScalarMaterialParameterInfoAndCurve for a specific scalar parameter. */
+	MOVIESCENETRACKS_API FScalarMaterialParameterInfoAndCurve(const FMaterialParameterInfo& InParameterInfo);
+
+	/** The material parameter info of the scalar parameter which is being animated. */
+	UPROPERTY()
+	FMaterialParameterInfo ParameterInfo;
+
+	/** The curve which contains the animation data for the scalar parameter. */
+	UPROPERTY()
+	FMovieSceneFloatChannel ParameterCurve;
+
+#if WITH_EDITORONLY_DATA
+	/** Editor-only name of parameter layer, if applicable */
+	UPROPERTY()
+	FString ParameterLayerName;
+
+	/** Editor-only name of parameter asset, if applicable */
+	UPROPERTY()
+	FString ParameterAssetName;
+#endif
+};
+
+/**
+* Structure representing an animated color parameter and it's associated animation curve.
+*/
+USTRUCT()
+struct FColorMaterialParameterInfoAndCurves
+{
+	GENERATED_USTRUCT_BODY()
+
+	FColorMaterialParameterInfoAndCurves()
+	{}
+
+	/** Creates a new FVectorParameterNameAndCurve for a specific color parameter. */
+	MOVIESCENETRACKS_API FColorMaterialParameterInfoAndCurves(const FMaterialParameterInfo& InParameterInfo);
+
+	/** The name of the color parameter which is being animated. */
+	UPROPERTY()
+	FMaterialParameterInfo ParameterInfo;
+
+	/** The curve which contains the animation data for the red component of the color parameter. */
+	UPROPERTY()
+	FMovieSceneFloatChannel RedCurve;
+
+	/** The curve which contains the animation data for the green component of the color parameter. */
+	UPROPERTY()
+	FMovieSceneFloatChannel GreenCurve;
+
+	/** The curve which contains the animation data for the blue component of the color parameter. */
+	UPROPERTY()
+	FMovieSceneFloatChannel BlueCurve;
+
+	/** The curve which contains the animation data for the alpha component of the color parameter. */
+	UPROPERTY()
+	FMovieSceneFloatChannel AlphaCurve;
+
+#if WITH_EDITORONLY_DATA
+	/** Name of parameter layer, if applicable */
+	UPROPERTY()
+	FString ParameterLayerName;
+
+	/** Name of parameter asset, if applicable */
+	UPROPERTY()
+	FString ParameterAssetName;
+#endif
+};
+
+/**
+ * A movie scene section containing data for material parameters.
+ */
+UCLASS(MinimalAPI)
+class UMovieSceneComponentMaterialParameterSection
+	: public UMovieSceneSection
+	, public IMovieSceneEntityProvider
+{
+	GENERATED_UCLASS_BODY()
+
+public:
+	/** Adds a a key for a specific scalar parameter at the specified time with the specified value. */
+	UFUNCTION(BlueprintCallable, Category = "Sequencer|Section")
+	MOVIESCENETRACKS_API void AddScalarParameterKey(const FMaterialParameterInfo& InParameterInfo, FFrameNumber InTime, float InValue, const FString& InLayerName, const FString& InAssetName);
+
+	/** Adds a a key for a specific color parameter at the specified time with the specified value. */
+	UFUNCTION(BlueprintCallable, Category = "Sequencer|Section")
+	MOVIESCENETRACKS_API void AddColorParameterKey(const FMaterialParameterInfo& InParameterInfo, FFrameNumber InTime, FLinearColor InValue, const FString& InLayerName, const FString& InAssetName);
+
+	/** 
+	 * Removes a scalar parameter from this section. 
+	 * 
+	 * @param InParameterInfo The material parameter info of the scalar parameter to remove.
+	 * @returns True if a parameter with that name was found and removed, otherwise false.
+	 */
+	UFUNCTION(BlueprintCallable, Category = "Sequencer|Section")
+	MOVIESCENETRACKS_API bool RemoveScalarParameter(const FMaterialParameterInfo& InParameterInfo);
+	
+	/**
+	 * Removes a color parameter from this section.
+	 *
+	 * @param InParameterInfo The material parameter info of the color parameter to remove.
+	 * @returns True if a parameter with that name was found and removed, otherwise false.
+	 */
+	UFUNCTION(BlueprintCallable, Category = "Sequencer|Section")
+	MOVIESCENETRACKS_API bool RemoveColorParameter(const FMaterialParameterInfo& InParameterInfo);
+
+
+#if WITH_EDITOR
+	/**
+	* Removes a scalar parameter from this section.
+	*
+	* @param InParameterPath The full path name of the scalar parameter to remove.
+	* @returns True if a parameter with that name was found and removed, otherwise false.
+	*/
+	MOVIESCENETRACKS_API bool RemoveScalarParameter(FName InParameterPath);
+
+	/**
+	 * Removes a color parameter from this section.
+	 *
+	 * @param InParameterPath The full path name of the color parameter to remove.
+	 * @returns True if a parameter with that name was found and removed, otherwise false.
+	 */
+	MOVIESCENETRACKS_API bool RemoveColorParameter(FName InParameterPath);
+#endif
+
+	MOVIESCENETRACKS_API virtual void ImportEntityImpl(UMovieSceneEntitySystemLinker* EntityLinker, const FEntityImportParams& Params, FImportedEntity* OutImportedEntity) override;
+	MOVIESCENETRACKS_API virtual bool PopulateEvaluationFieldImpl(const TRange<FFrameNumber>& EffectiveRange, const FMovieSceneEvaluationFieldEntityMetaData& InMetaData, FMovieSceneEntityComponentFieldBuilder* OutFieldBuilder) override;
+
+	MOVIESCENETRACKS_API void ExternalPopulateEvaluationField(const TRange<FFrameNumber>& EffectiveRange, const FMovieSceneEvaluationFieldEntityMetaData& InMetaData, FMovieSceneEntityComponentFieldBuilder* OutFieldBuilder);
+
+
+public:
+	
+	/** The scalar parameter infos and their associated curves. */
+	UPROPERTY()
+	TArray<FScalarMaterialParameterInfoAndCurve> ScalarParameterInfosAndCurves;
+
+	/** The color parameter infos and their associated curves. */
+	UPROPERTY()
+	TArray<FColorMaterialParameterInfoAndCurves> ColorParameterInfosAndCurves;
+
+protected:
+
+	//~ UMovieSceneSection interface
+	virtual EMovieSceneChannelProxyType CacheChannelProxy() override;
+	MOVIESCENETRACKS_API virtual void PostEditImport() override;
+#if WITH_EDITOR
+	MOVIESCENETRACKS_API virtual void PostEditUndo() override;
+#endif
+};
diff --git a/Engine/Source/Runtime/MovieSceneTracks/Public/Systems/MovieSceneMaterialParameterSystem.h b/Engine/Source/Runtime/MovieSceneTracks/Public/Systems/MovieSceneMaterialParameterSystem.h
--- a/Engine/Source/Runtime/MovieSceneTracks/Public/Systems/MovieSceneMaterialParameterSystem.h
+++ b/Engine/Source/Runtime/MovieSceneTracks/Public/Systems/MovieSceneMaterialParameterSystem.h
@@ -5,6 +5,7 @@
 #include "EntitySystem/MovieSceneEntitySystem.h"
 #include "EntitySystem/MovieSceneOverlappingEntityTracker.h"
 #include "Evaluation/PreAnimatedState/MovieScenePreAnimatedStateStorage.h"
+#include "MaterialTypes.h"
 
 #include "MovieSceneMaterialParameterSystem.generated.h"
 
@@ -39,12 +40,12 @@ MOVIESCENETRACKS_API void CollectGarbageForOutput(FAnimatedMaterialParameterInfo
  * System responsible for tracking and animating material parameter entities.
  * Operates on the following component types from FMovieSceneTracksComponentTypes:
  *
- * Instantiation: Tracks any BoundMaterial with a ScalarParameterName, ColorParameterName or VectorParameterName.
+ * Instantiation: Tracks any BoundMaterial with a MaterialParameterInfo, as well as deprecated ScalarParameterName, ColorParameterName or VectorParameterName.
  *                Manages adding BlendChannelInputs and Outputs where multiple entities animate the same parameter
  *                on the same bound material.
  *                BoundMaterials may be a UMaterialInstanceDynamic, or a UMaterialParameterCollectionInstance.
  *
- * Evaluation:    Visits any BoundMaterial with the supported parameter names and either a BlendChannelOutput component
+ * Evaluation:    Visits any BoundMaterial with the supported parameter infos/names and either a BlendChannelOutput component
  *                or no BlendChannelInput, and applies the resulting parameter to the bound material instance.
  */
 UCLASS(MinimalAPI)
@@ -70,8 +71,8 @@ private:
 private:
 
 	/** Overlapping trackers that track multiple entities animating the same bound object and name */
-	UE::MovieScene::TOverlappingEntityTracker<UE::MovieScene::FAnimatedMaterialParameterInfo, UE::MovieScene::FObjectComponent, FName> ScalarParameterTracker;
-	UE::MovieScene::TOverlappingEntityTracker<UE::MovieScene::FAnimatedMaterialParameterInfo, UE::MovieScene::FObjectComponent, FName> VectorParameterTracker;
+	UE::MovieScene::TOverlappingEntityTracker<UE::MovieScene::FAnimatedMaterialParameterInfo, UE::MovieScene::FObjectComponent, FMaterialParameterInfo> ScalarParameterTracker;
+	UE::MovieScene::TOverlappingEntityTracker<UE::MovieScene::FAnimatedMaterialParameterInfo, UE::MovieScene::FObjectComponent, FMaterialParameterInfo> VectorParameterTracker;
 
 	/** Holds pre-animated values for scalar values */
 	TSharedPtr<UE::MovieScene::FPreAnimatedScalarMaterialParameterStorage> ScalarParameterStorage;
diff --git a/Engine/Source/Runtime/MovieSceneTracks/Public/Systems/MovieSceneMaterialSystem.h b/Engine/Source/Runtime/MovieSceneTracks/Public/Systems/MovieSceneMaterialSystem.h
--- a/Engine/Source/Runtime/MovieSceneTracks/Public/Systems/MovieSceneMaterialSystem.h
+++ b/Engine/Source/Runtime/MovieSceneTracks/Public/Systems/MovieSceneMaterialSystem.h
@@ -278,7 +278,8 @@ void TMovieSceneMaterialSystem<AccessorType, RequiredComponents...>::OnLink(UMov
 	// Imported entities are implicitly excluded by way of filtering by BoundObject, which do not exist on imported entities
 	MaterialParameterFilter.Reset();
 	MaterialParameterFilter.All({ InRequiredComponents... });
-	MaterialParameterFilter.Any({ TracksComponents->ScalarParameterName, TracksComponents->ColorParameterName, TracksComponents->VectorParameterName });
+	MaterialParameterFilter.Any({ TracksComponents->ScalarParameterName, TracksComponents->ColorParameterName, TracksComponents->VectorParameterName, // Old style parameter types for deprecated UMovieSceneParameterSections
+		TracksComponents->ScalarMaterialParameterInfo, TracksComponents->ColorMaterialParameterInfo, TracksComponents->VectorMaterialParameterInfo }); // New style parameter types
 
 	Linker->Events.PostSpawnEvent.AddRaw(this, &TMovieSceneMaterialSystem<AccessorType, RequiredComponents...>::OnPostSpawn, InRequiredComponents...);
 
diff --git a/Engine/Source/Runtime/MovieSceneTracks/Public/Systems/MovieScenePreAnimatedMaterialParameters.h b/Engine/Source/Runtime/MovieSceneTracks/Public/Systems/MovieScenePreAnimatedMaterialParameters.h
--- a/Engine/Source/Runtime/MovieSceneTracks/Public/Systems/MovieScenePreAnimatedMaterialParameters.h
+++ b/Engine/Source/Runtime/MovieSceneTracks/Public/Systems/MovieScenePreAnimatedMaterialParameters.h
@@ -6,6 +6,7 @@
 #include "UObject/ObjectKey.h"
 #include "Evaluation/PreAnimatedState/MovieScenePreAnimatedObjectStorage.h"
 #include "Evaluation/PreAnimatedState/MovieScenePreAnimatedStateStorage.h"
+#include "MaterialTypes.h"
 
 namespace UE::MovieScene
 {
@@ -13,9 +14,10 @@ namespace UE::MovieScene
 struct FMaterialParameterKey
 {
 	FObjectKey BoundMaterial;
-	FName ParameterName;
+	FMaterialParameterInfo ParameterInfo;
 
 	MOVIESCENETRACKS_API FMaterialParameterKey(const FObjectComponent& InBoundMaterial, const FName& InParameterName);
+	MOVIESCENETRACKS_API FMaterialParameterKey(const FObjectComponent& InBoundMaterial, const FMaterialParameterInfo& InParameterInfo);
 
 	friend MOVIESCENETRACKS_API uint32 GetTypeHash(const FMaterialParameterKey& InKey);
 
@@ -30,6 +32,7 @@ struct FMaterialParameterCollectionScalarTraits : FBoundObjectPreAnimatedStateTr
 	static MOVIESCENETRACKS_API void ReplaceObject(FMaterialParameterKey& InOutKey, const FObjectKey& NewObject);
 
 	static MOVIESCENETRACKS_API float CachePreAnimatedValue(const FObjectComponent& InBoundMaterial, const FName& ParameterName);
+	static MOVIESCENETRACKS_API float CachePreAnimatedValue(const FObjectComponent& InBoundMaterial, const FMaterialParameterInfo& ParameterInfo);
 
 	static MOVIESCENETRACKS_API void RestorePreAnimatedValue(const FMaterialParameterKey& InKey, float OldValue, const FRestoreStateParams& Params);
 };
@@ -42,6 +45,7 @@ struct FMaterialParameterCollectionVectorTraits : FBoundObjectPreAnimatedStateTr
 	static MOVIESCENETRACKS_API void ReplaceObject(FMaterialParameterKey& InOutKey, const FObjectKey& NewObject);
 
 	static MOVIESCENETRACKS_API FLinearColor CachePreAnimatedValue(const FObjectComponent& InBoundMaterial, const FName& ParameterName);
+	static MOVIESCENETRACKS_API FLinearColor CachePreAnimatedValue(const FObjectComponent& InBoundMaterial, const FMaterialParameterInfo& ParameterInfo);
 
 	static MOVIESCENETRACKS_API void RestorePreAnimatedValue(const FMaterialParameterKey& InKey, const FLinearColor& OldValue, const FRestoreStateParams& Params);
 };
diff --git a/Engine/Source/Runtime/MovieSceneTracks/Public/Tracks/MovieSceneMaterialTrack.h b/Engine/Source/Runtime/MovieSceneTracks/Public/Tracks/MovieSceneMaterialTrack.h
--- a/Engine/Source/Runtime/MovieSceneTracks/Public/Tracks/MovieSceneMaterialTrack.h
+++ b/Engine/Source/Runtime/MovieSceneTracks/Public/Tracks/MovieSceneMaterialTrack.h
@@ -7,6 +7,7 @@
 #include "Sections/MovieSceneParameterSection.h"
 #include "MovieSceneNameableTrack.h"
 #include "EntitySystem/IMovieSceneEntityProvider.h"
+#include "MaterialTypes.h"
 #include "MovieSceneMaterialTrack.generated.h"
 
 /**
@@ -40,11 +41,11 @@ public:
 public:
 
 	/**
-	 * Adds a scalar parameter key to the track. 
-	 * @param ParameterName The name of the parameter to add a key for.
-	 * @param Time The time to add the new key.
-	 * @param The value for the new key.
-	 */
+	* Adds a scalar parameter key to the track. 
+	* @param ParameterName The name of the parameter to add a key for.
+	* @param Time The time to add the new key.
+	* @param The value for the new key.
+	*/
 	MOVIESCENETRACKS_API void AddScalarParameterKey(FName ParameterName, FFrameNumber Position, float Value);
 
 	/**
@@ -57,22 +58,64 @@ public:
 	MOVIESCENETRACKS_API void AddScalarParameterKey(FName ParameterName, FFrameNumber Position, int32 RowIndex, float Value);
 
 	/**
-	 * Adds a color parameter key to the track.
-	 * @param ParameterName The name of the parameter to add a key for.
-	 * @param Time The time to add the new key.
-	 * @param The value for the new key.
-	 */
+	* Adds a scalar parameter key to the track.
+	* @param ParameterInfo The material parameter info for the parameter you want to add a key for.
+	* @param Time The time to add the new key.
+	* @param The value for the new key.
+	* @param InLayerName Optional layer name for use in UI.
+	* @param InAssetName Optional asset name for use in UI.
+	*/
+	MOVIESCENETRACKS_API void AddScalarParameterKey(const FMaterialParameterInfo& ParameterInfo, FFrameNumber Position, float Value, const FString& InLayerName, const FString& InAssetName);
+
+	/**
+	* Adds a scalar parameter key to the track.
+	* @param ParameterInfo The material parameter info for the parameter you want to add a key for.
+	* @param Time The time to add the new key.
+	* @param RowIndex The preferred row index on which to look for sections.
+	* @param The value for the new key.
+	* @param InLayerName Optional layer name for use in UI.
+	* @param InAssetName Optional asset name for use in UI.
+	*/
+	MOVIESCENETRACKS_API void AddScalarParameterKey(const FMaterialParameterInfo& ParameterInfo, FFrameNumber Position, int32 RowIndex, float Value, const FString& InLayerName, const FString& InAssetName);
+
+	/**
+	* Adds a color parameter key to the track.
+	* @param ParameterName The name of the parameter to add a key for.
+	* @param Time The time to add the new key.
+	* @param The value for the new key.
+	*/
 	MOVIESCENETRACKS_API void AddColorParameterKey(FName ParameterName, FFrameNumber Position, FLinearColor Value);
 
 	/**
-	 * Adds a color parameter key to the track.
-	 * @param ParameterName The name of the parameter to add a key for.
-	 * @param Time The time to add the new key.
-	 * @param RowIndex The preferred row index on which to look for sections.
-	 * @param The value for the new key.
-	 */
+	* Adds a color parameter key to the track.
+	* @param ParameterInfo The material parameter info for the parameter you want to add a key for.
+	* @param Time The time to add the new key.
+	* @param RowIndex The preferred row index on which to look for sections.
+	* @param The value for the new key.
+	*/
 	MOVIESCENETRACKS_API void AddColorParameterKey(FName ParameterName, FFrameNumber Position, int32 RowIndex, FLinearColor Value);
 
+	/**
+	* Adds a color parameter key to the track.
+	* @param ParameterName The name of the parameter to add a key for.
+	* @param Time The time to add the new key.
+	* @param The value for the new key.
+	* @param InLayerName Optional layer name for use in UI.
+	* @param InAssetName Optional asset name for use in UI.
+	*/
+	MOVIESCENETRACKS_API void AddColorParameterKey(const FMaterialParameterInfo& ParameterInfo, FFrameNumber Position, FLinearColor Value, const FString& InLayerName, const FString& InAssetName);
+
+	/**
+	* Adds a color parameter key to the track.
+	* @param ParameterInfo The material parameter info for the parameter you want to add a key for.
+	* @param Time The time to add the new key.
+	* @param RowIndex The preferred row index on which to look for sections.
+	* @param The value for the new key.
+	* @param InLayerName Optional layer name for use in UI.
+	* @param InAssetName Optional asset name for use in UI.
+	*/
+	MOVIESCENETRACKS_API void AddColorParameterKey(const FMaterialParameterInfo& ParameterInfo, FFrameNumber Position, int32 RowIndex, FLinearColor Value, const FString& InLayerName, const FString& InAssetName);
+
 private:
 
 	/** The sections owned by this track .*/
