Sequencer: Mark the render state dirty when the scale sign changes.

The cached mesh draw command includes the culling order. If the scale sign changes, the render state needs to be marked dirty so that the culling order is refreshed (the renderer does not currently do this automatically)

diff --git a/Engine/Source/Runtime/MovieSceneTracks/Private/MovieSceneTracksCustomAccessors.cpp b/Engine/Source/Runtime/MovieSceneTracks/Private/MovieSceneTracksCustomAccessors.cpp
--- a/Engine/Source/Runtime/MovieSceneTracks/Private/MovieSceneTracksCustomAccessors.cpp
+++ b/Engine/Source/Runtime/MovieSceneTracks/Private/MovieSceneTracksCustomAccessors.cpp
@@ -43,15 +43,25 @@ void SetComponentTransform(USceneComponent* SceneComponent, const FIntermediate3
 	// @todo: We would really rather not 
 	AActor* Actor = SceneComponent->GetOwner();
 	USceneComponent* RootComponent = Actor ? Actor->GetRootComponent() : nullptr;
-	bool bIsSimulatingPhysics = RootComponent ? RootComponent->IsSimulatingPhysics() : false;
+	const bool bIsSimulatingPhysics = RootComponent ? RootComponent->IsSimulatingPhysics() : false;
 
 	FVector Translation = Transform.GetTranslation();
 	FRotator Rotation = Transform.GetRotation();
-	SceneComponent->SetRelativeTransform(FTransform(Rotation, Translation, Transform.GetScale()), false, nullptr, bIsSimulatingPhysics ? ETeleportType::ResetPhysics : ETeleportType::None);
+	FTransform NewTransform(Rotation, Translation, Transform.GetScale());
+	const bool bFlipped = NewTransform.GetDeterminant() * SceneComponent->GetComponentTransform().GetDeterminant() <= 0;
+
+	SceneComponent->SetRelativeTransform(NewTransform, false, nullptr, bIsSimulatingPhysics ? ETeleportType::ResetPhysics : ETeleportType::None);
 
 	// Force the location and rotation values to avoid Rot->Quat->Rot conversions
 	SceneComponent->SetRelativeLocation_Direct(Translation);
 	SceneComponent->SetRelativeRotation_Direct(Rotation);
+
+	// The cached mesh draw command includes the culling order. If the scale sign changes, the render state needs 
+	// to be marked dirty so that the culling order is refreshed (the renderer does not currently do this automatically)
+	if (bFlipped)
+	{
+		SceneComponent->MarkRenderStateDirty();
+	}
 }
 
 void SetComponentTransformAndVelocity(UObject* Object, const FIntermediate3DTransform& InTransform)
