Fix a crash when adding new Material Parameter Collection sequencer tracks caused by CL 26774357.
MaterialParameterCollection tracks were inadvertently creating Component Material Parameter Sections instead of generic Parameter Sections.
When reviewing this, I realized I did not need to support the Component Material Parameter Sections in the Material Parameter Collection System, so I have rolled that back as well.

diff --git a/Engine/Source/Runtime/MovieSceneTracks/Private/Systems/MovieSceneMaterialParameterCollectionSystem.cpp b/Engine/Source/Runtime/MovieSceneTracks/Private/Systems/MovieSceneMaterialParameterCollectionSystem.cpp
--- a/Engine/Source/Runtime/MovieSceneTracks/Private/Systems/MovieSceneMaterialParameterCollectionSystem.cpp
+++ b/Engine/Source/Runtime/MovieSceneTracks/Private/Systems/MovieSceneMaterialParameterCollectionSystem.cpp
@@ -78,15 +78,10 @@ void UMovieSceneMaterialParameterCollectionSystem::OnRun(FSystemTaskPrerequisite
 			TComponentReader<TWeakObjectPtr<UMaterialParameterCollection>> MPCs = Allocation->ReadComponents(TracksComponents->MPC);
 			TComponentReader<FInstanceHandle> InstanceHandles = Allocation->ReadComponents(BuiltInComponents->InstanceHandle);
 
-			// Support older sections that still use ParameterName rather than ParameterInfo
 			TOptionalComponentReader<FName> ScalarParameterNames = Allocation->TryReadComponents(TracksComponents->ScalarParameterName);
 			TOptionalComponentReader<FName> VectorParameterNames = Allocation->TryReadComponents(TracksComponents->VectorParameterName);
 			TOptionalComponentReader<FName> ColorParameterNames = Allocation->TryReadComponents(TracksComponents->ColorParameterName);
 
-			TOptionalComponentReader<FMaterialParameterInfo> ScalarParameterInfos = Allocation->TryReadComponents(TracksComponents->ScalarMaterialParameterInfo);
-			TOptionalComponentReader<FMaterialParameterInfo> VectorParameterInfos = Allocation->TryReadComponents(TracksComponents->VectorMaterialParameterInfo);
-			TOptionalComponentReader<FMaterialParameterInfo> ColorParameterInfos = Allocation->TryReadComponents(TracksComponents->ColorMaterialParameterInfo);
-
 			const int32 Num = Allocation->Num();
 			for (int32 Index = 0; Index < Num; ++Index)
 			{
@@ -106,10 +101,9 @@ void UMovieSceneMaterialParameterCollectionSystem::OnRun(FSystemTaskPrerequisite
 						TEXT("Unable to create MPC instance for %s with World %s. Material parameter collection tracks will not function."),
 						*Collection->GetName(), *World->GetName()))
 					{
-						if (ScalarParameterInfos || ScalarParameterNames)
+						if (ScalarParameterNames)
 						{
-							// Support older sections that still use ParameterName rather than ParameterInfo
-							FName Name = ScalarParameterInfos ? ScalarParameterInfos[Index].Name : ScalarParameterNames[Index];
+							FName Name = ScalarParameterNames[Index];
 							if (Collection->GetScalarParameterByName(Name) == nullptr)
 							{
 								if (!Instance->bLoggedMissingParameterWarning)
@@ -118,14 +112,9 @@ void UMovieSceneMaterialParameterCollectionSystem::OnRun(FSystemTaskPrerequisite
 								}
 							}
 						}
-						else if (VectorParameterInfos || ColorParameterInfos || VectorParameterNames || ColorParameterNames)
+						else if (VectorParameterNames || ColorParameterNames)
 						{
-							// Support older sections that still use ParameterName rather than ParameterInfo
-							FName Name = VectorParameterInfos ? VectorParameterInfos[Index].Name :
-								ColorParameterInfos ? ColorParameterInfos[Index].Name :
-								VectorParameterNames ? VectorParameterNames[Index] :
-								ColorParameterNames[Index];
-
+							FName Name = VectorParameterNames ? VectorParameterNames[Index] : ColorParameterNames[Index];
 							if (Collection->GetVectorParameterByName(Name) == nullptr)
 							{
 								if (!Instance->bLoggedMissingParameterWarning)
@@ -183,8 +172,8 @@ void UMovieSceneMaterialParameterCollectionSystem::OnRun(FSystemTaskPrerequisite
 
 	Params.AdditionalFilter.None({ BuiltInComponents->BlendChannelOutput });
 	Params.AdditionalFilter.All({ TracksComponents->MPC });
-	ScalarParameterStorage->BeginTrackingAndCachePreAnimatedValuesTask(Linker, Params, TracksComponents->BoundMaterial, TracksComponents->ScalarMaterialParameterInfo);
-	VectorParameterStorage->BeginTrackingAndCachePreAnimatedValuesTask(Linker, Params, TracksComponents->BoundMaterial, TracksComponents->VectorMaterialParameterInfo);
-	VectorParameterStorage->BeginTrackingAndCachePreAnimatedValuesTask(Linker, Params, TracksComponents->BoundMaterial, TracksComponents->ColorMaterialParameterInfo);
+	ScalarParameterStorage->BeginTrackingAndCachePreAnimatedValuesTask(Linker, Params, TracksComponents->BoundMaterial, TracksComponents->ScalarParameterName);
+	VectorParameterStorage->BeginTrackingAndCachePreAnimatedValuesTask(Linker, Params, TracksComponents->BoundMaterial, TracksComponents->VectorParameterName);
+	VectorParameterStorage->BeginTrackingAndCachePreAnimatedValuesTask(Linker, Params, TracksComponents->BoundMaterial, TracksComponents->ColorParameterName);
 }
 
diff --git a/Engine/Source/Runtime/MovieSceneTracks/Private/Tracks/MovieSceneMaterialParameterCollectionTrack.cpp b/Engine/Source/Runtime/MovieSceneTracks/Private/Tracks/MovieSceneMaterialParameterCollectionTrack.cpp
--- a/Engine/Source/Runtime/MovieSceneTracks/Private/Tracks/MovieSceneMaterialParameterCollectionTrack.cpp
+++ b/Engine/Source/Runtime/MovieSceneTracks/Private/Tracks/MovieSceneMaterialParameterCollectionTrack.cpp
@@ -22,7 +22,8 @@ UMovieSceneMaterialParameterCollectionTrack::UMovieSceneMaterialParameterCollect
 
 UMovieSceneSection* UMovieSceneMaterialParameterCollectionTrack::CreateNewSection()
 {
-	UMovieSceneSection* NewSection = Super::CreateNewSection();
+	UMovieSceneSection* NewSection = NewObject<UMovieSceneParameterSection>(this, NAME_None, RF_Transactional);
+	NewSection->SetBlendType(EMovieSceneBlendType::Absolute);
 	NewSection->SetRange(TRange<FFrameNumber>::All());
 	return NewSection;
 }
