Sequencer: Constraints: Section may be null/pending kill now, so need to check for it.

diff --git a/Engine/Source/Runtime/MovieSceneTracks/Private/Systems/MovieSceneConstraintSystem.cpp b/Engine/Source/Runtime/MovieSceneTracks/Private/Systems/MovieSceneConstraintSystem.cpp
--- a/Engine/Source/Runtime/MovieSceneTracks/Private/Systems/MovieSceneConstraintSystem.cpp
+++ b/Engine/Source/Runtime/MovieSceneTracks/Private/Systems/MovieSceneConstraintSystem.cpp
@@ -120,6 +120,10 @@ struct FPreAnimatedConstraintTraits : FBoundObjectPreAnimatedStateTraits
 	{
 		if (USceneComponent* SceneComponent = OldValue.SceneComponent.Get())
 		{
+			if (OldValue.Section.IsValid() == false)
+			{
+				return;
+			}
 			if (FConstraintAndActiveChannel* ConstraintAndActiveChannel = OldValue.Section->GetConstraintChannel(OldValue.ConstraintName))
 			{				
 				if (ConstraintAndActiveChannel->Constraint.IsPending())
