Sequencer: Fixed skeletal animation data not getting reset when using the new scheduler

#localization none
#tests All engine tests and editor preflight tests completed

diff --git a/Engine/Source/Runtime/MovieSceneTracks/Private/Systems/MovieSceneSkeletalAnimationSystem.cpp b/Engine/Source/Runtime/MovieSceneTracks/Private/Systems/MovieSceneSkeletalAnimationSystem.cpp
--- a/Engine/Source/Runtime/MovieSceneTracks/Private/Systems/MovieSceneSkeletalAnimationSystem.cpp
+++ b/Engine/Source/Runtime/MovieSceneTracks/Private/Systems/MovieSceneSkeletalAnimationSystem.cpp
@@ -322,7 +322,8 @@ struct FGatherSkeletalAnimations
 
 	void PreTask() const
 	{
-		SystemData->SkeletalAnimations.Reset();
+		// Start fresh every frame, gathering all active skeletal animations.
+		SystemData->ResetSkeletalAnimations();
 	}
 
 	void ForEachAllocation(
@@ -1024,9 +1025,6 @@ void UMovieSceneSkeletalAnimationSystem::OnRun(FSystemTaskPrerequisites& InPrere
 	FBuiltInComponentTypes* BuiltInComponents = FBuiltInComponentTypes::Get();
 	FMovieSceneTracksComponentTypes* TrackComponents = FMovieSceneTracksComponentTypes::Get();
 
-	// Start fresh every frame, gathering all active skeletal animations.
-	SystemData.ResetSkeletalAnimations();
-
 	FGraphEventRef GatherTask = FEntityTaskBuilder()
 	.ReadEntityIDs()
 	.Read(BuiltInComponents->RootInstanceHandle)
