Fix compiler warning to do with integer overflow. Not sure why this wasn't failing compiles locally or failed by Presubmit.

diff --git a/Engine/Source/Runtime/MovieSceneTracks/Private/Sections/MovieSceneCustomPrimitiveDataSection.cpp b/Engine/Source/Runtime/MovieSceneTracks/Private/Sections/MovieSceneCustomPrimitiveDataSection.cpp
--- a/Engine/Source/Runtime/MovieSceneTracks/Private/Sections/MovieSceneCustomPrimitiveDataSection.cpp
+++ b/Engine/Source/Runtime/MovieSceneTracks/Private/Sections/MovieSceneCustomPrimitiveDataSection.cpp
@@ -115,7 +115,7 @@ void UMovieSceneCustomPrimitiveDataSection::ReconstructChannelProxy()
 		FMovieSceneChannelMetaData MetaData(Scalar.ParameterName, FText::Format(LOCTEXT("CustomPrimitiveData_Scalar_DisplayName", "Scalar (Index {0})"), FText::FromName(Scalar.ParameterName)));
 		// Should be the start index
 		check(Scalar.ParameterName.ToString().IsNumeric());
-		int32 StartIndex = FCString::Atoi(*Scalar.ParameterName.ToString());
+		uint64 StartIndex = FCString::Atoi(*Scalar.ParameterName.ToString());
 		check(StartIndex >= 0 && StartIndex < FCustomPrimitiveData::NumCustomPrimitiveDataFloats);
 		
 		// Prevent single channels from collapsing to the track node
@@ -132,9 +132,9 @@ void UMovieSceneCustomPrimitiveDataSection::ReconstructChannelProxy()
 	{
 		FString ParameterString = Vector2D.ParameterName.ToString(); 
 		check(ParameterString.IsNumeric()); 
-		int32 StartIndex = FCString::Atoi(*ParameterString);
+		uint64 StartIndex = FCString::Atoi64(*ParameterString);
 		check(StartIndex >= 0 && StartIndex < FCustomPrimitiveData::NumCustomPrimitiveDataFloats - 1);
-		ChannelsUsedBitmap |= (0b11 << StartIndex);
+		ChannelsUsedBitmap |= ((uint64)0b11 << StartIndex);
 		FText Group = FText::FromString(ParameterString);
 		FString GroupDisplayName = FText::Format(LOCTEXT("CustomPrimitiveData_Vector2D_DisplayName", "Vector2D (Index {0})"), FText::FromName(Vector2D.ParameterName)).ToString();
 
@@ -156,9 +156,9 @@ void UMovieSceneCustomPrimitiveDataSection::ReconstructChannelProxy()
 	{
 		FString ParameterString = Vector.ParameterName.ToString();
 		check(ParameterString.IsNumeric());
-		int32 StartIndex = FCString::Atoi(*ParameterString);
+		uint64 StartIndex = FCString::Atoi64(*ParameterString);
 		check(StartIndex >= 0 && StartIndex < FCustomPrimitiveData::NumCustomPrimitiveDataFloats - 2);
-		ChannelsUsedBitmap |= (0b111 << StartIndex);
+		ChannelsUsedBitmap |= ((uint64)0b111 << StartIndex);
 		FText Group = FText::FromString(ParameterString);
 		FString GroupDisplayName = FText::Format(LOCTEXT("CustomPrimitiveData_Vector_DisplayName", "Vector (Index {0})"), FText::FromName(Vector.ParameterName)).ToString();
 		
@@ -186,9 +186,9 @@ void UMovieSceneCustomPrimitiveDataSection::ReconstructChannelProxy()
 	{
 		FString ParameterString = Color.ParameterName.ToString();
 		check(ParameterString.IsNumeric());
-		int32 StartIndex = FCString::Atoi(*ParameterString);
+		uint64 StartIndex = FCString::Atoi64(*ParameterString);
 		check(StartIndex >= 0 && StartIndex < FCustomPrimitiveData::NumCustomPrimitiveDataFloats - 3);
-		ChannelsUsedBitmap |= (0b1111 << StartIndex);
+		ChannelsUsedBitmap |= ((uint64)0b1111 << StartIndex);
 		FText Group = FText::FromString(ParameterString);
 		FString GroupDisplayName = FText::Format(LOCTEXT("CustomPrimitiveData_Color_DisplayName", "Color (Index {0})"), FText::FromName(Color.ParameterName)).ToString();
 		
