Woops, remove pragma optmization instructions

diff --git a/Engine/Source/Runtime/MovieScene/Private/EntitySystem/MovieSceneEntityGroupingSystem.cpp b/Engine/Source/Runtime/MovieScene/Private/EntitySystem/MovieSceneEntityGroupingSystem.cpp
--- a/Engine/Source/Runtime/MovieScene/Private/EntitySystem/MovieSceneEntityGroupingSystem.cpp
+++ b/Engine/Source/Runtime/MovieScene/Private/EntitySystem/MovieSceneEntityGroupingSystem.cpp
@@ -9,7 +9,7 @@
 #include "EntitySystem/MovieSceneEntityIDs.h"
 #include "EntitySystem/MovieSceneEntityMutations.h"
 #include "EntitySystem/MovieSceneEntitySystemLinker.h"
-UE_DISABLE_OPTIMIZATION
+
 namespace UE::MovieScene
 {
 
@@ -321,4 +321,4 @@ TArrayView<const UE::MovieScene::FMovieSceneEntityID> UMovieSceneEntityGroupingS
 	}
 	return TArrayView<const UE::MovieScene::FMovieSceneEntityID>();
 }
-UE_ENABLE_OPTIMIZATION
+
diff --git a/Engine/Source/Runtime/MovieSceneTracks/Private/Systems/MovieScenePropertyInstantiator.cpp b/Engine/Source/Runtime/MovieSceneTracks/Private/Systems/MovieScenePropertyInstantiator.cpp
--- a/Engine/Source/Runtime/MovieSceneTracks/Private/Systems/MovieScenePropertyInstantiator.cpp
+++ b/Engine/Source/Runtime/MovieSceneTracks/Private/Systems/MovieScenePropertyInstantiator.cpp
@@ -20,7 +20,7 @@
 DECLARE_CYCLE_STAT(TEXT("DiscoverInvalidatedProperties"), MovieSceneEval_DiscoverInvalidatedProperties, STATGROUP_MovieSceneECS);
 DECLARE_CYCLE_STAT(TEXT("ProcessInvalidatedProperties"), MovieSceneEval_ProcessInvalidatedProperties, STATGROUP_MovieSceneECS);
 DECLARE_CYCLE_STAT(TEXT("InitializePropertyMetaData"), MovieSceneEval_InitializePropertyMetaData, STATGROUP_MovieSceneECS);
-UE_DISABLE_OPTIMIZATION
+
 namespace UE::MovieScene
 {
 
@@ -1214,4 +1214,4 @@ void UMovieScenePropertyInstantiatorSystem::FPropertyParameters::MakeOutputCompo
 		OutComponentType.Remove(FBuiltInComponentTypes::Get()->Tags.RestoreState);
 	}
 }
-UE_ENABLE_OPTIMIZATION
+
