Sequencer: make material systems use the new grouping system and benefit from hierarchical bias

The component material system now registers a grouping policy for material indexes, so that these kinds of entities will now be processed for hierchical bias. This requires adding a dependency on the hierarchical bias system to run after the component material system. This however leads to a cyclic dependency, since the material parameter system, which depends on the component material system, sets the hierarchical bias system as a downstream dependency. That dependency is removed in this change for now until we can better fix the problem.

The generic material system helper class now also filters out entities tagged with "Ignored", which is how hierarchical bias is applied.

diff --git a/Engine/Source/Runtime/MovieSceneTracks/Private/Systems/MovieSceneComponentMaterialSystem.cpp b/Engine/Source/Runtime/MovieSceneTracks/Private/Systems/MovieSceneComponentMaterialSystem.cpp
--- a/Engine/Source/Runtime/MovieSceneTracks/Private/Systems/MovieSceneComponentMaterialSystem.cpp
+++ b/Engine/Source/Runtime/MovieSceneTracks/Private/Systems/MovieSceneComponentMaterialSystem.cpp
@@ -8,6 +8,7 @@
 #include "Evaluation/PreAnimatedState/MovieScenePreAnimatedStorageID.inl"
 
 #include "Systems/FloatChannelEvaluatorSystem.h"
+#include "Systems/MovieSceneHierarchicalBiasSystem.h"
 #include "Systems/MovieScenePiecewiseDoubleBlenderSystem.h"
 
 #include "Materials/MaterialInstanceDynamic.h"
@@ -191,8 +192,11 @@ UMovieSceneComponentMaterialSystem::UMovieSceneComponentMaterialSystem(const FOb
 	{
 		DefineComponentConsumer(GetClass(), BuiltInComponents->ObjectResult);
 		DefineComponentConsumer(GetClass(), BuiltInComponents->BoundObject);
+
 		DefineComponentProducer(GetClass(), TracksComponents->BoundMaterial);
+
 		DefineImplicitPrerequisite(UMovieSceneCachePreAnimatedStateSystem::StaticClass(), GetClass());
+		DefineImplicitPrerequisite(UMovieSceneHierarchicalBiasSystem::StaticClass(), GetClass());
 	}
 }
 
diff --git a/Engine/Source/Runtime/MovieSceneTracks/Private/Systems/MovieSceneMaterialParameterSystem.cpp b/Engine/Source/Runtime/MovieSceneTracks/Private/Systems/MovieSceneMaterialParameterSystem.cpp
--- a/Engine/Source/Runtime/MovieSceneTracks/Private/Systems/MovieSceneMaterialParameterSystem.cpp
+++ b/Engine/Source/Runtime/MovieSceneTracks/Private/Systems/MovieSceneMaterialParameterSystem.cpp
@@ -341,9 +341,9 @@ template<typename Mixin>
 struct TOverlappingMaterialParameterHandler : Mixin
 {
 	UMovieSceneEntitySystemLinker* Linker;
-	UMovieSceneMaterialParameterSystem* System;
+	UMovieSceneMaterialParameterInstantiatorSystem* System;
 
-	TOverlappingMaterialParameterHandler(UMovieSceneMaterialParameterSystem* InSystem)
+	TOverlappingMaterialParameterHandler(UMovieSceneMaterialParameterInstantiatorSystem* InSystem)
 		: Linker(InSystem->GetLinker())
 		, System(InSystem)
 	{}
@@ -483,7 +483,7 @@ struct TOverlappingMaterialParameterHandler : Mixin
 
 } // namespace UE::MovieScene
 
-UMovieSceneMaterialParameterSystem::UMovieSceneMaterialParameterSystem(const FObjectInitializer& ObjInit)
+UMovieSceneMaterialParameterInstantiatorSystem::UMovieSceneMaterialParameterInstantiatorSystem(const FObjectInitializer& ObjInit)
 	: Super(ObjInit)
 {
 	using namespace UE::MovieScene;
@@ -492,29 +492,19 @@ UMovieSceneMaterialParameterSystem::UMovieSceneMaterialParameterSystem(const FOb
 	FMovieSceneTracksComponentTypes* TracksComponents = FMovieSceneTracksComponentTypes::Get();
 
 	RelevantComponent = TracksComponents->BoundMaterial;
-	Phase = ESystemPhase::Instantiation | ESystemPhase::Scheduling;
 
 	if (HasAnyFlags(RF_ClassDefaultObject))
 	{
 		DefineComponentConsumer(GetClass(), TracksComponents->BoundMaterial);
-		DefineComponentProducer(GetClass(), BuiltInComponents->HierarchicalBlendTarget);
-
-		DefineImplicitPrerequisite(UFloatChannelEvaluatorSystem::StaticClass(), GetClass());
-		DefineImplicitPrerequisite(UDoubleChannelEvaluatorSystem::StaticClass(), GetClass());
 
-		for (int32 Index = 0; Index < UE_ARRAY_COUNT(BuiltInComponents->DoubleResult); ++Index)
-		{
-			DefineComponentConsumer(GetClass(), BuiltInComponents->DoubleResult[Index]);
-		}
+		DefineComponentProducer(GetClass(), BuiltInComponents->HierarchicalBlendTarget);
 
 		DefineImplicitPrerequisite(UMovieSceneHierarchicalEasingInstantiatorSystem::StaticClass(), GetClass());
-		DefineImplicitPrerequisite(UMovieScenePiecewiseDoubleBlenderSystem::StaticClass(), GetClass());
-		DefineImplicitPrerequisite(GetClass(), UMovieSceneHierarchicalBiasSystem::StaticClass());
 		DefineImplicitPrerequisite(GetClass(), UMovieSceneInitialValueSystem::StaticClass());
 	}
 }
 
-void UMovieSceneMaterialParameterSystem::OnLink()
+void UMovieSceneMaterialParameterInstantiatorSystem::OnLink()
 {
 	using namespace UE::MovieScene;
 
@@ -525,7 +515,7 @@ void UMovieSceneMaterialParameterSystem::OnLink()
 	VectorParameterStorage = Linker->PreAnimatedState.GetOrCreateStorage<FPreAnimatedVectorMaterialParameterStorage>();
 }
 
-void UMovieSceneMaterialParameterSystem::OnUnlink()
+void UMovieSceneMaterialParameterInstantiatorSystem::OnUnlink()
 {
 	using namespace UE::MovieScene;
 
@@ -536,29 +526,7 @@ void UMovieSceneMaterialParameterSystem::OnUnlink()
 	VectorParameterTracker.Destroy(TOverlappingMaterialParameterHandler<FVectorMixin>(this));
 }
 
-void UMovieSceneMaterialParameterSystem::OnRun(FSystemTaskPrerequisites& InPrerequisites, FSystemSubsequentTasks& Subsequents)
-{
-	using namespace UE::MovieScene;
-
-	FMovieSceneEntitySystemRunner* ActiveRunner = Linker->GetActiveRunner();
-	if (!ensure(ActiveRunner))
-	{
-		return;
-	}
-
-	ESystemPhase CurrentPhase = ActiveRunner->GetCurrentPhase();
-	if (CurrentPhase == ESystemPhase::Instantiation)
-	{
-		OnInstantiation();
-	}
-	else if (CurrentPhase == ESystemPhase::Evaluation)
-	{
-		OnEvaluation(InPrerequisites, Subsequents);
-	}
-}
-
-
-void UMovieSceneMaterialParameterSystem::OnInstantiation()
+void UMovieSceneMaterialParameterInstantiatorSystem::OnRun(FSystemTaskPrerequisites& InPrerequisites, FSystemSubsequentTasks& Subsequents)
 {
 	using namespace UE::MovieScene;
 
@@ -644,7 +612,31 @@ void UMovieSceneMaterialParameterSystem::OnInstantiation()
 	}
 }
 
-void UMovieSceneMaterialParameterSystem::OnSchedulePersistentTasks(UE::MovieScene::IEntitySystemScheduler* TaskScheduler)
+UMovieSceneMaterialParameterEvaluationSystem::UMovieSceneMaterialParameterEvaluationSystem(const FObjectInitializer& ObjInit)
+	: Super(ObjInit)
+{
+	using namespace UE::MovieScene;
+
+	FBuiltInComponentTypes* BuiltInComponents = FBuiltInComponentTypes::Get();
+	FMovieSceneTracksComponentTypes* TracksComponents = FMovieSceneTracksComponentTypes::Get();
+
+	RelevantComponent = TracksComponents->BoundMaterial;
+	Phase = ESystemPhase::Scheduling;
+
+	if (HasAnyFlags(RF_ClassDefaultObject))
+	{
+		DefineComponentConsumer(GetClass(), TracksComponents->BoundMaterial);
+
+		for (int32 Index = 0; Index < UE_ARRAY_COUNT(BuiltInComponents->DoubleResult); ++Index)
+		{
+			DefineComponentConsumer(GetClass(), BuiltInComponents->DoubleResult[Index]);
+		}
+
+		DefineImplicitPrerequisite(UMovieScenePiecewiseDoubleBlenderSystem::StaticClass(), GetClass());
+	}
+}
+
+void UMovieSceneMaterialParameterEvaluationSystem::OnSchedulePersistentTasks(UE::MovieScene::IEntitySystemScheduler* TaskScheduler)
 {
 	using namespace UE::MovieScene;
 
@@ -671,7 +663,7 @@ void UMovieSceneMaterialParameterSystem::OnSchedulePersistentTasks(UE::MovieScen
 	.Fork_PerAllocation<FApplyVectorParameters>(&Linker->EntityManager, TaskScheduler);
 }
 
-void UMovieSceneMaterialParameterSystem::OnEvaluation(FSystemTaskPrerequisites& InPrerequisites, FSystemSubsequentTasks& Subsequents)
+void UMovieSceneMaterialParameterEvaluationSystem::OnRun(FSystemTaskPrerequisites& InPrerequisites, FSystemSubsequentTasks& Subsequents)
 {
 	using namespace UE::MovieScene;
 
diff --git a/Engine/Source/Runtime/MovieSceneTracks/Public/Systems/MovieSceneMaterialParameterSystem.h b/Engine/Source/Runtime/MovieSceneTracks/Public/Systems/MovieSceneMaterialParameterSystem.h
--- a/Engine/Source/Runtime/MovieSceneTracks/Public/Systems/MovieSceneMaterialParameterSystem.h
+++ b/Engine/Source/Runtime/MovieSceneTracks/Public/Systems/MovieSceneMaterialParameterSystem.h
@@ -3,6 +3,7 @@
 #pragma once
 
 #include "EntitySystem/MovieSceneEntitySystem.h"
+#include "EntitySystem/MovieSceneEntityInstantiatorSystem.h"
 #include "EntitySystem/MovieSceneOverlappingEntityTracker.h"
 #include "Evaluation/PreAnimatedState/MovieScenePreAnimatedStateStorage.h"
 #include "MaterialTypes.h"
@@ -37,37 +38,29 @@ MOVIESCENETRACKS_API void CollectGarbageForOutput(FAnimatedMaterialParameterInfo
 
 
 /**
- * System responsible for tracking and animating material parameter entities.
- * Operates on the following component types from FMovieSceneTracksComponentTypes:
+ * System responsible for tracking material parameter entities.
  *
- * Instantiation: Tracks any BoundMaterial with a MaterialParameterInfo, as well as deprecated ScalarParameterName, ColorParameterName or VectorParameterName.
- *                Manages adding BlendChannelInputs and Outputs where multiple entities animate the same parameter
- *                on the same bound material.
- *                BoundMaterials may be a UMaterialInstanceDynamic, or a UMaterialParameterCollectionInstance.
- *
- * Evaluation:    Visits any BoundMaterial with the supported parameter infos/names and either a BlendChannelOutput component
- *                or no BlendChannelInput, and applies the resulting parameter to the bound material instance.
+ * Tracks any BoundMaterial with a MaterialParameterInfo, as well as deprecated ScalarParameterName,
+ * ColorParameterName or VectorParameterName.
+ * Manages adding BlendChannelInputs and Outputs where multiple entities animate the same parameter
+ * on the same bound material.
+ * BoundMaterials may be a UMaterialInstanceDynamic, or a UMaterialParameterCollectionInstance.
  */
 UCLASS(MinimalAPI)
-class UMovieSceneMaterialParameterSystem
-	: public UMovieSceneEntitySystem
+class UMovieSceneMaterialParameterInstantiatorSystem : public UMovieSceneEntityInstantiatorSystem
 {
 public:
 
 	GENERATED_BODY()
 
-	UMovieSceneMaterialParameterSystem(const FObjectInitializer& ObjInit);
+	UMovieSceneMaterialParameterInstantiatorSystem(const FObjectInitializer& ObjInit);
 
 private:
 
 	virtual void OnLink() override;
 	virtual void OnUnlink() override;
-	virtual void OnSchedulePersistentTasks(UE::MovieScene::IEntitySystemScheduler* TaskScheduler) override;
 	virtual void OnRun(FSystemTaskPrerequisites& InPrerequisites, FSystemSubsequentTasks& Subsequents) override;
 
-	void OnInstantiation();
-	void OnEvaluation(FSystemTaskPrerequisites& InPrerequisites, FSystemSubsequentTasks& Subsequents);
-
 private:
 
 	/** Overlapping trackers that track multiple entities animating the same bound object and name */
@@ -78,8 +71,31 @@ private:
 	TSharedPtr<UE::MovieScene::FPreAnimatedScalarMaterialParameterStorage> ScalarParameterStorage;
 	/** Holds pre-animated values for vector or color values */
 	TSharedPtr<UE::MovieScene::FPreAnimatedVectorMaterialParameterStorage> VectorParameterStorage;
+
 public:
 
 	UPROPERTY()
 	TObjectPtr<UMovieScenePiecewiseDoubleBlenderSystem> DoubleBlenderSystem;
 };
+
+
+/**
+ * System responsible for animating material parameter entities.
+ *
+ * Visits any BoundMaterial with the supported parameter names and either a BlendChannelOutput component
+ * or no BlendChannelInput, and applies the resulting parameter to the bound material instance.
+ */
+UCLASS(MinimalAPI)
+class UMovieSceneMaterialParameterEvaluationSystem : public UMovieSceneEntitySystem
+{
+public:
+
+	GENERATED_BODY()
+
+	UMovieSceneMaterialParameterEvaluationSystem(const FObjectInitializer& ObjInit);
+
+private:
+
+	virtual void OnSchedulePersistentTasks(UE::MovieScene::IEntitySystemScheduler* TaskScheduler) override;
+	virtual void OnRun(FSystemTaskPrerequisites& InPrerequisites, FSystemSubsequentTasks& Subsequents) override;
+};
diff --git a/Engine/Source/Runtime/MovieSceneTracks/Public/Systems/MovieSceneMaterialSystem.h b/Engine/Source/Runtime/MovieSceneTracks/Public/Systems/MovieSceneMaterialSystem.h
--- a/Engine/Source/Runtime/MovieSceneTracks/Public/Systems/MovieSceneMaterialSystem.h
+++ b/Engine/Source/Runtime/MovieSceneTracks/Public/Systems/MovieSceneMaterialSystem.h
@@ -13,6 +13,7 @@
 #include "EntitySystem/MovieSceneEntitySystemRunner.h"
 #include "EntitySystem/MovieSceneComponentTypeInfo.h"
 #include "EntitySystem/MovieSceneCachedEntityFilterResult.h"
+#include "EntitySystem/MovieSceneEntityGroupingSystem.h"
 #include "Evaluation/PreAnimatedState/MovieScenePreAnimatedObjectStorage.h"
 #include "MovieSceneTracksComponentTypes.h"
 #include "UObject/SoftObjectPtr.h"
@@ -104,6 +105,8 @@ public:
 
 protected:
 
+	UE::MovieScene::FEntityGroupingPolicyKey GroupingKey;
+
 	FEntityComponentFilter MaterialSwitcherFilter;
 	FEntityComponentFilter MaterialParameterFilter;
 
@@ -271,6 +274,10 @@ void TMovieSceneMaterialSystem<AccessorType, RequiredComponents...>::OnLink(UMov
 	FBuiltInComponentTypes*          BuiltInComponents = FBuiltInComponentTypes::Get();
 	FMovieSceneTracksComponentTypes* TracksComponents  = FMovieSceneTracksComponentTypes::Get();
 
+	// Define a grouping for these materials. This will make hierarchical bias work.
+	UMovieSceneEntityGroupingSystem* GroupingSystem = Linker->LinkSystem<UMovieSceneEntityGroupingSystem>();
+	GroupingKey = GroupingSystem->AddGrouping<RequiredComponents...>(InRequiredComponents...);
+
 	MaterialSwitcherFilter.Reset();
 	MaterialSwitcherFilter.All({ InRequiredComponents..., BuiltInComponents->ObjectResult });
 
@@ -290,6 +297,13 @@ void TMovieSceneMaterialSystem<AccessorType, RequiredComponents...>::OnLink(UMov
 template<typename AccessorType, typename... RequiredComponents>
 void TMovieSceneMaterialSystem<AccessorType, RequiredComponents...>::OnUnlink(UMovieSceneEntitySystemLinker* Linker)
 {
+	UMovieSceneEntityGroupingSystem* GroupingSystem = Linker->FindSystem<UMovieSceneEntityGroupingSystem>();
+	if (ensure(GroupingSystem))
+	{
+		GroupingSystem->RemoveGrouping(GroupingKey);
+	}
+	GroupingKey = FEntityGroupingPolicyKey();
+
 	Linker->Events.PostSpawnEvent.RemoveAll(this);
 }
 
@@ -300,7 +314,6 @@ void TMovieSceneMaterialSystem<AccessorType, RequiredComponents...>::OnPostSpawn
 
 	SCOPE_CYCLE_COUNTER(MovieSceneEval_ReinitializeBoundMaterials)
 
-	FBuiltInComponentTypes*          BuiltInComponents = FBuiltInComponentTypes::Get();
 	FMovieSceneTracksComponentTypes* TracksComponents  = FMovieSceneTracksComponentTypes::Get();
 
 	FEntityAllocationWriteContext WriteContext(InLinker->EntityManager);
@@ -344,6 +357,7 @@ void TMovieSceneMaterialSystem<AccessorType, RequiredComponents...>::OnRun(UMovi
 		.ReadAllOf(InRequiredComponents...)
 		.Read(BuiltInComponents->ObjectResult)
 		.FilterAll({ BuiltInComponents->Tags.NeedsLink })
+		.FilterNone({ BuiltInComponents->Tags.Ignored })
 		.RunInline_PerEntity(&Linker->EntityManager, ApplyMaterialSwitchers);
 
 		// --------------------------------------------------------------------------------------
@@ -364,7 +378,7 @@ void TMovieSceneMaterialSystem<AccessorType, RequiredComponents...>::OnRun(UMovi
 		.ReadAllOf(InRequiredComponents...)
 		.Write(TracksComponents->BoundMaterial)
 		.FilterAll({ BuiltInComponents->Tags.NeedsLink })
-		.FilterNone({ BuiltInComponents->Tags.NeedsUnlink })
+		.FilterNone({ BuiltInComponents->Tags.Ignored, BuiltInComponents->Tags.NeedsUnlink })
 		.RunInline_PerEntity(&Linker->EntityManager, ReinitializeBoundMaterialsTask);
 	}
 }
