Sequencer: convert slomo track to use the ECS evaluation

diff --git a/Engine/Source/Runtime/MovieSceneTracks/Private/Evaluation/MovieSceneSlomoTemplate.cpp b/Engine/Source/Runtime/MovieSceneTracks/Private/Evaluation/MovieSceneSlomoTemplate.cpp
deleted file mode 100644
diff --git a/Engine/Source/Runtime/MovieSceneTracks/Private/Evaluation/MovieSceneSlomoTemplate.h b/Engine/Source/Runtime/MovieSceneTracks/Private/Evaluation/MovieSceneSlomoTemplate.h
deleted file mode 100644
diff --git a/Engine/Source/Runtime/MovieSceneTracks/Private/MovieSceneTracksComponentTypes.cpp b/Engine/Source/Runtime/MovieSceneTracks/Private/MovieSceneTracksComponentTypes.cpp
--- a/Engine/Source/Runtime/MovieSceneTracks/Private/MovieSceneTracksComponentTypes.cpp
+++ b/Engine/Source/Runtime/MovieSceneTracks/Private/MovieSceneTracksComponentTypes.cpp
@@ -564,6 +564,9 @@ FMovieSceneTracksComponentTypes::FMovieSceneTracksComponentTypes()
 	Tags.BoundMaterialChanged = ComponentRegistry->NewTag(TEXT("Bound Material Changed"));
 	FBuiltInComponentTypes::Get()->RequiresInstantiationMask.Set(Tags.BoundMaterialChanged);
 
+	Tags.Slomo = ComponentRegistry->NewTag(TEXT("Slomo"));
+	ComponentRegistry->Factories.DefineChildComponent(Tags.Slomo, Tags.Slomo);
+
 	// Used to indicate the ParameterName component for certain parameter types (scalar, vector2d, vector, color)
 	// should be interpreted as an index for custom primitive data.
 	Tags.CustomPrimitiveData = ComponentRegistry->NewTag(TEXT("Custom Primitive Data"));
diff --git a/Engine/Source/Runtime/MovieSceneTracks/Private/Sections/MovieSceneSlomoSection.cpp b/Engine/Source/Runtime/MovieSceneTracks/Private/Sections/MovieSceneSlomoSection.cpp
--- a/Engine/Source/Runtime/MovieSceneTracks/Private/Sections/MovieSceneSlomoSection.cpp
+++ b/Engine/Source/Runtime/MovieSceneTracks/Private/Sections/MovieSceneSlomoSection.cpp
@@ -3,6 +3,8 @@
 #include "Sections/MovieSceneSlomoSection.h"
 #include "UObject/SequencerObjectVersion.h"
 #include "Channels/MovieSceneChannelProxy.h"
+#include "EntitySystem/BuiltInComponentTypes.h"
+#include "MovieSceneTracksComponentTypes.h"
 
 #include UE_INLINE_GENERATED_CPP_BY_NAME(MovieSceneSlomoSection)
 
@@ -35,3 +37,23 @@ UMovieSceneSlomoSection::UMovieSceneSlomoSection()
 #endif
 }
 
+void UMovieSceneSlomoSection::ImportEntityImpl(UMovieSceneEntitySystemLinker* EntityLinker, const FEntityImportParams& Params, FImportedEntity* OutImportedEntity)
+{
+	using namespace UE::MovieScene;
+
+	if (!FloatCurve.HasAnyData())
+	{
+		return;
+	}
+
+	const FBuiltInComponentTypes* BuiltInComponents = FBuiltInComponentTypes::Get();
+	const FMovieSceneTracksComponentTypes* TrackComponents = FMovieSceneTracksComponentTypes::Get();
+
+	OutImportedEntity->AddBuilder(
+		FEntityBuilder()
+		.AddTag(BuiltInComponents->Tags.Root)
+		.AddTag(TrackComponents->Tags.Slomo)
+		.Add(BuiltInComponents->FloatChannel[0], &FloatCurve)
+	);
+}
+
diff --git a/Engine/Source/Runtime/MovieSceneTracks/Private/Systems/MovieSceneSlomoSystem.cpp b/Engine/Source/Runtime/MovieSceneTracks/Private/Systems/MovieSceneSlomoSystem.cpp
new file mode 100644
--- /dev/null
+++ b/Engine/Source/Runtime/MovieSceneTracks/Private/Systems/MovieSceneSlomoSystem.cpp
@@ -0,0 +1,189 @@
+// Copyright Epic Games, Inc. All Rights Reserved.
+
+#include "Systems/MovieSceneSlomoSystem.h"
+
+#include "Engine/World.h"
+#include "EntitySystem/BuiltInComponentTypes.h"
+#include "Evaluation/PreAnimatedState/MovieScenePreAnimatedStateStorage.h"
+#include "Evaluation/PreAnimatedState/MovieScenePreAnimatedStorageID.inl"
+#include "GameFramework/WorldSettings.h"
+#include "IMovieScenePlayer.h"
+#include "MovieSceneTracksComponentTypes.h"
+#include "Sections/MovieSceneSlomoSection.h"
+
+#include UE_INLINE_GENERATED_CPP_BY_NAME(MovieSceneSlomoSystem)
+
+namespace UE::MovieScene
+{
+
+struct FSlomoUtil
+{
+	static void ApplySlomo(IMovieScenePlayer& Player, double TimeDilation)
+	{
+		UObject* PlaybackContext = Player.GetPlaybackContext();
+		UWorld* World = PlaybackContext ? PlaybackContext->GetWorld() : nullptr;
+
+		if (!World || (!GIsEditor && World->GetNetMode() == NM_Client) || TimeDilation <= 0.f)
+		{
+			return;
+		}
+
+		AWorldSettings* WorldSettings = World->GetWorldSettings();
+
+		if (WorldSettings)
+		{
+			WorldSettings->CinematicTimeDilation = TimeDilation;
+			WorldSettings->ForceNetUpdate();
+		}
+	}
+};
+
+struct FPreAnimatedSlomoState
+{
+	TOptional<double> TimeDilation;
+
+	static FPreAnimatedSlomoState SaveState(IMovieScenePlayer* Player)
+	{
+		if (AWorldSettings* WorldSettings = Player->GetPlaybackContext()->GetWorld()->GetWorldSettings())
+		{
+			return FPreAnimatedSlomoState{ WorldSettings->CinematicTimeDilation };
+		}
+		return FPreAnimatedSlomoState();
+	}
+
+	void RestoreState(const FMovieSceneAnimTypeID& Unused, const FRestoreStateParams& Params)
+	{
+		IMovieScenePlayer* Player = Params.GetTerminalPlayer();
+		if (!ensure(Player))
+		{
+			return;
+		}
+
+		if (TimeDilation.IsSet())
+		{
+			FSlomoUtil::ApplySlomo(*Player, TimeDilation.GetValue());
+		}
+	}
+};
+
+struct FPreAnimatedSlomoStateStorage : TSimplePreAnimatedStateStorage<FMovieSceneAnimTypeID, FPreAnimatedSlomoState>
+{
+	static TAutoRegisterPreAnimatedStorageID<FPreAnimatedSlomoStateStorage> StorageID;
+
+	FPreAnimatedStorageID GetStorageType() const override { return StorageID; }
+};
+
+TAutoRegisterPreAnimatedStorageID<FPreAnimatedSlomoStateStorage> FPreAnimatedSlomoStateStorage::StorageID;
+
+struct FEvaluateSlomo
+{
+	const FInstanceRegistry* InstanceRegistry;
+
+	FEvaluateSlomo(const FInstanceRegistry* InInstanceRegistry)
+		: InstanceRegistry(InInstanceRegistry)
+	{}
+
+	void ForEachAllocation(const FEntityAllocation* Allocation, TRead<FRootInstanceHandle> RootInstanceHandles, TRead<double> TimeDilations) const
+	{
+		const int32 Num = Allocation->Num();
+		for (int32 Index = 0; Index < Num; ++Index)
+		{
+			FRootInstanceHandle RootInstanceHandle = RootInstanceHandles[Index];
+			const FSequenceInstance& Instance = InstanceRegistry->GetInstance(RootInstanceHandle);
+			IMovieScenePlayer* Player = Instance.GetPlayer();
+
+			const double TimeDilation(TimeDilations[Index]);
+			FSlomoUtil::ApplySlomo(*Player, TimeDilation);
+		}
+	}
+};
+
+} // namespace UE::MovieScene
+
+UMovieSceneSlomoSystem::UMovieSceneSlomoSystem(const FObjectInitializer& ObjInit)
+	: Super(ObjInit)
+{
+	using namespace UE::MovieScene;
+
+	RelevantComponent = FMovieSceneTracksComponentTypes::Get()->Tags.Slomo;
+	Phase = ESystemPhase::Scheduling;
+
+	if (HasAnyFlags(RF_ClassDefaultObject))
+	{
+		const FBuiltInComponentTypes* BuiltInComponents = FBuiltInComponentTypes::Get();
+		DefineComponentConsumer(GetClass(), BuiltInComponents->DoubleResult[0]);
+	}
+}
+
+void UMovieSceneSlomoSystem::OnLink()
+{
+	using namespace UE::MovieScene;
+
+	PreAnimatedStorage = Linker->PreAnimatedState.GetOrCreateStorage<FPreAnimatedSlomoStateStorage>();
+}
+
+void UMovieSceneSlomoSystem::OnSchedulePersistentTasks(UE::MovieScene::IEntitySystemScheduler* TaskScheduler)
+{
+	using namespace UE::MovieScene;
+
+	const FBuiltInComponentTypes* BuiltInComponents = FBuiltInComponentTypes::Get();
+	const FMovieSceneTracksComponentTypes* TrackComponents = FMovieSceneTracksComponentTypes::Get();
+
+	FEntityTaskBuilder()
+	.Read(BuiltInComponents->RootInstanceHandle)
+	.Read(BuiltInComponents->DoubleResult[0])
+	.FilterAll({ TrackComponents->Tags.Slomo })
+	.SetDesiredThread(Linker->EntityManager.GetGatherThread())
+	.Schedule_PerAllocation<FEvaluateSlomo>(& Linker->EntityManager, TaskScheduler, Linker->GetInstanceRegistry());
+}
+
+void UMovieSceneSlomoSystem::OnRun(FSystemTaskPrerequisites& InPrerequisites, FSystemSubsequentTasks& Subsequents)
+{
+	using namespace UE::MovieScene;
+
+	const FBuiltInComponentTypes* BuiltInComponents = FBuiltInComponentTypes::Get();
+	const FMovieSceneTracksComponentTypes* TrackComponents = FMovieSceneTracksComponentTypes::Get();
+
+	FEntityTaskBuilder()
+	.Read(BuiltInComponents->RootInstanceHandle)
+	.Read(BuiltInComponents->DoubleResult[0])
+	.FilterAll({ TrackComponents->Tags.Slomo })
+	.SetDesiredThread(Linker->EntityManager.GetGatherThread())
+	.Dispatch_PerAllocation<FEvaluateSlomo>(& Linker->EntityManager, InPrerequisites, &Subsequents, Linker->GetInstanceRegistry());
+}
+
+void UMovieSceneSlomoSystem::SavePreAnimatedState(const FPreAnimationParameters& InParameters)
+{
+	using namespace UE::MovieScene;
+
+	static const TMovieSceneAnimTypeID<UMovieSceneSlomoSystem> AnimTypeID;
+
+	const FBuiltInComponentTypes* BuiltInComponents = FBuiltInComponentTypes::Get();
+	const FMovieSceneTracksComponentTypes* TrackComponents = FMovieSceneTracksComponentTypes::Get();
+	FInstanceRegistry* InstanceRegistry = Linker->GetInstanceRegistry();
+
+	auto SaveSlomoStates = [this, BuiltInComponents, InstanceRegistry](
+			const FEntityAllocation* Allocation, TRead<FMovieSceneEntityID> EntityIDs, TRead<FRootInstanceHandle> RootInstanceHandles)
+	{
+		const bool bWantsRestoreState = Allocation->HasComponent(BuiltInComponents->Tags.RestoreState);
+		const FMovieSceneAnimTypeID Key = AnimTypeID;
+
+		const int32 Num = Allocation->Num();
+		for (int32 Index = 0; Index < Num; ++Index)
+		{
+			FRootInstanceHandle RootInstanceHandle = RootInstanceHandles[Index];
+			const FSequenceInstance& Instance = InstanceRegistry->GetInstance(RootInstanceHandle);
+			IMovieScenePlayer* Player = Instance.GetPlayer();
+
+			PreAnimatedStorage->BeginTrackingEntity(EntityIDs[Index], bWantsRestoreState, RootInstanceHandle, Key);
+			PreAnimatedStorage->CachePreAnimatedValue(Key, [Player](const FMovieSceneAnimTypeID&) { return FPreAnimatedSlomoState::SaveState(Player); });
+		}
+	};
+
+	FEntityTaskBuilder()
+	.ReadEntityIDs()
+	.Read(BuiltInComponents->RootInstanceHandle)
+	.FilterAll({ TrackComponents->Tags.Slomo })
+	.Iterate_PerAllocation(&Linker->EntityManager, SaveSlomoStates);
+}
+
diff --git a/Engine/Source/Runtime/MovieSceneTracks/Private/Systems/MovieSceneSlomoSystem.h b/Engine/Source/Runtime/MovieSceneTracks/Private/Systems/MovieSceneSlomoSystem.h
new file mode 100644
--- /dev/null
+++ b/Engine/Source/Runtime/MovieSceneTracks/Private/Systems/MovieSceneSlomoSystem.h
@@ -0,0 +1,39 @@
+// Copyright Epic Games, Inc. All Rights Reserved.
+
+#pragma once
+
+#include "CoreTypes.h"
+#include "EntitySystem/MovieSceneEntitySystem.h"
+#include "EntitySystem/MovieScenePreAnimatedStateSystem.h"
+
+#include "MovieSceneSlomoSystem.generated.h"
+
+namespace UE::MovieScene
+{
+	struct FPreAnimatedSlomoStateStorage;
+}
+
+/**
+ * System for evaluating and applying world time dilation
+ */
+UCLASS()
+class UMovieSceneSlomoSystem 
+	: public UMovieSceneEntitySystem
+	, public IMovieScenePreAnimatedStateSystemInterface
+{
+	GENERATED_BODY()
+
+public:
+	UMovieSceneSlomoSystem(const FObjectInitializer& ObjInit);
+
+	virtual void OnLink() override;
+	virtual void OnSchedulePersistentTasks(UE::MovieScene::IEntitySystemScheduler* TaskScheduler) override;
+	virtual void OnRun(FSystemTaskPrerequisites& InPrerequisites, FSystemSubsequentTasks& Subsequents) override;
+
+private:
+	virtual void SavePreAnimatedState(const FPreAnimationParameters& InParameters) override;
+
+private:
+	TSharedPtr<UE::MovieScene::FPreAnimatedSlomoStateStorage> PreAnimatedStorage;
+};
+
diff --git a/Engine/Source/Runtime/MovieSceneTracks/Private/Tracks/MovieSceneSlomoTrack.cpp b/Engine/Source/Runtime/MovieSceneTracks/Private/Tracks/MovieSceneSlomoTrack.cpp
--- a/Engine/Source/Runtime/MovieSceneTracks/Private/Tracks/MovieSceneSlomoTrack.cpp
+++ b/Engine/Source/Runtime/MovieSceneTracks/Private/Tracks/MovieSceneSlomoTrack.cpp
@@ -2,7 +2,6 @@
 
 #include "Tracks/MovieSceneSlomoTrack.h"
 #include "Sections/MovieSceneSlomoSection.h"
-#include "Evaluation/MovieSceneSlomoTemplate.h"
 #include "Evaluation/MovieSceneEvaluationTrack.h"
 #include "Templates/Casts.h"
 
@@ -29,12 +28,6 @@ UMovieSceneSection* UMovieSceneSlomoTrack::CreateNewSection()
 	return NewObject<UMovieSceneSlomoSection>(this, NAME_None, RF_Transactional);
 }
 
-FMovieSceneEvalTemplatePtr UMovieSceneSlomoTrack::CreateTemplateForSection(const UMovieSceneSection& InSection) const
-{
-	return FMovieSceneSlomoSectionTemplate(*CastChecked<UMovieSceneSlomoSection>(&InSection));
-}
-
-
 #if WITH_EDITORONLY_DATA
 
 FText UMovieSceneSlomoTrack::GetDefaultDisplayName() const
@@ -44,6 +37,5 @@ FText UMovieSceneSlomoTrack::GetDefaultDisplayName() const
 
 #endif
 
-
 #undef LOCTEXT_NAMESPACE
 
diff --git a/Engine/Source/Runtime/MovieSceneTracks/Public/MovieSceneTracksComponentTypes.h b/Engine/Source/Runtime/MovieSceneTracks/Public/MovieSceneTracksComponentTypes.h
--- a/Engine/Source/Runtime/MovieSceneTracks/Public/MovieSceneTracksComponentTypes.h
+++ b/Engine/Source/Runtime/MovieSceneTracks/Public/MovieSceneTracksComponentTypes.h
@@ -2,6 +2,7 @@
 
 #pragma once
 
+#include "Math/Range.h"
 #include "Math/Transform.h"
 #include "EntitySystem/BuiltInComponentTypes.h"
 #include "EntitySystem/MovieSceneEntityIDs.h"
@@ -598,6 +599,7 @@ struct FMovieSceneTracksComponentTypes
 	{
 		FComponentTypeID BoundMaterialChanged;
 		FComponentTypeID CustomPrimitiveData;
+		FComponentTypeID Slomo;
 	} Tags;
 
 	TComponentTypeID<FLevelVisibilityComponentData> LevelVisibility;
diff --git a/Engine/Source/Runtime/MovieSceneTracks/Public/Sections/MovieSceneSlomoSection.h b/Engine/Source/Runtime/MovieSceneTracks/Public/Sections/MovieSceneSlomoSection.h
--- a/Engine/Source/Runtime/MovieSceneTracks/Public/Sections/MovieSceneSlomoSection.h
+++ b/Engine/Source/Runtime/MovieSceneTracks/Public/Sections/MovieSceneSlomoSection.h
@@ -4,6 +4,7 @@
 
 #include "Channels/MovieSceneFloatChannel.h"
 #include "CoreMinimal.h"
+#include "EntitySystem/IMovieSceneEntityProvider.h"
 #include "MovieSceneSection.h"
 #include "Sections/MovieSceneFloatSection.h"
 #include "UObject/ObjectMacros.h"
@@ -20,12 +21,16 @@ class UObject;
 UCLASS(MinimalAPI)
 class UMovieSceneSlomoSection
 	: public UMovieSceneSection
+	, public IMovieSceneEntityProvider
 {
 	GENERATED_BODY()
 
 	/** Default constructor. */
 	UMovieSceneSlomoSection();
 
+	/** IMovieSceneEntityProvider interface */
+	void ImportEntityImpl(UMovieSceneEntitySystemLinker* EntityLinker, const FEntityImportParams& Params, FImportedEntity* OutImportedEntity) override;
+
 public:
 
 	/** Float data */
diff --git a/Engine/Source/Runtime/MovieSceneTracks/Public/Tracks/MovieSceneSlomoTrack.h b/Engine/Source/Runtime/MovieSceneTracks/Public/Tracks/MovieSceneSlomoTrack.h
--- a/Engine/Source/Runtime/MovieSceneTracks/Public/Tracks/MovieSceneSlomoTrack.h
+++ b/Engine/Source/Runtime/MovieSceneTracks/Public/Tracks/MovieSceneSlomoTrack.h
@@ -5,7 +5,6 @@
 #include "CoreMinimal.h"
 #include "UObject/ObjectMacros.h"
 #include "Tracks/MovieSceneFloatTrack.h"
-#include "Compilation/IMovieSceneTrackTemplateProducer.h"
 #include "MovieSceneSlomoTrack.generated.h"
 
 /**
@@ -14,7 +13,6 @@
 UCLASS(MinimalAPI)
 class UMovieSceneSlomoTrack
 	: public UMovieSceneFloatTrack
-	, public IMovieSceneTrackTemplateProducer
 {
 	GENERATED_BODY()
 
@@ -26,7 +24,6 @@ public:
 
 	virtual bool SupportsType(TSubclassOf<UMovieSceneSection> SectionClass) const override;
 	virtual UMovieSceneSection* CreateNewSection() override;
-	virtual FMovieSceneEvalTemplatePtr CreateTemplateForSection(const UMovieSceneSection& InSection) const override;
 
 #if WITH_EDITORONLY_DATA
 	virtual FText GetDefaultDisplayName() const override;
