Sequencer: fix false positive assert from firing when an incorrect name is given to property bindings.

diff --git a/Engine/Source/Runtime/MovieSceneTracks/Private/Systems/MovieScenePropertyInstantiator.cpp b/Engine/Source/Runtime/MovieSceneTracks/Private/Systems/MovieScenePropertyInstantiator.cpp
--- a/Engine/Source/Runtime/MovieSceneTracks/Private/Systems/MovieScenePropertyInstantiator.cpp
+++ b/Engine/Source/Runtime/MovieSceneTracks/Private/Systems/MovieScenePropertyInstantiator.cpp
@@ -1113,7 +1113,7 @@ bool UMovieScenePropertyInstantiatorSystem::ResolveProperty(UE::MovieScene::FCus
 #if !UE_BUILD_SHIPPING
 		const FObjectPropertyInfo& ResolvedProperty = ResolvedProperties[GroupID.GroupIndex];
 		ensure(ResolvedProperty.BoundObject == Object);
-		ensure(ResolvedProperty.PropertyBinding == PropertyBinding);
+		ensure(ResolvedProperty.PropertyBinding.PropertyPath == PropertyBinding.PropertyPath);
 #endif
 		return true;
 	}
