Sequencer: upgrade camera shake tracks to use the ECS evaluation

diff --git a/Engine/Source/Runtime/MovieSceneTracks/Private/Evaluation/MovieSceneCameraShakeSourceShakeTemplate.cpp b/Engine/Source/Runtime/MovieSceneTracks/Private/Evaluation/MovieSceneCameraShakeSourceShakeTemplate.cpp
deleted file mode 100644
diff --git a/Engine/Source/Runtime/MovieSceneTracks/Private/Evaluation/MovieSceneCameraShakeSourceShakeTemplate.h b/Engine/Source/Runtime/MovieSceneTracks/Private/Evaluation/MovieSceneCameraShakeSourceShakeTemplate.h
deleted file mode 100644
diff --git a/Engine/Source/Runtime/MovieSceneTracks/Private/Evaluation/MovieSceneCameraShakeSourceTriggerTemplate.cpp b/Engine/Source/Runtime/MovieSceneTracks/Private/Evaluation/MovieSceneCameraShakeSourceTriggerTemplate.cpp
deleted file mode 100644
diff --git a/Engine/Source/Runtime/MovieSceneTracks/Private/Evaluation/MovieSceneCameraShakeSourceTriggerTemplate.h b/Engine/Source/Runtime/MovieSceneTracks/Private/Evaluation/MovieSceneCameraShakeSourceTriggerTemplate.h
deleted file mode 100644
diff --git a/Engine/Source/Runtime/MovieSceneTracks/Private/Evaluation/MovieSceneCameraShakeTemplate.cpp b/Engine/Source/Runtime/MovieSceneTracks/Private/Evaluation/MovieSceneCameraShakeTemplate.cpp
deleted file mode 100644
diff --git a/Engine/Source/Runtime/MovieSceneTracks/Private/MovieSceneTracksComponentTypes.cpp b/Engine/Source/Runtime/MovieSceneTracks/Private/MovieSceneTracksComponentTypes.cpp
--- a/Engine/Source/Runtime/MovieSceneTracks/Private/MovieSceneTracksComponentTypes.cpp
+++ b/Engine/Source/Runtime/MovieSceneTracks/Private/MovieSceneTracksComponentTypes.cpp
@@ -561,6 +561,8 @@ FMovieSceneTracksComponentTypes::FMovieSceneTracksComponentTypes()
 	ComponentRegistry->NewComponentType(&AudioInputs,            TEXT("Audio Inputs"), EComponentTypeFlags::CopyToChildren);
 	ComponentRegistry->NewComponentType(&AudioTriggerName,       TEXT("Audio Trigger Name"), EComponentTypeFlags::CopyToChildren);
 
+	ComponentRegistry->NewComponentType(&CameraShake,            TEXT("Camera Shake"), EComponentTypeFlags::CopyToChildren);
+
 	Tags.BoundMaterialChanged = ComponentRegistry->NewTag(TEXT("Bound Material Changed"));
 	FBuiltInComponentTypes::Get()->RequiresInstantiationMask.Set(Tags.BoundMaterialChanged);
 
diff --git a/Engine/Source/Runtime/MovieSceneTracks/Private/Sections/MovieSceneCameraShakeSection.cpp b/Engine/Source/Runtime/MovieSceneTracks/Private/Sections/MovieSceneCameraShakeSection.cpp
--- a/Engine/Source/Runtime/MovieSceneTracks/Private/Sections/MovieSceneCameraShakeSection.cpp
+++ b/Engine/Source/Runtime/MovieSceneTracks/Private/Sections/MovieSceneCameraShakeSection.cpp
@@ -1,8 +1,8 @@
 // Copyright Epic Games, Inc. All Rights Reserved.
 
 #include "Sections/MovieSceneCameraShakeSection.h"
-#include "Tracks/MovieSceneCameraShakeTrack.h"
 #include "UObject/SequencerObjectVersion.h"
+#include "MovieSceneTracksComponentTypes.h"
 
 #include UE_INLINE_GENERATED_CPP_BY_NAME(MovieSceneCameraShakeSection)
 
@@ -46,3 +46,25 @@ void UMovieSceneCameraShakeSection::PostLoad()
 	Super::PostLoad();
 }
 
+void UMovieSceneCameraShakeSection::ImportEntityImpl(UMovieSceneEntitySystemLinker* EntityLinker, const FEntityImportParams& Params, FImportedEntity* OutImportedEntity)
+{
+	using namespace UE::MovieScene;
+
+	if (!ShakeData.ShakeClass)
+	{
+		return;
+	}
+
+	const FBuiltInComponentTypes* BuiltInComponents = FBuiltInComponentTypes::Get();
+	const FMovieSceneTracksComponentTypes* TrackComponents = FMovieSceneTracksComponentTypes::Get();
+
+	const FGuid ObjectBindingID = Params.GetObjectBindingID();
+	FMovieSceneCameraShakeComponentData ComponentData(ShakeData, GetRange());
+
+	OutImportedEntity->AddBuilder(
+		FEntityBuilder()
+		.AddConditional(BuiltInComponents->GenericObjectBinding, ObjectBindingID, ObjectBindingID.IsValid())
+		.Add(TrackComponents->CameraShake, ComponentData)
+	);
+}
+
diff --git a/Engine/Source/Runtime/MovieSceneTracks/Private/Sections/MovieSceneCameraShakeSourceShakeSection.cpp b/Engine/Source/Runtime/MovieSceneTracks/Private/Sections/MovieSceneCameraShakeSourceShakeSection.cpp
--- a/Engine/Source/Runtime/MovieSceneTracks/Private/Sections/MovieSceneCameraShakeSourceShakeSection.cpp
+++ b/Engine/Source/Runtime/MovieSceneTracks/Private/Sections/MovieSceneCameraShakeSourceShakeSection.cpp
@@ -1,7 +1,7 @@
 // Copyright Epic Games, Inc. All Rights Reserved.
 
 #include "Sections/MovieSceneCameraShakeSourceShakeSection.h"
-#include "Tracks/MovieSceneCameraShakeSourceShakeTrack.h"
+#include "MovieSceneTracksComponentTypes.h"
 
 #include UE_INLINE_GENERATED_CPP_BY_NAME(MovieSceneCameraShakeSourceShakeSection)
 
@@ -11,3 +11,20 @@ UMovieSceneCameraShakeSourceShakeSection::UMovieSceneCameraShakeSourceShakeSecti
 	EvalOptions.EnableAndSetCompletionMode(EMovieSceneCompletionMode::ProjectDefault);
 }
 
+void UMovieSceneCameraShakeSourceShakeSection::ImportEntityImpl(UMovieSceneEntitySystemLinker* EntityLinker, const FEntityImportParams& Params, FImportedEntity* OutImportedEntity)
+{
+	using namespace UE::MovieScene;
+
+	const FBuiltInComponentTypes* BuiltInComponents = FBuiltInComponentTypes::Get();
+	const FMovieSceneTracksComponentTypes* TrackComponents = FMovieSceneTracksComponentTypes::Get();
+
+	const FGuid ObjectBindingID = Params.GetObjectBindingID();
+	FMovieSceneCameraShakeComponentData ComponentData(ShakeData, GetRange());
+
+	OutImportedEntity->AddBuilder(
+		FEntityBuilder()
+		.AddConditional(BuiltInComponents->GenericObjectBinding, ObjectBindingID, ObjectBindingID.IsValid())
+		.Add(TrackComponents->CameraShake, ComponentData)
+	);
+}
+
diff --git a/Engine/Source/Runtime/MovieSceneTracks/Private/Sections/MovieSceneCameraShakeSourceTriggerSection.cpp b/Engine/Source/Runtime/MovieSceneTracks/Private/Sections/MovieSceneCameraShakeSourceTriggerSection.cpp
--- a/Engine/Source/Runtime/MovieSceneTracks/Private/Sections/MovieSceneCameraShakeSourceTriggerSection.cpp
+++ b/Engine/Source/Runtime/MovieSceneTracks/Private/Sections/MovieSceneCameraShakeSourceTriggerSection.cpp
@@ -2,6 +2,10 @@
 
 #include "Sections/MovieSceneCameraShakeSourceTriggerSection.h"
 #include "Channels/MovieSceneChannelProxy.h"
+#include "EntitySystem/MovieSceneEntityBuilder.h"
+#include "EntitySystem/MovieSceneEntitySystemLinker.h"
+#include "Evaluation/MovieSceneEvaluationField.h"
+#include "Systems/MovieSceneCameraShakeSystem.h"
 
 #include UE_INLINE_GENERATED_CPP_BY_NAME(MovieSceneCameraShakeSourceTriggerSection)
 
@@ -15,4 +19,50 @@ UMovieSceneCameraShakeSourceTriggerSection::UMovieSceneCameraShakeSourceTriggerS
 #endif
 }
 
+bool UMovieSceneCameraShakeSourceTriggerSection::PopulateEvaluationFieldImpl(const TRange<FFrameNumber>& EffectiveRange, const FMovieSceneEvaluationFieldEntityMetaData& InMetaData, FMovieSceneEntityComponentFieldBuilder* OutFieldBuilder)
+{
+	const int32 MetaDataIndex = OutFieldBuilder->AddMetaData(InMetaData);
+
+	TArrayView<const FFrameNumber> Times = Channel.GetData().GetTimes();
+	for (int32 Index = 0; Index < Times.Num(); ++Index)
+	{
+		if (EffectiveRange.Contains(Times[Index]))
+		{
+			TRange<FFrameNumber> Range(Times[Index]);
+			OutFieldBuilder->AddOneShotEntity(Range, this, Index, MetaDataIndex);
+		}
+	}
+
+	return true;
+}
+
+void UMovieSceneCameraShakeSourceTriggerSection::ImportEntityImpl(UMovieSceneEntitySystemLinker* EntityLinker, const FEntityImportParams& Params, FImportedEntity* OutImportedEntity)
+{
+	using namespace UE::MovieScene;
+
+	const int32 ShakeIndex = static_cast<int32>(Params.EntityID);
+
+	TArrayView<const FFrameNumber> Times  = Channel.GetData().GetTimes();
+	TArrayView<const FMovieSceneCameraShakeSourceTrigger> Shakes = Channel.GetData().GetValues();
+	if (!ensureMsgf(
+			Shakes.IsValidIndex(ShakeIndex), 
+			TEXT("Attempting to import a camera shake entity for an invalid index (Index: %d, Num: %d)"), ShakeIndex, Shakes.Num()))
+	{
+		return;
+	}
+
+	UMovieSceneCameraShakeInstantiatorSystem* ShakeSystem = EntityLinker->LinkSystem<UMovieSceneCameraShakeInstantiatorSystem>();
+	if (ensure(ShakeSystem))
+	{
+		const FSequenceInstance& ThisInstance = EntityLinker->GetInstanceRegistry()->GetInstance(Params.Sequence.InstanceHandle);
+		const FMovieSceneContext& Context = ThisInstance.GetContext();
+		
+		const FFrameTime ShakeTime = Times[ShakeIndex] * Context.GetSequenceToRootSequenceTransform();
+
+		ShakeSystem->AddShakeTrigger(Params.Sequence.InstanceHandle, Params.GetObjectBindingID(), ShakeTime, Shakes[ShakeIndex]);
+
+		// Mimic the structure changing in order to ensure that the instantiation phase runs
+		EntityLinker->EntityManager.MimicStructureChanged();
+	}
+}
 
diff --git a/Engine/Source/Runtime/MovieSceneTracks/Private/Systems/MovieSceneCameraShakeSystem.cpp b/Engine/Source/Runtime/MovieSceneTracks/Private/Systems/MovieSceneCameraShakeSystem.cpp
new file mode 100644
--- /dev/null
+++ b/Engine/Source/Runtime/MovieSceneTracks/Private/Systems/MovieSceneCameraShakeSystem.cpp
@@ -0,0 +1,766 @@
+// Copyright Epic Games, Inc. All Rights Reserved.
+
+#include "Systems/MovieSceneCameraShakeSystem.h"
+
+#include "Camera/CameraComponent.h"
+#include "Camera/CameraModifier_CameraShake.h"
+#include "Camera/CameraShakeSourceComponent.h"
+#include "EntitySystem/BuiltInComponentTypes.h"
+#include "EntitySystem/MovieSceneEntityManager.h"
+#include "EntitySystem/MovieSceneEntitySystemLinker.h"
+#include "Evaluation/PreAnimatedState/MovieScenePreAnimatedStateStorage.h"
+#include "Evaluation/PreAnimatedState/MovieScenePreAnimatedStorageID.inl"
+#include "Evaluation/PreAnimatedState/MovieScenePreAnimatedObjectStorage.h"
+#include "Evaluation/MovieSceneCameraShakePreviewer.h"
+#include "IMovieScenePlayer.h"
+#include "MovieSceneTracksComponentTypes.h"
+#include "Sections/MovieSceneCameraShakeSection.h"
+#include "UObject/Object.h"
+#include "UObject/Package.h"
+
+#include UE_INLINE_GENERATED_CPP_BY_NAME(MovieSceneCameraShakeSystem)
+
+namespace UE::MovieScene
+{
+
+struct FPreAnimatedCameraShakeTraits : FBoundObjectPreAnimatedStateTraits
+{
+	using KeyType = FObjectKey;
+	using StorageType = bool;
+
+	bool CachePreAnimatedValue(UObject* InKey)
+	{
+		return true;
+	}
+
+	void RestorePreAnimatedValue(const FObjectKey& InKey, const bool Unused, const FRestoreStateParams& Params)
+	{
+		if (UCameraShakeBase* CameraShake = Cast<UCameraShakeBase>(InKey.ResolveObjectPtr()))
+		{
+			if (!CameraShake->IsFinished())
+			{
+				CameraShake->StopShake(true);
+			}
+			CameraShake->TeardownShake();
+		}
+	}
+};
+
+struct FPreAnimatedCameraComponentShakeTraits : FBoundObjectPreAnimatedStateTraits
+{
+	using KeyType = FObjectKey;
+	using StorageType = bool;
+
+	bool CachePreAnimatedValue(UObject* InKey)
+	{
+		return true;
+	}
+
+	void RestorePreAnimatedValue(const FObjectKey& InKey, const bool Unused, const FRestoreStateParams& Params)
+	{
+		if (UCameraComponent* CameraComponent = Cast<UCameraComponent>(InKey.ResolveObjectPtr()))
+		{
+			CameraComponent->ClearAdditiveOffset();
+			CameraComponent->ClearExtraPostProcessBlends();
+		}
+	}
+};
+
+struct FPreAnimatedCameraSourceShakeTraits : FBoundObjectPreAnimatedStateTraits
+{
+	using KeyType = FObjectKey;
+	using StorageType = bool;
+
+	bool CachePreAnimatedValue(UObject* InKey)
+	{
+		return true;
+	}
+
+	void RestorePreAnimatedValue(const FObjectKey& InKey, const bool Unused, const FRestoreStateParams& Params)
+	{
+		if (UCameraShakeSourceComponent* ShakeSourceComponent = Cast<UCameraShakeSourceComponent>(InKey.ResolveObjectPtr()))
+		{
+			ShakeSourceComponent->StopAllCameraShakes(true);
+		}
+	}
+};
+
+struct FPreAnimatedCameraShakeStateStorage : TPreAnimatedStateStorage_ObjectTraits<FPreAnimatedCameraShakeTraits>
+{
+	static TAutoRegisterPreAnimatedStorageID<FPreAnimatedCameraShakeStateStorage> StorageID;
+	FPreAnimatedStorageID GetStorageType() const override { return StorageID; }
+};
+
+TAutoRegisterPreAnimatedStorageID<FPreAnimatedCameraShakeStateStorage> FPreAnimatedCameraShakeStateStorage::StorageID;
+
+struct FPreAnimatedCameraComponentShakeStateStorage : TPreAnimatedStateStorage_ObjectTraits<FPreAnimatedCameraComponentShakeTraits>
+{
+	static TAutoRegisterPreAnimatedStorageID<FPreAnimatedCameraComponentShakeStateStorage> StorageID;
+	FPreAnimatedStorageID GetStorageType() const override { return StorageID; }
+};
+
+TAutoRegisterPreAnimatedStorageID<FPreAnimatedCameraComponentShakeStateStorage> FPreAnimatedCameraComponentShakeStateStorage::StorageID;
+
+struct FPreAnimatedCameraSourceShakeStateStorage : TPreAnimatedStateStorage_ObjectTraits<FPreAnimatedCameraSourceShakeTraits>
+{
+	static TAutoRegisterPreAnimatedStorageID<FPreAnimatedCameraSourceShakeStateStorage> StorageID;
+	FPreAnimatedStorageID GetStorageType() const override { return StorageID; }
+};
+
+TAutoRegisterPreAnimatedStorageID<FPreAnimatedCameraSourceShakeStateStorage> FPreAnimatedCameraSourceShakeStateStorage::StorageID;
+
+#if WITH_EDITOR
+
+TEntitySystemLinkerExtensionID<FCameraShakePreviewerLinkerExtension> FCameraShakePreviewerLinkerExtension::GetExtensionID()
+{
+	static TEntitySystemLinkerExtensionID<FCameraShakePreviewerLinkerExtension> ID = UMovieSceneEntitySystemLinker::RegisterExtension<FCameraShakePreviewerLinkerExtension>();
+	return ID;
+}
+
+TSharedPtr<FCameraShakePreviewerLinkerExtension> FCameraShakePreviewerLinkerExtension::GetOrCreateExtension(UMovieSceneEntitySystemLinker* Linker)
+{
+	if (FCameraShakePreviewerLinkerExtension* PreviewerExtension = Linker->FindExtension<FCameraShakePreviewerLinkerExtension>())
+	{
+		return PreviewerExtension->AsShared();
+	}
+
+	TSharedPtr<FCameraShakePreviewerLinkerExtension> NewPreviewerExtension = MakeShared<FCameraShakePreviewerLinkerExtension>(Linker);
+	Linker->AddExtension(NewPreviewerExtension.Get());
+	return NewPreviewerExtension;
+}
+
+FCameraShakePreviewerLinkerExtension::FCameraShakePreviewerLinkerExtension(UMovieSceneEntitySystemLinker* Linker)
+	: TSharedEntitySystemLinkerExtension(Linker)
+{
+}
+
+FCameraShakePreviewerLinkerExtension::~FCameraShakePreviewerLinkerExtension()
+{
+	for (TPair<FInstanceHandle, FCameraShakePreviewer>& Pair : Previewers)
+	{
+		FCameraShakePreviewer& Previewer = Pair.Value;
+		if (ensure(Previewer.IsInitialized()))
+		{
+			Previewer.UnRegisterViewModifier();
+			Previewer.Teardown();
+		}
+	}
+	Previewers.Reset();
+}
+
+FCameraShakePreviewer* FCameraShakePreviewerLinkerExtension::FindPreviewer(FInstanceHandle InstanceHandle)
+{
+	return Previewers.Find(InstanceHandle);
+}
+
+FCameraShakePreviewer& FCameraShakePreviewerLinkerExtension::GetPreviewer(FInstanceHandle InstanceHandle)
+{
+	if (FCameraShakePreviewer* Previewer = Previewers.Find(InstanceHandle))
+	{
+		ensure(Previewer->IsInitialized());
+		return *Previewer;
+	}
+
+	FCameraShakePreviewer& NewPreviewer = Previewers.Add(InstanceHandle);
+	NewPreviewer.Initialize(WeakLinker.Get()->GetWorld());
+	NewPreviewer.RegisterViewModifier();
+	return NewPreviewer;
+}
+
+void FCameraShakePreviewerLinkerExtension::UpdateAllPreviewers(FInstanceRegistry* InstanceRegistry)
+{
+	const TSparseArray<FSequenceInstance>& Instances = InstanceRegistry->GetSparseInstances();
+	for (auto It = Instances.CreateConstIterator(); It; ++It)
+	{
+		FInstanceHandle InstanceHandle = It->GetInstanceHandle();
+		if (FCameraShakePreviewer* Previewer = Previewers.Find(InstanceHandle))
+		{
+			const FMovieSceneContext& Context = It->GetContext();
+			if (Context.HasJumped())
+			{
+				Previewer->GetCameraModifier()->RemoveAllCameraShakes(true);
+			}
+			else
+			{
+				const float DeltaTime = Context.GetFrameRate().AsSeconds(Context.GetDelta());
+				const bool bIsPlaying = Context.GetStatus() == EMovieScenePlayerStatus::Playing;
+				Previewer->Update(DeltaTime, bIsPlaying);
+			}
+		}
+	}
+}
+
+bool FCameraShakePreviewerLinkerExtension::HasAnyShake() const
+{
+	TArray<FActiveCameraShakeInfo> TempCameraShakes;
+	for (const TPair<FInstanceHandle, FCameraShakePreviewer>& Pair : Previewers)
+	{
+		const UCameraModifier_CameraShake* CameraModifier = Pair.Value.GetCameraModifier();
+		CameraModifier->GetActiveCameraShakes(TempCameraShakes);
+		if (TempCameraShakes.Num() > 0)
+		{
+			return true;
+		}
+	}
+	return false;
+}
+
+#endif  // WITH_EDITOR
+
+} // namespace UE::MovieScene
+
+UMovieSceneCameraShakeInstantiatorSystem::UMovieSceneCameraShakeInstantiatorSystem(const FObjectInitializer& ObjInit)
+	: Super(ObjInit)
+{
+	using namespace UE::MovieScene;
+
+	RelevantComponent = FMovieSceneTracksComponentTypes::Get()->CameraShake;
+	Phase = ESystemPhase::Instantiation;
+
+	if (HasAnyFlags(RF_ClassDefaultObject))
+	{
+		const FBuiltInComponentTypes* BuiltInComponents = FBuiltInComponentTypes::Get();
+		DefineComponentConsumer(GetClass(), BuiltInComponents->BoundObject);
+	}
+}
+
+bool UMovieSceneCameraShakeInstantiatorSystem::IsRelevantImpl(UMovieSceneEntitySystemLinker* InLinker) const
+{
+	return TriggersByInstance.Num() > 0;
+}
+
+void UMovieSceneCameraShakeInstantiatorSystem::OnLink()
+{
+	using namespace UE::MovieScene;
+
+	PreAnimatedCameraShakeStorage = Linker->PreAnimatedState.GetOrCreateStorage<FPreAnimatedCameraShakeStateStorage>();
+	PreAnimatedCameraComponentShakeStorage = Linker->PreAnimatedState.GetOrCreateStorage<FPreAnimatedCameraComponentShakeStateStorage>();
+	PreAnimatedCameraSourceShakeStorage = Linker->PreAnimatedState.GetOrCreateStorage<FPreAnimatedCameraSourceShakeStateStorage>();
+
+#if WITH_EDITOR
+	PreviewerExtension = FCameraShakePreviewerLinkerExtension::GetOrCreateExtension(Linker);
+#endif
+}
+
+void UMovieSceneCameraShakeInstantiatorSystem::OnUnlink()
+{
+	using namespace UE::MovieScene;
+
+#if WITH_EDITOR
+	// Only the two camera shake systems hold pointers to the extension, so it should delete itself
+	// once both systems are unlinked.
+	PreviewerExtension = nullptr;
+#endif
+
+	if (!ensure(TriggersByInstance.Num() == 0))
+	{
+		TriggersByInstance.Reset();
+	}
+}
+
+void UMovieSceneCameraShakeInstantiatorSystem::OnRun(FSystemTaskPrerequisites& InPrerequisites, FSystemSubsequentTasks& Subsequents)
+{
+	using namespace UE::MovieScene;
+
+	const FBuiltInComponentTypes* BuiltInComponents = FBuiltInComponentTypes::Get();
+	const FMovieSceneTracksComponentTypes* TrackComponents = FMovieSceneTracksComponentTypes::Get();
+
+	FEntityManager& EntityManager = Linker->EntityManager;
+	FInstanceRegistry* InstanceRegistry = Linker->GetInstanceRegistry();
+
+	// Create camera shake instances for new shakes, and start them.
+	auto VisitNewShakes = [this, BuiltInComponents, &EntityManager, InstanceRegistry](
+			FMovieSceneEntityID EntityID,
+			FInstanceHandle InstanceHandle, 
+			UObject* BoundObject, 
+			FMovieSceneCameraShakeComponentData& ShakeData)
+	{
+		const FSequenceInstance& Instance = InstanceRegistry->GetInstance(InstanceHandle);
+		const FMovieSceneContext& Context = Instance.GetContext();
+		IMovieScenePlayer* Player = Instance.GetPlayer();
+
+		TSubclassOf<UCameraShakeBase> ShakeClass = ShakeData.SectionData.ShakeClass;
+		UCameraShakeSourceComponent* ShakeSourceComponent = Cast<UCameraShakeSourceComponent>(BoundObject);
+		
+		if (ShakeClass.Get() == nullptr)
+		{
+			if (ShakeSourceComponent)
+			{
+				ShakeClass = ShakeSourceComponent->CameraShake;
+			}
+		}
+		if (ShakeClass.Get() == nullptr)
+		{
+			return;
+		}
+
+		// Get the duration of the shake and store it in the component data.
+		FCameraShakeDuration ShakeDuration;
+		UCameraShakeBase::GetCameraShakeDuration(ShakeClass, ShakeDuration);
+		ShakeData.Duration = ShakeDuration.IsFixed() ? 
+			Context.GetFrameRate().AsFrameTime(ShakeDuration.Get()) :
+			FFrameTime(-1);
+
+		// Get the blend out duration and also store it in the instance data.
+		float BlendIn = 0.f, BlendOut = 0.f;
+		UCameraShakeBase::GetCameraShakeBlendTimes(ShakeClass, BlendIn, BlendOut);
+		ShakeData.BlendOutTime = (BlendOut > 0) ?
+			Context.GetFrameRate().AsFrameTime(BlendOut) :
+			FFrameTime(0);
+
+		const bool bWantsRestoreState = EntityManager.HasComponent(EntityID, BuiltInComponents->Tags.RestoreState);
+		const FRootInstanceHandle RootInstanceHandle = Instance.GetRootInstanceHandle();
+		FCachePreAnimatedValueParams CacheParams;
+
+		// Start playing the shake.
+		if (ShakeSourceComponent)
+		{
+			PreAnimatedCameraSourceShakeStorage->BeginTrackingEntity(EntityID, bWantsRestoreState, RootInstanceHandle, ShakeSourceComponent);
+			PreAnimatedCameraSourceShakeStorage->CachePreAnimatedValue(CacheParams, ShakeSourceComponent);
+
+			ShakeSourceComponent->StartCameraShake(
+					ShakeClass, 
+					ShakeData.SectionData.PlayScale, 
+					ShakeData.SectionData.PlaySpace,
+					ShakeData.SectionData.UserDefinedPlaySpace);
+
+#if WITH_EDITOR
+			// Shake source components start shakes in the world, unlike the other shakes
+			// (in the `else` clause) who directly affect the bound camera. This means that
+			// the shake we have just started won't affect the Sequencer preview unless we
+			// add some shaking ourselves. Let's do that here.
+			FCameraShakePreviewer& Previewer = PreviewerExtension->GetPreviewer(InstanceHandle);
+			UCameraModifier_CameraShake* const PreviewCameraShake = Previewer.GetCameraModifier();
+			
+			FAddCameraShakeParams Params;
+			Params.SourceComponent = ShakeSourceComponent;
+			Params.Scale = ShakeData.SectionData.PlayScale;
+			Params.PlaySpace = ShakeData.SectionData.PlaySpace;
+			Params.UserPlaySpaceRot = ShakeData.SectionData.UserDefinedPlaySpace;
+			PreviewCameraShake->AddCameraShake(ShakeClass, Params);
+#endif
+		}
+		else if (UCameraComponent* CameraComponent = MovieSceneHelpers::CameraComponentFromRuntimeObject(BoundObject))
+		{
+			UObject* OuterObject = Player->GetPlaybackContext() ? Player->GetPlaybackContext() : GetTransientPackage();
+			UCameraShakeBase* ShakeInstance = NewObject<UCameraShakeBase>(OuterObject, ShakeClass);
+
+			PreAnimatedCameraShakeStorage->BeginTrackingEntity(EntityID, bWantsRestoreState, RootInstanceHandle, ShakeInstance);
+			PreAnimatedCameraShakeStorage->CachePreAnimatedValue(CacheParams, ShakeInstance);
+
+			PreAnimatedCameraComponentShakeStorage->BeginTrackingEntity(EntityID, bWantsRestoreState, RootInstanceHandle, CameraComponent);
+			PreAnimatedCameraComponentShakeStorage->CachePreAnimatedValue(CacheParams, CameraComponent);
+
+			ShakeInstance->StartShake(
+					nullptr, 
+					ShakeData.SectionData.PlayScale,
+					ShakeData.SectionData.PlaySpace,
+					ShakeData.SectionData.UserDefinedPlaySpace);
+
+			ShakeData.ShakeInstance = ShakeInstance;
+		}
+
+		ShakeData.Status = EMovieSceneCameraShakeStatus::Started;
+	};
+
+	FEntityTaskBuilder()
+	.ReadEntityIDs()
+	.Read(BuiltInComponents->InstanceHandle)
+	.Read(BuiltInComponents->BoundObject)
+	.Write(TrackComponents->CameraShake)
+	.FilterAll({ BuiltInComponents->Tags.NeedsLink })
+	.Iterate_PerEntity(&Linker->EntityManager, VisitNewShakes);
+
+	// We don't need to go over expired shakes from NeedsUnlink entities. Either these shakes
+	// will be stopped be the pre-animated state, and then their component data will be freed
+	// when the entities are deleted, or they come from KeepState sections and will therefore
+	// continue running inside the camera manager.
+
+	// Now trigger any one-shot shakes.
+	if (TriggersByInstance.Num() > 0)
+	{
+		TriggerOneShotShakes();
+	}
+}
+
+void UMovieSceneCameraShakeInstantiatorSystem::AddShakeTrigger(UE::MovieScene::FInstanceHandle InInstance, const FGuid& ObjectBindingID, const FFrameTime& InTime, const FMovieSceneCameraShakeSourceTrigger& InTrigger)
+{
+	TriggersByInstance.FindOrAdd(InInstance).Add(FTimedTrigger{ ObjectBindingID, InTime, InTrigger });
+}
+
+void UMovieSceneCameraShakeInstantiatorSystem::TriggerOneShotShakes()
+{
+	using namespace UE::MovieScene;
+
+	FInstanceRegistry* InstanceRegistry = Linker->GetInstanceRegistry();
+
+	for (TPair<FInstanceHandle, TArray<FTimedTrigger>>& Pair : TriggersByInstance)
+	{
+		const FSequenceInstance& Instance = InstanceRegistry->GetInstance(Pair.Key);
+		IMovieScenePlayer* Player = Instance.GetPlayer();
+
+		for (const FTimedTrigger& Trigger : Pair.Value)
+		{
+			TArray<UCameraShakeSourceComponent*> ShakeSourceComponents;
+			if (Trigger.ObjectBindingID.IsValid())
+			{
+				for (TWeakObjectPtr<> WeakBoundObject : Player->FindBoundObjects(Trigger.ObjectBindingID, Instance.GetSequenceID()))
+				{
+					if (UObject* BoundObject = WeakBoundObject.Get())
+					{
+						if (UCameraShakeSourceComponent* ShakeSourceComponent = Cast<UCameraShakeSourceComponent>(BoundObject))
+						{
+							ShakeSourceComponents.Add(ShakeSourceComponent);
+						}
+					}
+				}
+			}
+
+			for (UCameraShakeSourceComponent* ShakeSourceComponent : ShakeSourceComponents)
+			{
+				TSubclassOf<UCameraShakeBase> ShakeClass = Trigger.Trigger.ShakeClass;
+				if (ShakeClass.Get() == nullptr)
+				{
+					ShakeClass = ShakeSourceComponent->CameraShake;
+				}
+
+				if (ShakeClass.Get() != nullptr)
+				{
+					// Start playing the shake.
+					ShakeSourceComponent->StartCameraShake(
+							ShakeClass, 
+							Trigger.Trigger.PlayScale, 
+							Trigger.Trigger.PlaySpace, 
+							Trigger.Trigger.UserDefinedPlaySpace);
+
+#if WITH_EDITOR
+					// Also start playing the shake in our editor preview.
+					FCameraShakePreviewer& Previewer = PreviewerExtension->GetPreviewer(Pair.Key);
+					UCameraModifier_CameraShake* const PreviewCameraShake = Previewer.GetCameraModifier();
+
+					FAddCameraShakeParams Params;
+					Params.SourceComponent = ShakeSourceComponent;
+					Params.Scale = Trigger.Trigger.PlayScale;
+					Params.PlaySpace = Trigger.Trigger.PlaySpace;
+					Params.UserPlaySpaceRot = Trigger.Trigger.UserDefinedPlaySpace;
+					PreviewCameraShake->AddCameraShake(ShakeClass, Params);
+#endif
+				}
+			}
+		}
+	}
+}
+
+namespace UE::MovieScene
+{
+
+struct FAccumulatedShake
+{
+	void AccumulateOffset(const FTransform& InTransformOffset, float InFOVOffset)
+	{
+		TotalTransformOffset = TotalTransformOffset * InTransformOffset;
+		TotalFOVOffset += InFOVOffset;
+		bApplyTransform = true;
+	}
+
+	void AccumulatePostProcessing(const FPostProcessSettings& InPostProcessSettings, float InWeight)
+	{
+		PostProcessSettings.Add({ InPostProcessSettings, InWeight });
+		bApplyPostProcessing = true;
+	}
+
+	void Apply(UCameraComponent* CameraComponent) const
+	{
+		if (bApplyTransform)
+		{
+			CameraComponent->ClearAdditiveOffset();
+			CameraComponent->AddAdditiveOffset(TotalTransformOffset, TotalFOVOffset);
+		}
+
+		if (bApplyPostProcessing)
+		{
+			CameraComponent->ClearExtraPostProcessBlends();
+			for (const TPair<FPostProcessSettings, float>& Pair : PostProcessSettings)
+			{
+				CameraComponent->AddExtraPostProcessBlend(Pair.Key, Pair.Value);
+			}
+		}
+	}
+
+private:
+	bool bApplyTransform = false;
+	bool bApplyPostProcessing = false;
+
+	FTransform TotalTransformOffset;
+	float TotalFOVOffset = 0.f;
+	TArray<TTuple<FPostProcessSettings, float>, TInlineAllocator<2>> PostProcessSettings;
+};
+
+struct FEvaluateCameraShake
+{
+	const FInstanceRegistry* InstanceRegistry;
+#if WITH_EDITOR
+	FCameraShakePreviewerLinkerExtension* PreviewerExtension;
+#endif
+
+	TMap<UCameraComponent*, FAccumulatedShake> AccumulatedShakes;
+
+	FEvaluateCameraShake(UMovieSceneEntitySystemLinker* InLinker)
+		: InstanceRegistry(InLinker->GetInstanceRegistry())
+	{
+#if WITH_EDITOR
+		PreviewerExtension = InLinker->FindExtension<FCameraShakePreviewerLinkerExtension>();
+#endif
+	}
+
+	void ForEachAllocation(const FEntityAllocation* Allocation, TRead<FMovieSceneEntityID> EntityIDs, TRead<FRootInstanceHandle> RootInstanceHandles, TRead<FInstanceHandle> InstanceHandles, TRead<UObject*> BoundObjects, TWrite<FMovieSceneCameraShakeComponentData> ShakeComponents)
+	{
+		const int32 Num = Allocation->Num();
+		for (int32 Index = 0; Index < Num; ++Index)
+		{
+			FRootInstanceHandle InstanceHandle = RootInstanceHandles[Index];
+			const FSequenceInstance& Instance = InstanceRegistry->GetInstance(InstanceHandle);
+			const FMovieSceneContext& Context = Instance.GetContext();
+			IMovieScenePlayer* Player = Instance.GetPlayer();
+
+			// If we have jumped, we could have ended up anywhere, so we need to start fresh without
+			// any running camera shake previews.
+#if WITH_EDITOR
+			if (Context.HasJumped())
+			{
+				if (FCameraShakePreviewer* Previewer = PreviewerExtension->FindPreviewer(InstanceHandle))
+				{
+					Previewer->GetCameraModifier()->RemoveAllCameraShakes();
+				}
+			}
+#endif
+
+			// Shakes should have been started by the instantiator system.
+			FMovieSceneCameraShakeComponentData& ShakeData = ShakeComponents[Index];
+			ensure(ShakeData.Status != EMovieSceneCameraShakeStatus::NotStarted);
+
+			// Let's see what this camera shake should be doing now.
+			EMovieSceneCameraShakeStatus DesiredNewStatus = EMovieSceneCameraShakeStatus::Started;
+			
+			// See if there is any blend-out and/or end time that we need to watch out for.
+			// There can be 3 situations here:
+			//
+			//   1. The shake has a duration, and our original section is long enough that the shake will
+			//		blend out and finish naturally on its own. In this case we have nothing to do 
+			//		except update our internal status.
+			//	 2. The shake has a duration, but our original section's size is cutting this short. We 
+			//	    will want to start making the shake blend out manually, or stop it abruptely if
+			//	    it doesn't have any blend-out time.
+			//	 3. The shake has no duration, so we need to make it blend out manually near the end of
+			//	    our source section, or end it abruptely at the end of our source section if it 
+			//	    doesn't have any blend-out time.
+			//
+			// Cases 2 and 3 require us to blend-out/stop the shake ourselves.
+			// Let's see which case we are in.
+			const bool bHasDuration = ShakeData.Duration > 0;
+			const bool bHasBlendOut = ShakeData.BlendOutTime > 0;
+
+			const bool bNeedsManualStop = 
+				// Case 3
+				!bHasDuration ||
+				// Case 2
+				(ShakeData.SectionStartTime + ShakeData.Duration) > ShakeData.SectionEndTime;
+
+			if (bNeedsManualStop)
+			{
+				// Let's see if we have reached the time when we need to start blending out, or the
+				// time we need to flat out finish.
+				if (bHasDuration && Context.GetTime() >= ShakeData.SectionStartTime + ShakeData.Duration)
+				{
+					DesiredNewStatus = EMovieSceneCameraShakeStatus::Finished;
+				}
+				else if (bHasBlendOut && Context.GetTime() >= ShakeData.SectionEndTime - ShakeData.BlendOutTime)
+				{
+					DesiredNewStatus = EMovieSceneCameraShakeStatus::BlendingOut;
+				}
+			}
+
+			if (UCameraShakeSourceComponent* ShakeSourceComponent = Cast<UCameraShakeSourceComponent>(BoundObjects[Index]))
+			{
+				EvaluateShakeSourceComponentShake(ShakeSourceComponent, ShakeData, Instance, DesiredNewStatus);
+			}
+			else if (UCameraComponent* CameraComponent = MovieSceneHelpers::CameraComponentFromRuntimeObject(BoundObjects[Index]))
+			{
+				EvaluateCameraComponentShake(CameraComponent, ShakeData, Instance, DesiredNewStatus);
+			}
+			ShakeData.Status = DesiredNewStatus;
+		}
+	}
+
+	void EvaluateShakeSourceComponentShake(
+			UCameraShakeSourceComponent* ShakeSourceComponent,
+			FMovieSceneCameraShakeComponentData& ShakeData,
+			const FSequenceInstance& Instance,
+			EMovieSceneCameraShakeStatus DesiredNewStatus)
+	{
+		// We don't need to evaluate the source component's camera shakes here, as they are ticking
+		// along by themselves in both the player camera manager (in the game) and the camera shake
+		// previewer (in the editor).
+		// All we need to do is intervene to stop/blend-out the shake if needed.
+
+		if (ShakeData.Status == EMovieSceneCameraShakeStatus::Started &&
+			(DesiredNewStatus == EMovieSceneCameraShakeStatus::BlendingOut ||
+			 DesiredNewStatus == EMovieSceneCameraShakeStatus::Finished))
+		{
+			const bool bImmediately = (DesiredNewStatus == EMovieSceneCameraShakeStatus::Finished);
+
+			// TODO-ludovic: this isn't exactly correct...
+			// We could be stopping other shakes of the same type started by other means.. but doing
+			// the correct thing would require storing multiple weak shake instance pointers mapped 
+			// to multiple player controllers, themselves mapped to multiple bound objects. Let's only
+			// do that if we run into the (quite unlikely) case where we need it?
+			ShakeSourceComponent->StopAllCameraShakesOfType(ShakeData.SectionData.ShakeClass, bImmediately);
+
+#if WITH_EDITOR
+			FCameraShakePreviewer& Previewer = PreviewerExtension->GetPreviewer(Instance.GetInstanceHandle());
+			UCameraModifier_CameraShake* const PreviewCameraShake = Previewer.GetCameraModifier();
+			PreviewCameraShake->RemoveAllCameraShakesOfClassFromSource(
+				ShakeData.SectionData.ShakeClass, ShakeSourceComponent, bImmediately);
+#endif
+		}
+	}
+
+	void EvaluateCameraComponentShake(
+			UCameraComponent* CameraComponent,
+			FMovieSceneCameraShakeComponentData& ShakeData,
+			const FSequenceInstance& Instance,
+			EMovieSceneCameraShakeStatus DesiredNewStatus)
+	{
+		FMinimalViewInfo POV;
+		POV.Location = CameraComponent->GetComponentLocation();
+		POV.Rotation = CameraComponent->GetComponentRotation();
+		POV.FOV = CameraComponent->FieldOfView;
+
+		float PostProcessBlendWeight = 0.f;
+		FPostProcessSettings PostProcessSettings;
+
+		// Update shake to the new time.
+		const FMovieSceneContext& Context = Instance.GetContext();
+		const FFrameTime NewShakeTime = Context.GetTime() - ShakeData.SectionStartTime;
+		ShakeData.ShakeInstance->ScrubAndApplyCameraShake(NewShakeTime / Context.GetFrameRate(), 1.f, POV);
+
+		// Grab transform and FOV changes.
+		FTransform WorldToBaseCamera = CameraComponent->GetComponentToWorld().Inverse();
+		float BaseFOV = CameraComponent->FieldOfView;
+		FTransform NewCameraToWorld(POV.Rotation, POV.Location);
+		float NewFOV = POV.FOV;
+
+		FTransform NewCameraToBaseCamera = NewCameraToWorld * WorldToBaseCamera;
+
+		float NewFOVToBaseFOV = BaseFOV - NewFOV;
+
+		{
+			// Accumumulate the offsets into the track data for application as part of the track execution token
+			FAccumulatedShake& AccumulatedShake = AccumulatedShakes.FindOrAdd(CameraComponent);
+			AccumulatedShake.AccumulateOffset(NewCameraToBaseCamera, NewFOVToBaseFOV);
+		}
+
+		// Grab post process changes.
+		if (PostProcessBlendWeight > 0.f)
+		{
+			FAccumulatedShake& AccumulatedShake = AccumulatedShakes.FindOrAdd(CameraComponent);
+			AccumulatedShake.AccumulatePostProcessing(PostProcessSettings, PostProcessBlendWeight);
+		}
+	}
+
+	void PostTask()
+	{
+		// Apply accumulated shakes.
+		for (const TPair<UCameraComponent*, FAccumulatedShake>& Pair : AccumulatedShakes)
+		{
+			Pair.Value.Apply(Pair.Key);
+		}
+	}
+};
+
+} // namespace UE::MovieScene
+
+UMovieSceneCameraShakeEvaluatorSystem::UMovieSceneCameraShakeEvaluatorSystem(const FObjectInitializer& ObjInit)
+	: Super(ObjInit)
+{
+	using namespace UE::MovieScene;
+
+	RelevantComponent = FMovieSceneTracksComponentTypes::Get()->CameraShake;
+
+	if (HasAnyFlags(RF_ClassDefaultObject))
+	{
+	}
+}
+
+void UMovieSceneCameraShakeEvaluatorSystem::OnRun(FSystemTaskPrerequisites& InPrerequisites, FSystemSubsequentTasks& Subsequents)
+{
+	using namespace UE::MovieScene;
+
+	const FBuiltInComponentTypes* BuiltInComponents = FBuiltInComponentTypes::Get();
+	const FMovieSceneTracksComponentTypes* TrackComponents = FMovieSceneTracksComponentTypes::Get();
+
+	FEntityTaskBuilder()
+	.ReadEntityIDs()
+	.Read(BuiltInComponents->RootInstanceHandle)
+	.Read(BuiltInComponents->InstanceHandle)
+	.Read(BuiltInComponents->BoundObject)
+	.Write(TrackComponents->CameraShake)
+	.SetDesiredThread(Linker->EntityManager.GetGatherThread())
+	.Dispatch_PerAllocation<FEvaluateCameraShake>(
+		&Linker->EntityManager, InPrerequisites, &Subsequents, Linker);
+}
+
+UMovieSceneCameraShakePreviewerEvaluatorSystem::UMovieSceneCameraShakePreviewerEvaluatorSystem(const FObjectInitializer& ObjInit)
+	: Super(ObjInit)
+{
+	using namespace UE::MovieScene;
+
+	if (HasAnyFlags(RF_ClassDefaultObject))
+	{
+		DefineImplicitPrerequisite(UMovieSceneCameraShakeEvaluatorSystem::StaticClass(), GetClass());
+	}
+}
+
+bool UMovieSceneCameraShakePreviewerEvaluatorSystem::IsRelevantImpl(UMovieSceneEntitySystemLinker* InLinker) const
+{
+	using namespace UE::MovieScene;
+
+#if WITH_EDITOR
+	if (FCameraShakePreviewerLinkerExtension* Extension = InLinker->FindExtension<FCameraShakePreviewerLinkerExtension>())
+	{
+		return Extension->HasAnyShake();
+	}
+#endif
+	return false;
+}
+
+void UMovieSceneCameraShakePreviewerEvaluatorSystem::OnLink()
+{
+	using namespace UE::MovieScene;
+
+#if WITH_EDITOR
+	PreviewerExtension = FCameraShakePreviewerLinkerExtension::GetOrCreateExtension(Linker);
+#endif
+}
+
+void UMovieSceneCameraShakePreviewerEvaluatorSystem::OnUnlink()
+{
+	using namespace UE::MovieScene;
+
+#if WITH_EDITOR
+	// Only the two camera shake systems hold pointers to the extension, so it should delete itself
+	// once both systems are unlinked.
+	PreviewerExtension = nullptr;
+#endif
+}
+
+void UMovieSceneCameraShakePreviewerEvaluatorSystem::OnRun(FSystemTaskPrerequisites& InPrerequisites, FSystemSubsequentTasks& Subsequents)
+{
+	using namespace UE::MovieScene;
+
+#if WITH_EDITOR
+	FInstanceRegistry* InstanceRegistry = Linker->GetInstanceRegistry();
+	PreviewerExtension->UpdateAllPreviewers(InstanceRegistry);
+#endif
+}
+
diff --git a/Engine/Source/Runtime/MovieSceneTracks/Private/Tracks/MovieSceneCameraShakeSourceShakeTrack.cpp b/Engine/Source/Runtime/MovieSceneTracks/Private/Tracks/MovieSceneCameraShakeSourceShakeTrack.cpp
--- a/Engine/Source/Runtime/MovieSceneTracks/Private/Tracks/MovieSceneCameraShakeSourceShakeTrack.cpp
+++ b/Engine/Source/Runtime/MovieSceneTracks/Private/Tracks/MovieSceneCameraShakeSourceShakeTrack.cpp
@@ -4,7 +4,6 @@
 #include "Camera/CameraShakeSourceComponent.h"
 #include "Compilation/MovieSceneCompilerRules.h"
 #include "Sections/MovieSceneCameraShakeSourceShakeSection.h"
-#include "Evaluation/MovieSceneCameraShakeSourceShakeTemplate.h"
 #include "MovieScene.h"
 
 #include UE_INLINE_GENERATED_CPP_BY_NAME(MovieSceneCameraShakeSourceShakeTrack)
@@ -106,11 +105,6 @@ bool UMovieSceneCameraShakeSourceShakeTrack::IsEmpty() const
 	return CameraShakeSections.Num() == 0;
 }
 
-FMovieSceneEvalTemplatePtr UMovieSceneCameraShakeSourceShakeTrack::CreateTemplateForSection(const UMovieSceneSection& InSection) const
-{
-	return FMovieSceneCameraShakeSourceShakeSectionTemplate(*CastChecked<const UMovieSceneCameraShakeSourceShakeSection>(&InSection));
-}
-
 #undef LOCTEXT_NAMESPACE
 
 
diff --git a/Engine/Source/Runtime/MovieSceneTracks/Private/Tracks/MovieSceneCameraShakeSourceTriggerTrack.cpp b/Engine/Source/Runtime/MovieSceneTracks/Private/Tracks/MovieSceneCameraShakeSourceTriggerTrack.cpp
--- a/Engine/Source/Runtime/MovieSceneTracks/Private/Tracks/MovieSceneCameraShakeSourceTriggerTrack.cpp
+++ b/Engine/Source/Runtime/MovieSceneTracks/Private/Tracks/MovieSceneCameraShakeSourceTriggerTrack.cpp
@@ -2,7 +2,6 @@
 
 #include "Tracks/MovieSceneCameraShakeSourceTriggerTrack.h"
 #include "Sections/MovieSceneCameraShakeSourceTriggerSection.h"
-#include "Evaluation/MovieSceneCameraShakeSourceTriggerTemplate.h"
 #include "Evaluation/MovieSceneEvaluationTrack.h"
 
 #include UE_INLINE_GENERATED_CPP_BY_NAME(MovieSceneCameraShakeSourceTriggerTrack)
@@ -67,16 +66,6 @@ const TArray<UMovieSceneSection*>& UMovieSceneCameraShakeSourceTriggerTrack::Get
 	return Sections;
 }
 
-FMovieSceneEvalTemplatePtr UMovieSceneCameraShakeSourceTriggerTrack::CreateTemplateForSection(const UMovieSceneSection& InSection) const
-{
-	return FMovieSceneCameraShakeSourceTriggerSectionTemplate(*CastChecked<const UMovieSceneCameraShakeSourceTriggerSection>(&InSection));
-}
-
-void UMovieSceneCameraShakeSourceTriggerTrack::PostCompile(FMovieSceneEvaluationTrack& Track, const FMovieSceneTrackCompilerArgs& Args) const
-{
-	Track.SetEvaluationMethod(EEvaluationMethod::Swept);
-}
-
 #if WITH_EDITORONLY_DATA
 
 FText UMovieSceneCameraShakeSourceTriggerTrack::GetDisplayName() const
diff --git a/Engine/Source/Runtime/MovieSceneTracks/Private/Tracks/MovieSceneCameraShakeTrack.cpp b/Engine/Source/Runtime/MovieSceneTracks/Private/Tracks/MovieSceneCameraShakeTrack.cpp
--- a/Engine/Source/Runtime/MovieSceneTracks/Private/Tracks/MovieSceneCameraShakeTrack.cpp
+++ b/Engine/Source/Runtime/MovieSceneTracks/Private/Tracks/MovieSceneCameraShakeTrack.cpp
@@ -2,10 +2,6 @@
 
 #include "Tracks/MovieSceneCameraShakeTrack.h"
 #include "Sections/MovieSceneCameraShakeSection.h"
-#include "Evaluation/PersistentEvaluationData.h"
-#include "Evaluation/MovieSceneCameraShakeTemplate.h"
-#include "Evaluation/MovieSceneEvaluationTrack.h"
-#include "Evaluation/MovieSceneEvaluationTemplate.h"
 #include "Compilation/MovieSceneSegmentCompiler.h"
 #include "Compilation/MovieSceneCompilerRules.h"
 #include "MovieScene.h"
@@ -39,15 +35,6 @@ FText UMovieSceneCameraShakeTrack::GetDisplayName() const
 }
 #endif
 
-FMovieSceneEvalTemplatePtr UMovieSceneCameraShakeTrack::CreateTemplateForSection(const UMovieSceneSection& InSection) const
-{
-	const UMovieSceneCameraShakeSection* CameraShakeSection = CastChecked<const UMovieSceneCameraShakeSection>(&InSection);
-	if (*CameraShakeSection->ShakeData.ShakeClass)
-	{
-		return FMovieSceneCameraShakeSectionTemplate(*CameraShakeSection);
-	}
-	return FMovieSceneEvalTemplatePtr();
-}
 
 /* UMovieSceneTrack interface
 *****************************************************************************/
diff --git a/Engine/Source/Runtime/MovieSceneTracks/Public/Evaluation/MovieSceneCameraShakeTemplate.h b/Engine/Source/Runtime/MovieSceneTracks/Public/Evaluation/MovieSceneCameraShakeTemplate.h
deleted file mode 100644
diff --git a/Engine/Source/Runtime/MovieSceneTracks/Public/MovieSceneTracksComponentTypes.h b/Engine/Source/Runtime/MovieSceneTracks/Public/MovieSceneTracksComponentTypes.h
--- a/Engine/Source/Runtime/MovieSceneTracks/Public/MovieSceneTracksComponentTypes.h
+++ b/Engine/Source/Runtime/MovieSceneTracks/Public/MovieSceneTracksComponentTypes.h
@@ -16,6 +16,7 @@
 #include "Styling/SlateColor.h"
 #include "ConstraintChannel.h"
 #include "MaterialTypes.h"
+#include "Sections/MovieSceneCameraShakeSection.h"
 #include "Tracks/MovieSceneMaterialTrack.h"
 #include "MovieSceneTracksComponentTypes.generated.h"
 
@@ -104,6 +105,60 @@ struct FMovieSceneAudioComponentData
 	TObjectPtr<UMovieSceneAudioSection> Section = nullptr;
 };
 
+UENUM()
+enum EMovieSceneCameraShakeStatus : uint8
+{
+	NotStarted,
+	Started,
+	BlendingOut,
+	Finished
+};
+
+/** Component data for camera shakes */
+USTRUCT()
+struct FMovieSceneCameraShakeComponentData
+{
+	GENERATED_BODY()
+
+	/** The shake data from the section that created this component */
+	UPROPERTY()
+	FMovieSceneCameraShakeSectionData SectionData;
+
+	/** The range of the section that created this component */
+	UPROPERTY()
+	FFrameNumber SectionStartTime;
+
+	UPROPERTY()
+	FFrameNumber SectionEndTime;
+
+	// ~~ Internal data for the camera shake systems ~~
+
+	// Shake instance created by the shake evaluation system.
+	UPROPERTY()
+	TObjectPtr<UCameraShakeBase> ShakeInstance;
+	
+	// Status for the shake.
+	UPROPERTY()
+	TEnumAsByte<EMovieSceneCameraShakeStatus> Status = EMovieSceneCameraShakeStatus::NotStarted;
+
+	// Duration of the shake.
+	UPROPERTY()
+	FFrameTime Duration;
+
+	// Time for blending out the shake.
+	UPROPERTY()
+	FFrameTime BlendOutTime;
+
+	FMovieSceneCameraShakeComponentData()
+	{}
+	FMovieSceneCameraShakeComponentData(const FMovieSceneCameraShakeSectionData& InSectionData, TRange<FFrameNumber> InSectionRange)
+		: SectionData(InSectionData)
+	{
+		SectionStartTime = InSectionRange.HasLowerBound() ? InSectionRange.GetLowerBoundValue() : 0;
+		SectionEndTime = InSectionRange.HasUpperBound() ? InSectionRange.GetUpperBoundValue() : FFrameNumber(TNumericLimits<int32>::Max());
+	}
+};
+
 /**
  * Component data for audio tracks inputs
  * This provides the names of the inputs whose values are stored on the
@@ -581,6 +636,8 @@ struct FMovieSceneTracksComponentTypes
 	TComponentTypeID<FMovieSceneAudioInputData> AudioInputs;
 	TComponentTypeID<FName> AudioTriggerName;
 
+	TComponentTypeID<FMovieSceneCameraShakeComponentData> CameraShake;
+
 	struct
 	{
 		TCustomPropertyRegistration<FBoolPropertyTraits> Bool;
diff --git a/Engine/Source/Runtime/MovieSceneTracks/Public/Sections/MovieSceneCameraShakeSection.h b/Engine/Source/Runtime/MovieSceneTracks/Public/Sections/MovieSceneCameraShakeSection.h
--- a/Engine/Source/Runtime/MovieSceneTracks/Public/Sections/MovieSceneCameraShakeSection.h
+++ b/Engine/Source/Runtime/MovieSceneTracks/Public/Sections/MovieSceneCameraShakeSection.h
@@ -6,6 +6,7 @@
 #include "Camera/CameraTypes.h"
 #include "CoreMinimal.h"
 #include "CoreTypes.h"
+#include "EntitySystem/IMovieSceneEntityProvider.h"
 #include "Math/Rotator.h"
 #include "MovieSceneSection.h"
 #include "Templates/SubclassOf.h"
@@ -45,11 +46,10 @@ struct FMovieSceneCameraShakeSectionData
 	FRotator UserDefinedPlaySpace;
 };
 
-/**
- *
- */
 UCLASS(MinimalAPI)
-class UMovieSceneCameraShakeSection : public UMovieSceneSection
+class UMovieSceneCameraShakeSection
+	: public UMovieSceneSection
+	, public IMovieSceneEntityProvider
 {
 	GENERATED_BODY()
 
@@ -58,11 +58,17 @@ public:
 	UMovieSceneCameraShakeSection(const FObjectInitializer& ObjectInitializer);
 
 	virtual void PostLoad() override;
-	
+
+	/** IMovieSceneEntityProvider interface */
+	void ImportEntityImpl(UMovieSceneEntitySystemLinker* EntityLinker, const FEntityImportParams& Params, FImportedEntity* OutImportedEntity) override;
+
+public:
+
 	UPROPERTY(EditAnywhere, Category="Camera Shake", meta=(ShowOnlyInnerProperties))
 	FMovieSceneCameraShakeSectionData ShakeData;
 
 public:
+
 	UPROPERTY()
 	TSubclassOf<UCameraShakeBase> ShakeClass_DEPRECATED;
 	
diff --git a/Engine/Source/Runtime/MovieSceneTracks/Public/Sections/MovieSceneCameraShakeSourceShakeSection.h b/Engine/Source/Runtime/MovieSceneTracks/Public/Sections/MovieSceneCameraShakeSourceShakeSection.h
--- a/Engine/Source/Runtime/MovieSceneTracks/Public/Sections/MovieSceneCameraShakeSourceShakeSection.h
+++ b/Engine/Source/Runtime/MovieSceneTracks/Public/Sections/MovieSceneCameraShakeSourceShakeSection.h
@@ -3,6 +3,7 @@
 #pragma once
 
 #include "CoreMinimal.h"
+#include "EntitySystem/IMovieSceneEntityProvider.h"
 #include "MovieSceneCameraShakeSection.h"
 #include "MovieSceneSection.h"
 #include "UObject/ObjectMacros.h"
@@ -13,13 +14,19 @@
 class UObject;
 
 UCLASS(MinimalAPI)
-class UMovieSceneCameraShakeSourceShakeSection : public UMovieSceneSection
+class UMovieSceneCameraShakeSourceShakeSection 
+	: public UMovieSceneSection
+	, public IMovieSceneEntityProvider
 {
 	GENERATED_BODY()
 
 public:
 	UMovieSceneCameraShakeSourceShakeSection(const FObjectInitializer& ObjectInitializer);
+
+	/** IMovieSceneEntityProvider interface */
+	void ImportEntityImpl(UMovieSceneEntitySystemLinker* EntityLinker, const FEntityImportParams& Params, FImportedEntity* OutImportedEntity) override;
 	
+public:
 	UPROPERTY(EditAnywhere, Category="Camera Shake", meta=(ShowOnlyInnerProperties))
 	FMovieSceneCameraShakeSectionData ShakeData;
 };
diff --git a/Engine/Source/Runtime/MovieSceneTracks/Public/Sections/MovieSceneCameraShakeSourceTriggerSection.h b/Engine/Source/Runtime/MovieSceneTracks/Public/Sections/MovieSceneCameraShakeSourceTriggerSection.h
--- a/Engine/Source/Runtime/MovieSceneTracks/Public/Sections/MovieSceneCameraShakeSourceTriggerSection.h
+++ b/Engine/Source/Runtime/MovieSceneTracks/Public/Sections/MovieSceneCameraShakeSourceTriggerSection.h
@@ -5,6 +5,7 @@
 #include "Camera/CameraShakeBase.h"
 #include "Channels/MovieSceneCameraShakeSourceTriggerChannel.h"
 #include "CoreMinimal.h"
+#include "EntitySystem/IMovieSceneEntityProvider.h"
 #include "MovieSceneSection.h"
 #include "UObject/ObjectMacros.h"
 #include "UObject/UObjectGlobals.h"
@@ -16,6 +17,7 @@ class UObject;
 UCLASS(MinimalAPI)
 class UMovieSceneCameraShakeSourceTriggerSection 
 	: public UMovieSceneSection
+	, public IMovieSceneEntityProvider
 {
 public:
 	GENERATED_BODY()
@@ -24,6 +26,10 @@ public:
 
 	const FMovieSceneCameraShakeSourceTriggerChannel& GetChannel() const { return Channel; }
 
+	/** IMovieSceneEntityProvider interface */
+	bool PopulateEvaluationFieldImpl(const TRange<FFrameNumber>& EffectiveRange, const FMovieSceneEvaluationFieldEntityMetaData& InMetaData, FMovieSceneEntityComponentFieldBuilder* OutFieldBuilder) override;
+	void ImportEntityImpl(UMovieSceneEntitySystemLinker* EntityLinker, const FEntityImportParams& Params, FImportedEntity* OutImportedEntity) override;
+
 private:
 	UPROPERTY()
 	FMovieSceneCameraShakeSourceTriggerChannel Channel;
diff --git a/Engine/Source/Runtime/MovieSceneTracks/Public/Systems/MovieSceneCameraShakeSystem.h b/Engine/Source/Runtime/MovieSceneTracks/Public/Systems/MovieSceneCameraShakeSystem.h
new file mode 100644
--- /dev/null
+++ b/Engine/Source/Runtime/MovieSceneTracks/Public/Systems/MovieSceneCameraShakeSystem.h
@@ -0,0 +1,132 @@
+// Copyright Epic Games, Inc. All Rights Reserved.
+
+#pragma once
+
+#include "CoreTypes.h"
+#include "EntitySystem/MovieSceneEntitySystem.h"
+#include "EntitySystem/MovieSceneEntitySystemLinkerExtension.h"
+#include "EntitySystem/MovieSceneEntitySystemLinkerSharedExtension.h"
+#include "Evaluation/MovieSceneCameraShakePreviewer.h"
+#include "Sections/MovieSceneCameraShakeSourceTriggerSection.h"
+
+#include "MovieSceneCameraShakeSystem.generated.h"
+
+namespace UE::MovieScene
+{
+	struct FPreAnimatedCameraShakeStateStorage;
+	struct FPreAnimatedCameraComponentShakeStateStorage;
+	struct FPreAnimatedCameraSourceShakeStateStorage;
+
+#if WITH_EDITOR
+	/**
+	 * Linker extension for storing camera shake previewers shared between the instantiator and the
+	 * evaluation systems below. This extension is held by a shared pointer, so it should unregister
+	 * and delete itself when both systems are unlinked.
+	 */
+	struct FCameraShakePreviewerLinkerExtension : public TSharedEntitySystemLinkerExtension<FCameraShakePreviewerLinkerExtension>
+	{
+		/** The ID of this extension */
+		static TEntitySystemLinkerExtensionID<FCameraShakePreviewerLinkerExtension> GetExtensionID();
+
+		/** Get the extension from the given linker, or create a new one */
+		static TSharedPtr<FCameraShakePreviewerLinkerExtension> GetOrCreateExtension(UMovieSceneEntitySystemLinker* Linker);
+
+		/** Constructor */
+		FCameraShakePreviewerLinkerExtension(UMovieSceneEntitySystemLinker* Linker);
+		/** Destructor */
+		virtual ~FCameraShakePreviewerLinkerExtension();
+
+		/** Find the previewer for the given sequence instance, if it has already been created */
+		FCameraShakePreviewer* FindPreviewer(FInstanceHandle InstanceHandle);
+		/** Get or create the previewer for the given sequence instance */
+		FCameraShakePreviewer& GetPreviewer(FInstanceHandle InstanceHandle);
+		/** Update all previewers by using the update context of each matching sequence instance */
+		void UpdateAllPreviewers(FInstanceRegistry* InstanceRegistry);
+		/** Whether there is any previewer with any shake */
+		bool HasAnyShake() const;
+
+	private:
+		TMap<FInstanceHandle, FCameraShakePreviewer> Previewers;
+	};
+#endif
+}
+
+/**
+ * The instantiator system for camera shakes.
+ *
+ * This system will create and initialize camera shakes as they come in, and save pre-animated state.
+ */
+UCLASS()
+class UMovieSceneCameraShakeInstantiatorSystem : public UMovieSceneEntitySystem
+{
+	GENERATED_BODY()
+
+public:
+	UMovieSceneCameraShakeInstantiatorSystem(const FObjectInitializer& ObjInit);
+
+	virtual bool IsRelevantImpl(UMovieSceneEntitySystemLinker* InLinker) const override;
+	virtual void OnLink() override;
+	virtual void OnUnlink() override;
+
+	virtual void OnRun(FSystemTaskPrerequisites& InPrerequisites, FSystemSubsequentTasks& Subsequents) override;
+
+	/** Add a fire-and-forget shake trigger */
+	void AddShakeTrigger(UE::MovieScene::FInstanceHandle InInstance, const FGuid& ObjectBindingID, const FFrameTime& InTime, const FMovieSceneCameraShakeSourceTrigger& InTrigger);
+
+private:
+	void TriggerOneShotShakes();
+
+private:
+	struct FTimedTrigger
+	{
+		FGuid ObjectBindingID;
+		FFrameTime Time;
+		FMovieSceneCameraShakeSourceTrigger Trigger;
+	};
+
+	TSharedPtr<UE::MovieScene::FPreAnimatedCameraShakeStateStorage> PreAnimatedCameraShakeStorage;
+	TSharedPtr<UE::MovieScene::FPreAnimatedCameraComponentShakeStateStorage> PreAnimatedCameraComponentShakeStorage;
+	TSharedPtr<UE::MovieScene::FPreAnimatedCameraSourceShakeStateStorage> PreAnimatedCameraSourceShakeStorage;
+
+	TMap<UE::MovieScene::FInstanceHandle, TArray<FTimedTrigger>> TriggersByInstance;
+
+#if WITH_EDITOR
+	TSharedPtr<UE::MovieScene::FCameraShakePreviewerLinkerExtension> PreviewerExtension;
+#endif
+};
+
+/**
+ * The evaluator system for camera shakes.
+ *
+ * This system takes care of updating ongoing camera shakes according to the evaluation range.
+ */
+UCLASS()
+class UMovieSceneCameraShakeEvaluatorSystem : public UMovieSceneEntitySystem
+{
+	GENERATED_BODY()
+
+public:
+	UMovieSceneCameraShakeEvaluatorSystem(const FObjectInitializer& ObjInit);
+
+	virtual void OnRun(FSystemTaskPrerequisites& InPrerequisites, FSystemSubsequentTasks& Subsequents) override;
+};
+
+UCLASS()
+class UMovieSceneCameraShakePreviewerEvaluatorSystem : public UMovieSceneEntitySystem
+{
+	GENERATED_BODY()
+
+public:
+	UMovieSceneCameraShakePreviewerEvaluatorSystem(const FObjectInitializer& ObjInit);
+
+	virtual void OnLink() override;
+	virtual void OnUnlink() override;
+	virtual bool IsRelevantImpl(UMovieSceneEntitySystemLinker* InLinker) const override;
+	virtual void OnRun(FSystemTaskPrerequisites& InPrerequisites, FSystemSubsequentTasks& Subsequents) override;
+
+private:
+#if WITH_EDITOR
+	TSharedPtr<UE::MovieScene::FCameraShakePreviewerLinkerExtension> PreviewerExtension;
+#endif
+};
+
diff --git a/Engine/Source/Runtime/MovieSceneTracks/Public/Tracks/MovieSceneCameraShakeSourceShakeTrack.h b/Engine/Source/Runtime/MovieSceneTracks/Public/Tracks/MovieSceneCameraShakeSourceShakeTrack.h
--- a/Engine/Source/Runtime/MovieSceneTracks/Public/Tracks/MovieSceneCameraShakeSourceShakeTrack.h
+++ b/Engine/Source/Runtime/MovieSceneTracks/Public/Tracks/MovieSceneCameraShakeSourceShakeTrack.h
@@ -3,11 +3,9 @@
 #pragma once
 
 #include "Camera/CameraShakeBase.h"
-#include "Compilation/IMovieSceneTrackTemplateProducer.h"
 #include "Containers/Array.h"
 #include "CoreMinimal.h"
 #include "CoreTypes.h"
-#include "Evaluation/MovieSceneEvalTemplate.h"
 #include "Internationalization/Text.h"
 #include "Misc/InlineValue.h"
 #include "MovieSceneNameableTrack.h"
@@ -30,7 +28,8 @@ struct FMovieSceneSegmentCompilerRules;
  * 
  */
 UCLASS(MinimalAPI)
-class UMovieSceneCameraShakeSourceShakeTrack : public UMovieSceneNameableTrack, public IMovieSceneTrackTemplateProducer
+class UMovieSceneCameraShakeSourceShakeTrack 
+	: public UMovieSceneNameableTrack
 {
 	GENERATED_BODY()
 
@@ -51,8 +50,6 @@ public:
 	virtual UMovieSceneSection* CreateNewSection() override;
 	virtual void RemoveAllAnimationData() override;
 
-	virtual FMovieSceneEvalTemplatePtr CreateTemplateForSection(const UMovieSceneSection& InSection) const override;
-
 #if WITH_EDITORONLY_DATA
 	virtual FText GetDisplayName() const override;
 #endif
diff --git a/Engine/Source/Runtime/MovieSceneTracks/Public/Tracks/MovieSceneCameraShakeSourceTriggerTrack.h b/Engine/Source/Runtime/MovieSceneTracks/Public/Tracks/MovieSceneCameraShakeSourceTriggerTrack.h
--- a/Engine/Source/Runtime/MovieSceneTracks/Public/Tracks/MovieSceneCameraShakeSourceTriggerTrack.h
+++ b/Engine/Source/Runtime/MovieSceneTracks/Public/Tracks/MovieSceneCameraShakeSourceTriggerTrack.h
@@ -2,11 +2,9 @@
 
 #pragma once
 
-#include "Compilation/IMovieSceneTrackTemplateProducer.h"
 #include "Containers/Array.h"
 #include "CoreMinimal.h"
 #include "CoreTypes.h"
-#include "Evaluation/MovieSceneEvalTemplate.h"
 #include "Internationalization/Text.h"
 #include "MovieSceneSection.h"
 #include "MovieSceneTrack.h"
@@ -23,7 +21,6 @@ struct FMovieSceneEvaluationTrack;
 UCLASS(MinimalAPI)
 class UMovieSceneCameraShakeSourceTriggerTrack
 	: public UMovieSceneTrack
-	, public IMovieSceneTrackTemplateProducer
 {
 public:
 
@@ -41,9 +38,6 @@ public:
 	MOVIESCENETRACKS_API virtual void RemoveAllAnimationData() override;
 	MOVIESCENETRACKS_API virtual bool IsEmpty() const override;
 	MOVIESCENETRACKS_API virtual const TArray<UMovieSceneSection*>& GetAllSections() const override;
-	MOVIESCENETRACKS_API virtual void PostCompile(FMovieSceneEvaluationTrack& Track, const FMovieSceneTrackCompilerArgs& Args) const override;
-
-	MOVIESCENETRACKS_API virtual FMovieSceneEvalTemplatePtr CreateTemplateForSection(const UMovieSceneSection& InSection) const override;
 
 #if WITH_EDITORONLY_DATA
 	MOVIESCENETRACKS_API virtual FText GetDisplayName() const override;
diff --git a/Engine/Source/Runtime/MovieSceneTracks/Public/Tracks/MovieSceneCameraShakeTrack.h b/Engine/Source/Runtime/MovieSceneTracks/Public/Tracks/MovieSceneCameraShakeTrack.h
--- a/Engine/Source/Runtime/MovieSceneTracks/Public/Tracks/MovieSceneCameraShakeTrack.h
+++ b/Engine/Source/Runtime/MovieSceneTracks/Public/Tracks/MovieSceneCameraShakeTrack.h
@@ -3,11 +3,9 @@
 #pragma once
 
 #include "Camera/CameraShakeBase.h"
-#include "Compilation/IMovieSceneTrackTemplateProducer.h"
 #include "Containers/Array.h"
 #include "CoreMinimal.h"
 #include "CoreTypes.h"
-#include "Evaluation/MovieSceneEvalTemplate.h"
 #include "Internationalization/Text.h"
 #include "Misc/InlineValue.h"
 #include "MovieSceneNameableTrack.h"
@@ -29,7 +27,7 @@ struct FMovieSceneSegmentCompilerRules;
  * 
  */
 UCLASS(MinimalAPI)
-class UMovieSceneCameraShakeTrack : public UMovieSceneNameableTrack, public IMovieSceneTrackTemplateProducer
+class UMovieSceneCameraShakeTrack : public UMovieSceneNameableTrack
 {
 	GENERATED_BODY()
 
@@ -50,9 +48,6 @@ public:
 	virtual UMovieSceneSection* CreateNewSection() override;
 	virtual void RemoveAllAnimationData() override;
 	
-	// ~IMovieSceneTrackTemplateProducer interface
-	virtual FMovieSceneEvalTemplatePtr CreateTemplateForSection(const UMovieSceneSection& InSection) const override;
-
 #if WITH_EDITORONLY_DATA
 	virtual FText GetDisplayName() const override;
 #endif
