Sequencer: Further improvements to garbage collection handling

  - When an entity was garbage collected, we previously assumed that there could be no contributors left animating that property.
  - This assumption is false if the fast path entity is only being destroyed as a result of a whole sequence being GC'd where there are other sequences also animating the same property (but ignored by hbias)
  - To handle this correctly we copy all the initial value and meta-data information off the expiring fast path entity, and add it to any one of the remaining contributors. All reminaing contributors also get NeedsLink.
    - When we process invalidated properties, these remaining contributors will get re-processed accordingly, and the initial values will end up on whichever entity is deemed the new authority (either a new fast path entity, or a blend output)

  - The assumption remains true if a blend output is being destroyed by GC. This should only happen if the actual animated object is being destroyed, and thus all contributors should also be destroyed.
  - We can handle the breaking of this assumption gracefully however by just tagging remaining contributors as needs link.

diff --git a/Engine/Source/Runtime/MovieSceneTracks/Private/Systems/MovieScenePropertyInstantiator.cpp b/Engine/Source/Runtime/MovieSceneTracks/Private/Systems/MovieScenePropertyInstantiator.cpp
--- a/Engine/Source/Runtime/MovieSceneTracks/Private/Systems/MovieScenePropertyInstantiator.cpp
+++ b/Engine/Source/Runtime/MovieSceneTracks/Private/Systems/MovieScenePropertyInstantiator.cpp
@@ -174,9 +174,13 @@ void UMovieScenePropertyInstantiatorSystem::OnUnlink()
 
 void UMovieScenePropertyInstantiatorSystem::OnCleanTaggedGarbage()
 {
+	using namespace UE::MovieScene;
+
 	// Only process expired properties for this GC pass to ensure we don't end up creating any new entities
 	DiscoverExpiredProperties(PendingInvalidatedProperties);
 
+	TArrayView<const FPropertyDefinition> Properties = this->BuiltInComponents->PropertyRegistry.GetProperties();
+
 	bool bAnyDestroyed = false;
 	// Look through our resolved properties to detect any outputs that have been destroyed
 	for (int32 Index = 0; Index < ResolvedProperties.GetMaxIndex(); ++Index)
@@ -188,14 +192,61 @@ void UMovieScenePropertyInstantiatorSystem::OnCleanTaggedGarbage()
 
 		FObjectPropertyInfo& PropertyInfo = ResolvedProperties[Index];
 
-		// If the final output entity is being destroyed, clean up the property info and blender so it can be re-created next time if necessary
-		if (
-			(PropertyInfo.FinalBlendOutputID && Linker->EntityManager.HasComponent(PropertyInfo.FinalBlendOutputID, BuiltInComponents->Tags.NeedsUnlink)) ||
-			(PropertyInfo.PreviousFastPathID && Linker->EntityManager.HasComponent(PropertyInfo.PreviousFastPathID, BuiltInComponents->Tags.NeedsUnlink))
-			)
+		FContributorKey ContributorKey(Index);
+
+		const bool bFastPathOutputBeingDestroyed = PropertyInfo.PreviousFastPathID && Linker->EntityManager.HasComponent(PropertyInfo.PreviousFastPathID, BuiltInComponents->Tags.NeedsUnlink);
+
+		// If the fast path is being destroyed, we need to copy over any initial values from that fast path entity to
+		// any additional contributors that might still be alive. This can happen if a sequence gets GC'd and removes its entities while another is still alive animating the same thing.
+		if (bFastPathOutputBeingDestroyed)
 		{
-			// Really we shouldn't have any contributors any more if the output is being destroyed
-			ensure(!Contributors.Contains(FContributorKey(Index)));
+			FMovieSceneEntityID TemporaryFastPathEntity;
+			for (auto It = Contributors.CreateConstKeyIterator(ContributorKey); It; ++It)
+			{
+				Linker->EntityManager.AddComponent(It->Value, BuiltInComponents->Tags.NeedsLink);
+				TemporaryFastPathEntity = It->Value;
+			}
+
+			if (TemporaryFastPathEntity)
+			{
+				const FPropertyDefinition& PropertyDefinition = Properties[PropertyInfo.PropertyDefinitionIndex];
+				FComponentMask CopyMask;
+				CopyMask.Set(PropertyDefinition.InitialValueType);
+				CopyMask.Set(BuiltInComponents->Tags.HasAssignedInitialValue);
+				for (FComponentTypeID Component : PropertyDefinition.MetaDataTypes)
+				{
+					CopyMask.Set(Component);
+				}
+
+				Linker->EntityManager.CopyComponents(PropertyInfo.PreviousFastPathID, TemporaryFastPathEntity, CopyMask);
+				PropertyInfo.PreviousFastPathID = TemporaryFastPathEntity;
+			}
+			else
+			{
+				// There is nothing else animating this - destroy the property entirely
+				DestroyStaleProperty(Index);
+				bAnyDestroyed = true;
+
+				if (PendingInvalidatedProperties.IsValidIndex(Index) && PendingInvalidatedProperties[Index] == true)
+				{
+					// This property index is no longer valid at all
+					PendingInvalidatedProperties[Index] = false;
+				}
+			}
+		}
+		else if (PropertyInfo.FinalBlendOutputID && Linker->EntityManager.HasComponent(PropertyInfo.FinalBlendOutputID, BuiltInComponents->Tags.NeedsUnlink))
+		{
+			// Really we shouldn't have any contributors any more if the output is being destroyed since the target object must be going away
+			//    (which means all the contributors that reference that target object must also be going away)
+			if (!ensureMsgf(!Contributors.Contains(ContributorKey), TEXT("Blend output is being destroyed while there are still contributors.")))
+			{
+				// We still have contributors?? Shouldn't happen, but it's recoverable by just re-resolving all the contributors
+				for (auto It = Contributors.CreateKeyIterator(ContributorKey); It; ++It)
+				{
+					Linker->EntityManager.AddComponent(It->Value, BuiltInComponents->Tags.NeedsLink);
+					It.RemoveCurrent();
+				}
+			}
 
 			DestroyStaleProperty(Index);
 			bAnyDestroyed = true;
