Sequencer: improve cinematic preview of camera shakes

- Add a duration override to the camera shake API so sequencer code doesn't have to manage shake duration itself. This happens when a camera shake section is shorter than the original shake duration.
- There are quite a few files changed for this, as the duration override parameter needs to be passed through many hoops.
- Handle the new duration override parameter in all built-in camera shake patterns.
- The sequencer camera shake system now simply passes this duration override when a shake section is short enough to affect a shake's duration. This removes a fair amount of code.
- Fix several bugs in the shake system: stop preview shakes on pre-animated state restore, don't use preview shakes if there is no GEditor, clear one-shot shake triggers once they've been triggered.
- Remove shake preview evaluator system, make the normal evaluator system tick the previewers.
- Simplify the shake previewer by removing the need to spawn a temp PlayerCameraManager actor in the world. Instead we just handle it all in the shake previewer class, at the cost of (1) a little bit of code duplication compared to UCameraModifier_CameraShake and (2) changing the base camera shake class to push post-process settings into the view info instead of the camera manager if there is no camera manager.

diff --git a/Engine/Plugins/Cameras/CameraShakePreviewer/Source/CameraShakePreviewer/Private/SCameraShakePreviewer.cpp b/Engine/Plugins/Cameras/CameraShakePreviewer/Source/CameraShakePreviewer/Private/SCameraShakePreviewer.cpp
--- a/Engine/Plugins/Cameras/CameraShakePreviewer/Source/CameraShakePreviewer/Private/SCameraShakePreviewer.cpp
+++ b/Engine/Plugins/Cameras/CameraShakePreviewer/Source/CameraShakePreviewer/Private/SCameraShakePreviewer.cpp
@@ -62,9 +62,6 @@ FCameraShakePreviewUpdater::FCameraShakePreviewUpdater()
 {
 	// Handle camera shakes being recompiled.
 	FCoreUObjectDelegates::OnObjectsReplaced.AddRaw(this, &FCameraShakePreviewUpdater::OnObjectsReplaced);
-
-	// Handle levels/worlds being cleared.
-	FEditorSupportDelegates::PrepareToCleanseEditorObject.AddRaw(this, &FCameraShakePreviewUpdater::OnPrepareToCleanseEditorObject);
 }
 
 FCameraShakePreviewUpdater::~FCameraShakePreviewUpdater()
@@ -73,47 +70,9 @@ FCameraShakePreviewUpdater::~FCameraShakePreviewUpdater()
 	FCoreUObjectDelegates::OnObjectsReplaced.RemoveAll(this);
 }
 
-void FCameraShakePreviewUpdater::Reinitialize(UWorld* InWorld)
-{
-	if (Previewer.IsInitialized() && Previewer.GetCameraManager()->GetWorld() != InWorld)
-	{
-		Previewer.Teardown();
-	}
-	if (!Previewer.IsInitialized())
-	{
-		Previewer.Initialize(InWorld);
-	}
-}
-
 void FCameraShakePreviewUpdater::OnObjectsReplaced(const TMap<UObject*, UObject*>& ReplacementMap)
 {
-	if (Previewer.IsInitialized())
-	{
-		// If a camera shake gets recompiled, we just stop and discard it.
-		TArray<FActiveCameraShakeInfo> ActiveCameraShakes;
-		Previewer.GetCameraModifier()->GetActiveCameraShakes(ActiveCameraShakes);
-
-		for (const FActiveCameraShakeInfo& ActiveCameraShake : ActiveCameraShakes)
-		{
-			UObject* NewShakeInstance = ReplacementMap.FindRef(ActiveCameraShake.ShakeInstance);
-			if (NewShakeInstance)
-			{
-				Previewer.GetCameraModifier()->RemoveCameraShake(ActiveCameraShake.ShakeInstance, true);
-			}
-		}
-	}
-}
-
-void FCameraShakePreviewUpdater::OnPrepareToCleanseEditorObject(UObject* Object)
-{
-	if (Previewer.IsInitialized())
-	{
-		ULevel* Level = Previewer.GetCameraManager()->GetLevel();
-		if (Object == Level || Object == Level->OwningWorld)
-		{
-			Previewer.Teardown();
-		}
-	}
+	Previewer.OnObjectsReplaced(ReplacementMap);
 }
 
 void FCameraShakePreviewUpdater::Tick(float DeltaTime)
@@ -128,31 +87,33 @@ void FCameraShakePreviewUpdater::ModifyCamera(FEditorViewportViewModifierParams&
 
 UCameraShakeBase* FCameraShakePreviewUpdater::AddCameraShake(TSubclassOf<UCameraShakeBase> ShakeClass, const FAddCameraShakeParams& Params)
 {
-	FAddCameraShakeParams ActualParams(Params);
-
-	UCameraShakeBase* ShakeInstance = Previewer.GetCameraModifier()->AddCameraShake(ShakeClass, ActualParams);
-
-	return ShakeInstance;
+	FCameraShakePreviewerAddParams ActualParams;
+	ActualParams.ShakeClass = ShakeClass;
+	ActualParams.SourceComponent = Params.SourceComponent;
+	ActualParams.Scale = Params.Scale;
+	ActualParams.PlaySpace = Params.PlaySpace;
+	ActualParams.UserPlaySpaceRot = Params.UserPlaySpaceRot;
+	return Previewer.AddCameraShake(ActualParams);
 }
 
 void FCameraShakePreviewUpdater::RemoveAllCameraShakesFromSource(const UCameraShakeSourceComponent* SourceComponent)
 {
-	Previewer.GetCameraModifier()->RemoveAllCameraShakesFromSource(SourceComponent);
+	Previewer.RemoveAllCameraShakesFromSource(SourceComponent);
 }
 
 void FCameraShakePreviewUpdater::GetActiveCameraShakes(TArray<FActiveCameraShakeInfo>& ActiveCameraShakes) const
 {
-	Previewer.GetCameraModifier()->GetActiveCameraShakes(ActiveCameraShakes);
+	Previewer.GetActiveCameraShakes(ActiveCameraShakes);
 }
 
 void FCameraShakePreviewUpdater::RemoveCameraShake(UCameraShakeBase* ShakeInstance)
 {
-	Previewer.GetCameraModifier()->RemoveCameraShake(ShakeInstance);
+	Previewer.RemoveCameraShake(ShakeInstance);
 }
 
 void FCameraShakePreviewUpdater::RemoveAllCameraShakes()
 {
-	Previewer.GetCameraModifier()->RemoveAllCameraShakes();
+	Previewer.RemoveAllCameraShakes();
 }
 
 /**
@@ -406,9 +367,6 @@ void SCameraShakePreviewer::Populate()
 		return;
 	}
 
-	// Reinitialize the previewer in case the persistent level has changed.
-	CameraShakePreviewUpdater->Reinitialize(CurrentWorld);
-
 	// Get all the shake sources from the level.
 	TArray<UCameraShakeSourceComponent*> ShakeSourceComponents;
 	for (ULevel* Level : CurrentWorld->GetLevels())
diff --git a/Engine/Plugins/Cameras/CameraShakePreviewer/Source/CameraShakePreviewer/Private/SCameraShakePreviewer.h b/Engine/Plugins/Cameras/CameraShakePreviewer/Source/CameraShakePreviewer/Private/SCameraShakePreviewer.h
--- a/Engine/Plugins/Cameras/CameraShakePreviewer/Source/CameraShakePreviewer/Private/SCameraShakePreviewer.h
+++ b/Engine/Plugins/Cameras/CameraShakePreviewer/Source/CameraShakePreviewer/Private/SCameraShakePreviewer.h
@@ -56,11 +56,8 @@ public:
 
 	void ModifyCamera(FEditorViewportViewModifierParams& Params);
 
-	void Reinitialize(UWorld* InWorld);
-
 private:
 	void OnObjectsReplaced(const TMap<UObject*, UObject*>& ReplacementMap);
-	void OnPrepareToCleanseEditorObject(UObject* Object);
 
 private:
 	FCameraShakePreviewer Previewer;
diff --git a/Engine/Plugins/Cameras/GameplayCameras/Source/GameplayCameras/Private/LegacyCameraShake.cpp b/Engine/Plugins/Cameras/GameplayCameras/Source/GameplayCameras/Private/LegacyCameraShake.cpp
--- a/Engine/Plugins/Cameras/GameplayCameras/Source/GameplayCameras/Private/LegacyCameraShake.cpp
+++ b/Engine/Plugins/Cameras/GameplayCameras/Source/GameplayCameras/Private/LegacyCameraShake.cpp
@@ -111,10 +111,12 @@ void ULegacyCameraShake::DoStopShake(bool bImmediately)
 
 void ULegacyCameraShake::DoStartShake(const FCameraShakePatternStartParams& Params)
 {
-	const float EffectiveOscillationDuration = (OscillationDuration > 0.f) ? OscillationDuration : TNumericLimits<float>::Max();
+	ActualOscillationDuration = Params.bOverrideDuration ? Params.DurationOverride : OscillationDuration;
+
+	const float EffectiveOscillationDuration = (ActualOscillationDuration > 0.f) ? ActualOscillationDuration : TNumericLimits<float>::Max();
 
 	// init oscillations
-	if (OscillationDuration != 0.f)
+	if (ActualOscillationDuration != 0.f)
 	{
 		if (OscillatorTimeRemaining > 0.f)
 		{
@@ -186,7 +188,7 @@ void ULegacyCameraShake::DoStartShake(const FCameraShakePatternStartParams& Para
 		SequenceShakePattern->StartShakePattern(Params);
 	}
 
-	UE_LOG(LogLegacyCameraShake, Verbose, TEXT("ULegacyCameraShake::DoStartShake %s Duration: %f"), *GetNameSafe(this), OscillationDuration);
+	UE_LOG(LogLegacyCameraShake, Verbose, TEXT("ULegacyCameraShake::DoStartShake %s Duration: %f"), *GetNameSafe(this), ActualOscillationDuration);
 
 	ReceivePlayShake(ShakeScale);
 }
@@ -224,7 +226,7 @@ void ULegacyCameraShake::DoUpdateShake(const FCameraShakePatternUpdateParams& Pa
 		bBlendingOut = true;
 		CurrentBlendOutTime = OscillationBlendOutTime - OscillatorTimeRemaining;
 	}
-	else if (OscillationDuration < 0.f)
+	else if (ActualOscillationDuration < 0.f)
 	{
 		// infinite oscillation, keep the time remaining up
 		OscillatorTimeRemaining = TNumericLimits<float>::Max();
@@ -333,7 +335,7 @@ void ULegacyCameraShake::DoUpdateShake(const FCameraShakePatternUpdateParams& Pa
 		OutResult.FOV = InOutPOV.FOV;
 	}
 
-	UE_LOG(LogLegacyCameraShake, Verbose, TEXT("ULegacyCameraShake::DoUpdateShake %s Finished: %i Duration: %f Remaining: %f"), *GetNameSafe(this), bOscillationFinished, OscillationDuration, OscillatorTimeRemaining);
+	UE_LOG(LogLegacyCameraShake, Verbose, TEXT("ULegacyCameraShake::DoUpdateShake %s Finished: %i Duration: %f Remaining: %f"), *GetNameSafe(this), bOscillationFinished, ActualOscillationDuration, OscillatorTimeRemaining);
 }
 
 void ULegacyCameraShake::DoScrubShake(const FCameraShakePatternScrubParams& Params, FCameraShakePatternUpdateResult& OutResult)
@@ -345,7 +347,7 @@ void ULegacyCameraShake::DoScrubShake(const FCameraShakePatternScrubParams& Para
 	RotSinOffset = InitialRotSinOffset;
 	FOVSinOffset = InitialFOVSinOffset;
 
-	const float EffectiveOscillationDuration = (OscillationDuration > 0.f) ? OscillationDuration : TNumericLimits<float>::Max();
+	const float EffectiveOscillationDuration = (ActualOscillationDuration > 0.f) ? ActualOscillationDuration : TNumericLimits<float>::Max();
 
 	OscillatorTimeRemaining = EffectiveOscillationDuration;
 
@@ -361,12 +363,12 @@ void ULegacyCameraShake::DoScrubShake(const FCameraShakePatternScrubParams& Para
 		CurrentBlendOutTime = 0.f;
 	}
 
-	if (OscillationDuration > 0.f)
+	if (ActualOscillationDuration > 0.f)
 	{
-		if ((OscillationBlendOutTime > 0.f) && (NewTime > (OscillationDuration - OscillationBlendOutTime)))
+		if ((OscillationBlendOutTime > 0.f) && (NewTime > (ActualOscillationDuration - OscillationBlendOutTime)))
 		{
 			bBlendingOut = true;
-			CurrentBlendOutTime = OscillationBlendOutTime - (OscillationDuration - NewTime);
+			CurrentBlendOutTime = OscillationBlendOutTime - (ActualOscillationDuration - NewTime);
 		}
 	}
 
@@ -404,7 +406,7 @@ bool ULegacyCameraShake::ReceiveIsFinished_Implementation() const
 
 bool ULegacyCameraShake::IsLooping() const
 {
-	return OscillationDuration < 0.0f;
+	return ActualOscillationDuration < 0.0f;
 }
 
 void ULegacyCameraShake::SetCurrentTimeAndApplyShake(float NewTime, FMinimalViewInfo& POV)
diff --git a/Engine/Plugins/Cameras/GameplayCameras/Source/GameplayCameras/Private/SimpleCameraShakePattern.cpp b/Engine/Plugins/Cameras/GameplayCameras/Source/GameplayCameras/Private/SimpleCameraShakePattern.cpp
--- a/Engine/Plugins/Cameras/GameplayCameras/Source/GameplayCameras/Private/SimpleCameraShakePattern.cpp
+++ b/Engine/Plugins/Cameras/GameplayCameras/Source/GameplayCameras/Private/SimpleCameraShakePattern.cpp
@@ -21,7 +21,7 @@ void USimpleCameraShakePattern::GetShakePatternInfoImpl(FCameraShakeInfo& OutInf
 
 void USimpleCameraShakePattern::StartShakePatternImpl(const FCameraShakePatternStartParams& Params)
 {
-	State.Start(this);
+	State.Start(this, Params);
 }
 
 bool USimpleCameraShakePattern::IsFinishedImpl() const
diff --git a/Engine/Plugins/Cameras/GameplayCameras/Source/GameplayCameras/Private/Tests/CameraShakeStateTests.cpp b/Engine/Plugins/Cameras/GameplayCameras/Source/GameplayCameras/Private/Tests/CameraShakeStateTests.cpp
--- a/Engine/Plugins/Cameras/GameplayCameras/Source/GameplayCameras/Private/Tests/CameraShakeStateTests.cpp
+++ b/Engine/Plugins/Cameras/GameplayCameras/Source/GameplayCameras/Private/Tests/CameraShakeStateTests.cpp
@@ -99,4 +99,35 @@ bool FCameraShakeStateScrub::RunTest(const FString& Parameters)
 	return true;
 }
 
+IMPLEMENT_SIMPLE_AUTOMATION_TEST(FCameraShakeStateDurationOverride, 
+		"System.Engine.Cameras.ShakeStateDurationOverride", 
+		EAutomationTestFlags::EditorContext | EAutomationTestFlags::EngineFilter)
+bool FCameraShakeStateDurationOverride::RunTest(const FString& Parameters)
+{
+	FCameraShakeInfo Info;
+	Info.Duration = 1.f;
+	Info.BlendIn = 0.2f;
+	Info.BlendOut = 0.2f;
+
+	FCameraShakeState State;
+	State.Start(Info, 0.5f);
+
+	UTEST_EQUAL("Update 1", State.Update(0.1f), 0.5f);
+	UTEST_TRUE("Update 1", State.IsPlaying());
+
+	UTEST_EQUAL("Update 2", State.Update(0.1f), 1.f);
+	UTEST_TRUE("Update 2", State.IsPlaying());
+
+	UTEST_EQUAL("Update 3", State.Update(0.1f), 1.f);
+	UTEST_TRUE("Update 3", State.IsPlaying());
+
+	UTEST_EQUAL("Update 4", State.Update(0.1f), 0.5f);
+	UTEST_TRUE("Update 4", State.IsPlaying());
+
+	UTEST_EQUAL("Update 5", State.Update(0.1f), 0.0f);
+	UTEST_FALSE("Update 5", State.IsPlaying());
+
+	return true;
+}
+
 #undef LOCTEXT_NAMESPACE
diff --git a/Engine/Plugins/Cameras/GameplayCameras/Source/GameplayCameras/Public/LegacyCameraShake.h b/Engine/Plugins/Cameras/GameplayCameras/Source/GameplayCameras/Public/LegacyCameraShake.h
--- a/Engine/Plugins/Cameras/GameplayCameras/Source/GameplayCameras/Public/LegacyCameraShake.h
+++ b/Engine/Plugins/Cameras/GameplayCameras/Source/GameplayCameras/Public/LegacyCameraShake.h
@@ -251,6 +251,9 @@ private:
 
 protected:
 
+	/** Actual oscillation duration. */
+	float ActualOscillationDuration;
+
 	/** Current location sinusoidal offset. */
 	FVector LocSinOffset;
 
diff --git a/Engine/Plugins/MovieScene/TemplateSequence/Source/TemplateSequence/Private/CameraAnimationSequencePlayer.cpp b/Engine/Plugins/MovieScene/TemplateSequence/Source/TemplateSequence/Private/CameraAnimationSequencePlayer.cpp
--- a/Engine/Plugins/MovieScene/TemplateSequence/Source/TemplateSequence/Private/CameraAnimationSequencePlayer.cpp
+++ b/Engine/Plugins/MovieScene/TemplateSequence/Source/TemplateSequence/Private/CameraAnimationSequencePlayer.cpp
@@ -305,7 +305,7 @@ FFrameTime UCameraAnimationSequencePlayer::GetCurrentPosition() const
 	return PlayPosition.GetCurrentPosition();
 }
 
-void UCameraAnimationSequencePlayer::Initialize(UMovieSceneSequence* InSequence, int32 StartOffset)
+void UCameraAnimationSequencePlayer::Initialize(UMovieSceneSequence* InSequence, int32 StartOffset, float DurationOverride)
 {
 	checkf(InSequence, TEXT("Invalid sequence given to player"));
 	
@@ -344,14 +344,33 @@ void UCameraAnimationSequencePlayer::Initialize(UMovieSceneSequence* InSequence,
 		}
 
 		DurationFrames = (EndTime - StartFrame);
+		TotalDurationFrames = DurationFrames;
+		bDurationRequiresLooping = false;
+
+		if (DurationOverride > 0.f)
+		{
+			const FFrameTime DurationOverrideFrames = DurationOverride * DisplayRate;
+			if (DurationOverrideFrames > DurationFrames)
+			{
+				TotalDurationFrames = DurationOverrideFrames;
+				bDurationRequiresLooping = true;
+			}
+			else
+			{
+				DurationFrames = DurationOverrideFrames;
+			}
+		}
 	}
 	else
 	{
 		StartFrame = FFrameNumber(0);
 		DurationFrames = FFrameTime(0);
+		TotalDurationFrames = DurationFrames;
+		bDurationRequiresLooping = false;
 	}
 
 	PlayPosition.Reset(StartFrame);
+	LoopsPlayed = 0;
 
 	UCameraAnimationSequenceSubsystem* Subsystem = UCameraAnimationSequenceSubsystem::GetCameraAnimationSequenceSubsystem(GetWorld());
 	ensureMsgf(Subsystem, TEXT("Unable to locate a valid camera animation sub-system. Camera anim sequences will not play."));
@@ -387,7 +406,8 @@ void UCameraAnimationSequencePlayer::Update(FFrameTime NewPosition)
 	check(RootTemplateInstance.IsValid());
 
 	bool bShouldStop = false;
-	if (bIsLooping)
+
+	if (bIsLooping || bDurationRequiresLooping)
 	{
 		// Unlike the level sequence player, we don't care about making sure to play the last few frames
 		// of the sequence before looping: we can jump straight to the looped time because we know we
@@ -400,6 +420,7 @@ void UCameraAnimationSequencePlayer::Update(FFrameTime NewPosition)
 		{
 			NewPosition -= DurationFrames;
 			PlayPosition.Reset(StartFrame);
+			++LoopsPlayed;
 		}
 	}
 	else
@@ -412,6 +433,17 @@ void UCameraAnimationSequencePlayer::Update(FFrameTime NewPosition)
 		}
 	}
 
+	if (bDurationRequiresLooping)
+	{
+		// If we are only looping a fixed number of times to meet a certain duration override,
+		// we need to check here whether we have reached the end.
+		const FFrameTime ElapsedPlayedFrames = (LoopsPlayed * DurationFrames) + NewPosition;
+		if (ElapsedPlayedFrames >= TotalDurationFrames)
+		{
+			bShouldStop = true;
+		}
+	}
+
 	if (TSharedPtr<FMovieSceneEntitySystemRunner> Runner = RootTemplateInstance.GetRunner())
 	{
 		const FMovieSceneEvaluationRange Range = PlayPosition.PlayTo(NewPosition);
diff --git a/Engine/Plugins/MovieScene/TemplateSequence/Source/TemplateSequence/Private/SequenceCameraShake.cpp b/Engine/Plugins/MovieScene/TemplateSequence/Source/TemplateSequence/Private/SequenceCameraShake.cpp
--- a/Engine/Plugins/MovieScene/TemplateSequence/Source/TemplateSequence/Private/SequenceCameraShake.cpp
+++ b/Engine/Plugins/MovieScene/TemplateSequence/Source/TemplateSequence/Private/SequenceCameraShake.cpp
@@ -68,11 +68,16 @@ void USequenceCameraShakePattern::StartShakePatternImpl(const FCameraShakePatter
 	Player->SetBoundObjectOverride(CameraStandIn);
 
 	// Initialize it and start playing.
-	Player->Initialize(Sequence);
+	float DurationOverride = 0.f;
+	if (Params.bOverrideDuration)
+	{
+		DurationOverride = Params.DurationOverride;
+	}
+	Player->Initialize(Sequence, 0, DurationOverride);
 	Player->Play(bRandomSegment, bRandomSegment);
 
 	// Initialize our state.
-	State.Start(this);
+	State.Start(this, Params);
 }
 
 void USequenceCameraShakePattern::UpdateShakePatternImpl(const FCameraShakePatternUpdateParams& Params, FCameraShakePatternUpdateResult& OutResult)
diff --git a/Engine/Plugins/MovieScene/TemplateSequence/Source/TemplateSequence/Public/CameraAnimationSequencePlayer.h b/Engine/Plugins/MovieScene/TemplateSequence/Source/TemplateSequence/Public/CameraAnimationSequencePlayer.h
--- a/Engine/Plugins/MovieScene/TemplateSequence/Source/TemplateSequence/Public/CameraAnimationSequencePlayer.h
+++ b/Engine/Plugins/MovieScene/TemplateSequence/Source/TemplateSequence/Public/CameraAnimationSequencePlayer.h
@@ -114,8 +114,9 @@ public:
 	 *
 	 * @param InSequence    The sequence to play
 	 * @param StartOffset   The offset to start at, in frames (display rate)
+	 * @param DurationOverride  A duration to use instead of the natural duration of the sequence
 	 */
-	void Initialize(UMovieSceneSequence* InSequence, int32 StartOffset = 0);
+	void Initialize(UMovieSceneSequence* InSequence, int32 StartOffset = 0, float DurationOverride = 0.f);
 
 	/**
 	 * Start playing the sequence
@@ -214,8 +215,17 @@ private:
 	/** The sequence duration in frames */
 	FFrameTime DurationFrames;
 
+	/** The total duration we need to play */
+	FFrameTime TotalDurationFrames;
+
+	/** Accumulated number of loops played so far */
+	uint16 LoopsPlayed = 0; 
+
 	/** Whether we should be looping */
-	bool bIsLooping;
+	bool bIsLooping = false;
+
+	/** Whether we need to loop due to a duration override */
+	bool bDurationRequiresLooping = false;
 
 	/** Movie player status. */
 	TEnumAsByte<EMovieScenePlayerStatus::Type> Status;
diff --git a/Engine/Source/Runtime/Engine/Classes/Camera/CameraModifier_CameraShake.h b/Engine/Source/Runtime/Engine/Classes/Camera/CameraModifier_CameraShake.h
--- a/Engine/Source/Runtime/Engine/Classes/Camera/CameraModifier_CameraShake.h
+++ b/Engine/Source/Runtime/Engine/Classes/Camera/CameraModifier_CameraShake.h
@@ -121,6 +121,11 @@ struct FAddCameraShakeParams
 	 */
 	FOnInitializeCameraShake Initializer;
 
+	/**
+	 * An optional override for the duration of the camera shake.
+	 */
+	TOptional<float> DurationOverride;
+
 	FAddCameraShakeParams()
 		: Scale(1.f), PlaySpace(ECameraShakePlaySpace::CameraLocal), UserPlaySpaceRot(FRotator::ZeroRotator), SourceComponent(nullptr)
 	{}
diff --git a/Engine/Source/Runtime/Engine/Classes/Camera/CameraShakeBase.h b/Engine/Source/Runtime/Engine/Classes/Camera/CameraShakeBase.h
--- a/Engine/Source/Runtime/Engine/Classes/Camera/CameraShakeBase.h
+++ b/Engine/Source/Runtime/Engine/Classes/Camera/CameraShakeBase.h
@@ -26,6 +26,14 @@ struct FCameraShakePatternStartParams
 	/** Whether the camera shake is restarting while playing */
 	UPROPERTY(EditAnywhere, BlueprintReadWrite, Category=CameraShake)
 	bool bIsRestarting = false;
+
+	/** Whether the camera shake's duration is overriden (see DurationOverride) */
+	UPROPERTY(EditAnywhere, BlueprintReadWrite, Category=CameraShake)
+	bool bOverrideDuration = false;
+
+	/** An optional override for the camera shake's duration */
+	UPROPERTY(EditAnywhere, BlueprintReadWrite, Category=CameraShake)
+	float DurationOverride;
 };
 
 /**
@@ -291,11 +299,21 @@ struct FCameraShakeState
 	 * Initialize the state with a shake's info and start playing.
 	 */
 	ENGINE_API void Start(const FCameraShakeInfo& InShakeInfo);
-	
+
+	/**
+	 * Initialize the state with a shake's info and start playing.
+	 */
+	ENGINE_API void Start(const FCameraShakeInfo& InShakeInfo, TOptional<float> InDurationOverride);
+
 	/**
 	 * Initialize the state with a shake's info and start playing.
 	 */
 	ENGINE_API void Start(const UCameraShakePattern* InShakePattern);
+	
+	/**
+	 * Initialize the state with a shake's info and start playing.
+	 */
+	ENGINE_API void Start(const UCameraShakePattern* InShakePattern, const FCameraShakePatternStartParams& InParams);
 
 	/**
 	 * Updates the state with a delta time.
@@ -398,6 +416,27 @@ private:
 	bool bHasBlendOut : 1;
 };
 
+/**
+ * Parameter struct for starting a camera shake.
+ */
+struct FCameraShakeBaseStartParams
+{
+	/** The parent camera manager */
+	TObjectPtr<APlayerCameraManager> CameraManager;
+
+	/** The scale for playing the shake */
+	float Scale = 1.f;
+
+	/** The coordinate system in which to play the shake */
+	ECameraShakePlaySpace PlaySpace = ECameraShakePlaySpace::CameraLocal;
+
+	/** A custom rotation, only used if PlaySpace is UserDefined */
+	FRotator UserPlaySpaceRot = FRotator::ZeroRotator;
+
+	/** An optional override for the camera shake's duration */
+	TOptional<float> DurationOverride;
+};
+
 /**
  * Base class for a camera shake. A camera shake contains a root shake "pattern" which is
  * the object that contains the actual logic driving how the camera is shaken. Keeping the two
@@ -536,6 +575,9 @@ public:
 	/** Starts this camera shake with the given parameters */
 	ENGINE_API void StartShake(APlayerCameraManager* Camera, float Scale, ECameraShakePlaySpace InPlaySpace, FRotator UserPlaySpaceRot = FRotator::ZeroRotator);
 
+	/** Starts this camera shake with the given parameters */
+	ENGINE_API void StartShake(const FCameraShakeBaseStartParams& Params);
+
 	/** Returns whether this camera shake is finished */
 	ENGINE_API bool IsFinished() const;
 
diff --git a/Engine/Source/Runtime/Engine/Classes/Camera/CameraShakeSourceComponent.h b/Engine/Source/Runtime/Engine/Classes/Camera/CameraShakeSourceComponent.h
--- a/Engine/Source/Runtime/Engine/Classes/Camera/CameraShakeSourceComponent.h
+++ b/Engine/Source/Runtime/Engine/Classes/Camera/CameraShakeSourceComponent.h
@@ -18,6 +18,24 @@ enum class ECameraShakeAttenuation : uint8
 	Quadratic
 };
 
+struct FCameraShakeSourceComponentStartParams
+{
+	/** The type of camera shake to create */
+	TSubclassOf<UCameraShakeBase> ShakeClass;
+
+	/* The scale for playing the shake */
+	float Scale = 1.f;
+
+	/** The coordinate system in which to play the shake */
+	ECameraShakePlaySpace PlaySpace = ECameraShakePlaySpace::CameraLocal;
+
+	/** A custom rotation, only used if PlaySpace is UserDefined */
+	FRotator UserPlaySpaceRot = FRotator::ZeroRotator;
+
+	/** An optional override for the camera shake's duration */
+	TOptional<float> DurationOverride;
+};
+
 UCLASS(Blueprintable, meta = (BlueprintSpawnableComponent), MinimalAPI)
 class UCameraShakeSourceComponent : public USceneComponent
 {
@@ -43,6 +61,10 @@ private:
 	TSubclassOf<UCameraShakeBase> PreviousCameraShake;
 #endif
 
+public:
+	/** Starts a new camera shake originating from this source, and apply it on all player controllers */
+	ENGINE_API void StartCameraShake(const FCameraShakeSourceComponentStartParams& Params);
+
 public:
 	/** The attenuation profile for how camera shakes' intensity falls off with distance */
 	UPROPERTY(EditAnywhere, BlueprintReadWrite, Category = Attenuation)
diff --git a/Engine/Source/Runtime/Engine/Classes/Camera/PlayerCameraManager.h b/Engine/Source/Runtime/Engine/Classes/Camera/PlayerCameraManager.h
--- a/Engine/Source/Runtime/Engine/Classes/Camera/PlayerCameraManager.h
+++ b/Engine/Source/Runtime/Engine/Classes/Camera/PlayerCameraManager.h
@@ -17,6 +17,7 @@ class UCameraModifier;
 class UCameraShakeBase;
 class UCameraShakeSourceComponent;
 class ICameraLensEffectInterface;
+struct FAddCameraShakeParams;
 
 DECLARE_DYNAMIC_MULTICAST_DELEGATE_TwoParams(FOnAudioFadeChangeSignature, bool, bFadeOut, float, FadeTime);
 
@@ -827,6 +828,13 @@ public:
 	UFUNCTION(BlueprintCallable, Category = "Camera Shakes")
 	ENGINE_API virtual void StopAllCameraShakesFromSource(UCameraShakeSourceComponent* SourceComponent, bool bImmediately = true);
 
+	/**
+	 * Plays a camera shake on this camera.
+	 * @param Shake - The class of camera shake to play.
+	 * @param Params - The parameters to pass to the camera shake camera modifier.
+	 */
+	ENGINE_API virtual UCameraShakeBase* StartCameraShake(TSubclassOf<UCameraShakeBase> ShakeClass, const FAddCameraShakeParams& Params);
+
 	//
 	//  Camera fades.
 	//
diff --git a/Engine/Source/Runtime/Engine/Private/Camera/CameraModifier_CameraShake.cpp b/Engine/Source/Runtime/Engine/Private/Camera/CameraModifier_CameraShake.cpp
--- a/Engine/Source/Runtime/Engine/Private/Camera/CameraModifier_CameraShake.cpp
+++ b/Engine/Source/Runtime/Engine/Private/Camera/CameraModifier_CameraShake.cpp
@@ -278,7 +278,13 @@ UCameraShakeBase* UCameraModifier_CameraShake::AddCameraShake(TSubclassOf<UCamer
 			}
 
 			// Initialize new shake and add it to the list of active shakes
-			NewInst->StartShake(CameraOwner, Scale, Params.PlaySpace, Params.UserPlaySpaceRot);
+			FCameraShakeBaseStartParams StartParams;
+			StartParams.CameraManager = CameraOwner;
+			StartParams.Scale = Scale;
+			StartParams.PlaySpace = Params.PlaySpace;
+			StartParams.UserPlaySpaceRot = Params.UserPlaySpaceRot;
+			StartParams.DurationOverride = Params.DurationOverride;
+			NewInst->StartShake(StartParams);
 
 			// Look for nulls in the array to replace first -- keeps the array compact
 			bool bReplacedNull = false;
diff --git a/Engine/Source/Runtime/Engine/Private/Camera/CameraShakeBase.cpp b/Engine/Source/Runtime/Engine/Private/Camera/CameraShakeBase.cpp
--- a/Engine/Source/Runtime/Engine/Private/Camera/CameraShakeBase.cpp
+++ b/Engine/Source/Runtime/Engine/Private/Camera/CameraShakeBase.cpp
@@ -57,6 +57,11 @@ FCameraShakeState::FCameraShakeState()
 }
 
 void FCameraShakeState::Start(const FCameraShakeInfo& InShakeInfo)
+{
+	Start(InShakeInfo, TOptional<float>());
+}
+
+void FCameraShakeState::Start(const FCameraShakeInfo& InShakeInfo, TOptional<float> InDurationOverride)
 {
 	const FCameraShakeState PrevState(*this);
 
@@ -65,6 +70,11 @@ void FCameraShakeState::Start(const FCameraShakeInfo& InShakeInfo)
 	bHasBlendIn = ShakeInfo.BlendIn > 0.f;
 	bHasBlendOut = ShakeInfo.BlendOut > 0.f;
 
+	if (InDurationOverride.IsSet())
+	{
+		ShakeInfo.Duration = FCameraShakeDuration(InDurationOverride.GetValue());
+	}
+
 	// Initialize our running state.
 	InitializePlaying();
 
@@ -94,6 +104,19 @@ void FCameraShakeState::Start(const UCameraShakePattern* InShakePattern)
 	Start(Info);
 }
 
+void FCameraShakeState::Start(const UCameraShakePattern* InShakePattern, const FCameraShakePatternStartParams& InParams)
+{
+	check(InShakePattern);
+	FCameraShakeInfo Info;
+	InShakePattern->GetShakePatternInfo(Info);
+	TOptional<float> DurationOverride;
+	if (InParams.bOverrideDuration)
+	{
+		DurationOverride = InParams.DurationOverride;
+	}
+	Start(Info, DurationOverride);
+}
+
 void FCameraShakeState::InitializePlaying()
 {
 	ElapsedTime = 0.f;
@@ -247,6 +270,16 @@ void UCameraShakeBase::GetShakeInfo(FCameraShakeInfo& OutInfo) const
 }
 
 void UCameraShakeBase::StartShake(APlayerCameraManager* Camera, float Scale, ECameraShakePlaySpace InPlaySpace, FRotator UserPlaySpaceRot)
+{
+	FCameraShakeBaseStartParams Params;
+	Params.CameraManager = Camera;
+	Params.Scale = Scale;
+	Params.PlaySpace = InPlaySpace;
+	Params.UserPlaySpaceRot = UserPlaySpaceRot;
+	StartShake(Params);
+}
+
+void UCameraShakeBase::StartShake(const FCameraShakeBaseStartParams& Params)
 {
 	SCOPE_CYCLE_COUNTER(STAT_StartShake);
 
@@ -256,11 +289,11 @@ void UCameraShakeBase::StartShake(APlayerCameraManager* Camera, float Scale, ECa
 
 	// Remember the various settings for this run.
 	// Note that the camera manager can be null, for example in unit tests.
-	CameraManager = Camera;
-	ShakeScale = Scale;
-	PlaySpace = InPlaySpace;
-	UserPlaySpaceMatrix = (InPlaySpace == ECameraShakePlaySpace::UserDefined) ? 
-		FRotationMatrix(UserPlaySpaceRot) : FRotationMatrix::Identity;
+	CameraManager = Params.CameraManager;
+	ShakeScale = Params.Scale;
+	PlaySpace = Params.PlaySpace;
+	UserPlaySpaceMatrix = (Params.PlaySpace == ECameraShakePlaySpace::UserDefined) ? 
+		FRotationMatrix(Params.UserPlaySpaceRot) : FRotationMatrix::Identity;
 
 	const bool bIsRestarting = bIsActive;
 	bIsActive = true;
@@ -270,6 +303,8 @@ void UCameraShakeBase::StartShake(APlayerCameraManager* Camera, float Scale, ECa
 	{
 		FCameraShakePatternStartParams StartParams;
 		StartParams.bIsRestarting = bIsRestarting;
+		StartParams.bOverrideDuration = Params.DurationOverride.IsSet();
+		StartParams.DurationOverride = Params.DurationOverride.Get(0.f);
 		RootShakePattern->StartShakePattern(StartParams);
 	}
 }
@@ -420,11 +455,20 @@ void UCameraShakeBase::ApplyResult(const FCameraShakeApplyResultParams& ApplyPar
 		InOutPOV.FOV += TempResult.FOV;
 	}
 
-	// It's weird but the post-process settings go directly on the camera manager, not on the view info.
-	if (ApplyParams.CameraManager.IsValid() && TempResult.PostProcessBlendWeight > 0.f)
+	if (TempResult.PostProcessBlendWeight > 0.f)
 	{
-		EViewTargetBlendOrder CameraShakeBlendOrder = GCameraShakeLegacyPostProcessBlending.GetValueOnGameThread() ? VTBlendOrder_Base : VTBlendOrder_Override;
-		ApplyParams.CameraManager->AddCachedPPBlend(TempResult.PostProcessSettings, TempResult.PostProcessBlendWeight, CameraShakeBlendOrder);
+		// If we have a camera manager, post-process settings go there. Otherwise, let's put them on
+		// the view-info.
+		if (ApplyParams.CameraManager.IsValid())
+		{
+			EViewTargetBlendOrder CameraShakeBlendOrder = GCameraShakeLegacyPostProcessBlending.GetValueOnGameThread() ? VTBlendOrder_Base : VTBlendOrder_Override;
+			ApplyParams.CameraManager->AddCachedPPBlend(TempResult.PostProcessSettings, TempResult.PostProcessBlendWeight, CameraShakeBlendOrder);
+		}
+		else
+		{
+			InOutPOV.PostProcessSettings = TempResult.PostProcessSettings;
+			InOutPOV.PostProcessBlendWeight = TempResult.PostProcessBlendWeight;
+		}
 	}
 }
 
diff --git a/Engine/Source/Runtime/Engine/Private/Camera/CameraShakeSourceComponent.cpp b/Engine/Source/Runtime/Engine/Private/Camera/CameraShakeSourceComponent.cpp
--- a/Engine/Source/Runtime/Engine/Private/Camera/CameraShakeSourceComponent.cpp
+++ b/Engine/Source/Runtime/Engine/Private/Camera/CameraShakeSourceComponent.cpp
@@ -2,6 +2,7 @@
 
 #include "Camera/CameraShakeSourceComponent.h"
 #include "Camera/CameraShakeBase.h"
+#include "Camera/CameraModifier_CameraShake.h"
 #include "Components/BillboardComponent.h"
 #include "Engine/Texture2D.h"
 #include "Engine/World.h"
@@ -80,15 +81,32 @@ void UCameraShakeSourceComponent::Start()
 }
 
 void UCameraShakeSourceComponent::StartCameraShake(TSubclassOf<UCameraShakeBase> InCameraShake, float Scale, ECameraShakePlaySpace PlaySpace, FRotator UserPlaySpaceRot)
+{
+	FCameraShakeSourceComponentStartParams Params;
+	Params.ShakeClass = InCameraShake;
+	Params.Scale = Scale;
+	Params.PlaySpace = PlaySpace;
+	Params.UserPlaySpaceRot = UserPlaySpaceRot;
+	StartCameraShake(Params);
+}
+
+void UCameraShakeSourceComponent::StartCameraShake(const FCameraShakeSourceComponentStartParams& Params)
 {
 	if (UWorld* World = GetWorld())
 	{
+		FAddCameraShakeParams ModParams;
+		ModParams.SourceComponent = this;
+		ModParams.Scale = Params.Scale;
+		ModParams.PlaySpace = Params.PlaySpace;
+		ModParams.UserPlaySpaceRot = Params.UserPlaySpaceRot;
+		ModParams.DurationOverride = Params.DurationOverride;
+
 		for (FConstPlayerControllerIterator Iterator = World->GetPlayerControllerIterator(); Iterator; ++Iterator)
 		{
 			APlayerController* PlayerController = Iterator->Get();
 			if (PlayerController != nullptr && PlayerController->PlayerCameraManager != nullptr)
 			{
-				PlayerController->PlayerCameraManager->StartCameraShakeFromSource(InCameraShake, this, Scale, PlaySpace, UserPlaySpaceRot);
+				PlayerController->PlayerCameraManager->StartCameraShake(Params.ShakeClass, ModParams);
 			}
 		}
 	}
diff --git a/Engine/Source/Runtime/Engine/Private/PlayerCameraManager.cpp b/Engine/Source/Runtime/Engine/Private/PlayerCameraManager.cpp
--- a/Engine/Source/Runtime/Engine/Private/PlayerCameraManager.cpp
+++ b/Engine/Source/Runtime/Engine/Private/PlayerCameraManager.cpp
@@ -1234,19 +1234,21 @@ void APlayerCameraManager::RemoveCameraLensEffect(AEmitterCameraLensEffectBase*
 
 UCameraShakeBase* APlayerCameraManager::StartCameraShake(TSubclassOf<UCameraShakeBase> ShakeClass, float Scale, ECameraShakePlaySpace PlaySpace, FRotator UserPlaySpaceRot)
 {
-	if (ShakeClass && CachedCameraShakeMod && (Scale > 0.0f) )
-	{
-		return CachedCameraShakeMod->AddCameraShake(ShakeClass, FAddCameraShakeParams(Scale, PlaySpace, UserPlaySpaceRot));
-	}
-
-	return nullptr;
+	FAddCameraShakeParams Params(Scale, PlaySpace, UserPlaySpaceRot);
+	return StartCameraShake(ShakeClass, Params);
 }
 
 UCameraShakeBase* APlayerCameraManager::StartCameraShakeFromSource(TSubclassOf<UCameraShakeBase> ShakeClass, UCameraShakeSourceComponent* SourceComponent, float Scale, ECameraShakePlaySpace PlaySpace, FRotator UserPlaySpaceRot)
+{
+	FAddCameraShakeParams Params(Scale, PlaySpace, UserPlaySpaceRot, SourceComponent);
+	return StartCameraShake(ShakeClass, Params);
+}
+
+UCameraShakeBase* APlayerCameraManager::StartCameraShake(TSubclassOf<UCameraShakeBase> ShakeClass, const FAddCameraShakeParams& Params)
 {
 	if (ShakeClass && CachedCameraShakeMod)
 	{
-		return CachedCameraShakeMod->AddCameraShake(ShakeClass, FAddCameraShakeParams(Scale, PlaySpace, UserPlaySpaceRot, SourceComponent));
+		return CachedCameraShakeMod->AddCameraShake(ShakeClass, Params);
 	}
 
 	return nullptr;
diff --git a/Engine/Source/Runtime/MovieSceneTracks/Private/Evaluation/MovieSceneCameraShakePreviewer.cpp b/Engine/Source/Runtime/MovieSceneTracks/Private/Evaluation/MovieSceneCameraShakePreviewer.cpp
--- a/Engine/Source/Runtime/MovieSceneTracks/Private/Evaluation/MovieSceneCameraShakePreviewer.cpp
+++ b/Engine/Source/Runtime/MovieSceneTracks/Private/Evaluation/MovieSceneCameraShakePreviewer.cpp
@@ -5,13 +5,12 @@
 #if WITH_EDITOR
 
 #include "LevelEditorViewport.h"
+#include "Camera/CameraShakeBase.h"
+#include "Camera/CameraShakeSourceComponent.h"
 #include "Camera/CameraModifier_CameraShake.h"
 
-#include UE_INLINE_GENERATED_CPP_BY_NAME(MovieSceneCameraShakePreviewer)
-
 FCameraShakePreviewer::FCameraShakePreviewer()
-	: PreviewCamera(nullptr)
-	, PreviewCameraShake(nullptr)
+	: LastDeltaTime(0.f)
 	, LastLocationModifier(FVector::ZeroVector)
 	, LastRotationModifier(FRotator::ZeroRotator)
 	, LastFOVModifier(0.f)
@@ -24,48 +23,108 @@ FCameraShakePreviewer::~FCameraShakePreviewer()
 	{
 		UnRegisterViewModifier();
 	}
-
-	Teardown();
 }
 
-void FCameraShakePreviewer::Initialize(UWorld* InWorld)
+UCameraShakeBase* FCameraShakePreviewer::AddCameraShake(const FCameraShakePreviewerAddParams& Params)
 {
-	FActorSpawnParameters SpawnInfo;
-	SpawnInfo.ObjectFlags |= RF_Transient;
-	SpawnInfo.SpawnCollisionHandlingOverride = ESpawnActorCollisionHandlingMethod::AlwaysSpawn;
+	UCameraShakeBase* NewShake = NewObject<UCameraShakeBase>(GetTransientPackage(), Params.ShakeClass);
+	ActiveShakes.Add({ NewShake, Params.SourceComponent, Params.GlobalStartTime });
 
-	PreviewCamera = InWorld->SpawnActor<APreviewPlayerCameraManager>(SpawnInfo);
-	PreviewCamera->SetIsTemporarilyHiddenInEditor(true);
+	FCameraShakeBaseStartParams StartParams;
+	StartParams.Scale = Params.Scale;
+	StartParams.PlaySpace = Params.PlaySpace;
+	StartParams.UserPlaySpaceRot = Params.UserPlaySpaceRot;
+	StartParams.DurationOverride = Params.DurationOverride;
+	NewShake->StartShake(StartParams);
+
+	return NewShake;
+}
 
-	PreviewCameraShake = CastChecked<UCameraModifier_CameraShake>(
-		PreviewCamera->AddNewCameraModifier(UCameraModifier_CameraShake::StaticClass()));
+void FCameraShakePreviewer::RemoveCameraShake(UCameraShakeBase* ShakeInstance)
+{
+	const bool bImmediately = true;
+	for (int32 i = ActiveShakes.Num() - 1; i >= 0; --i)
+	{
+		FPreviewCameraShakeInfo& ActiveShake = ActiveShakes[i];
+		if (ActiveShake.ShakeInstance == ShakeInstance)
+		{
+			ActiveShake.ShakeInstance->StopShake(bImmediately);
+			ActiveShake.ShakeInstance->TeardownShake();
+			ActiveShakes.RemoveAt(i, 1);
+			break;
+		}
+	}
+}
 
-	LastDeltaTime = 0.f;
+void FCameraShakePreviewer::RemoveAllCameraShakesFromSource(const UCameraShakeSourceComponent* SourceComponent)
+{
+	const bool bImmediately = true;
+	for (int32 i = ActiveShakes.Num() - 1; i >= 0; --i)
+	{
+		FPreviewCameraShakeInfo& ActiveShake = ActiveShakes[i];
+		if (ActiveShake.SourceComponent.Get() == SourceComponent && ActiveShake.ShakeInstance != nullptr)
+		{
+			ActiveShake.ShakeInstance->StopShake(bImmediately);
+			ActiveShake.ShakeInstance->TeardownShake();
+			ActiveShakes.RemoveAt(i, 1);
+		}
+	}
 }
 
-void FCameraShakePreviewer::Teardown()
+void FCameraShakePreviewer::RemoveAllCameraShakes()
 {
-	if (PreviewCamera != nullptr)
+	const bool bImmediately = true;
+	for (FPreviewCameraShakeInfo& ActiveShake : ActiveShakes)
 	{
-		PreviewCamera->Destroy(false, false);
+		if (ActiveShake.ShakeInstance)
+		{
+			ActiveShake.ShakeInstance->StopShake(bImmediately);
+			ActiveShake.ShakeInstance->TeardownShake();
+		}
 	}
+	ActiveShakes.Empty();
+}
 
-	PreviewCameraShake = nullptr;
-	PreviewCamera = nullptr;
+void FCameraShakePreviewer::GetActiveCameraShakes(TArray<FActiveCameraShakeInfo>& ActiveCameraShakes) const
+{
+	for (const FPreviewCameraShakeInfo& ActiveShake : ActiveShakes)
+	{
+		FActiveCameraShakeInfo ShakeInfo;
+		ShakeInfo.ShakeInstance = ActiveShake.ShakeInstance;
+		ShakeInfo.ShakeSource = ActiveShake.SourceComponent;
+		ActiveCameraShakes.Add(ShakeInfo);
+	}
 }
 
 void FCameraShakePreviewer::Update(float DeltaTime, bool bIsPlaying)
 {
 	LastDeltaTime = DeltaTime;
+	LastScrubTime.Reset();
 
 	if (!bIsPlaying)
 	{
-		LastLocationModifier = FVector::ZeroVector;
-		LastRotationModifier = FRotator::ZeroRotator;
-		LastFOVModifier = 0.f;
+		ResetModifiers();
 	}
 }
 
+void FCameraShakePreviewer::Scrub(float ScrubTime)
+{
+	LastDeltaTime.Reset();
+	LastScrubTime = ScrubTime;
+
+	ResetModifiers();
+}
+
+void FCameraShakePreviewer::ResetModifiers()
+{
+	LastLocationModifier = FVector::ZeroVector;
+	LastRotationModifier = FRotator::ZeroRotator;
+	LastFOVModifier = 0.f;
+
+	LastPostProcessSettings.Reset();
+	LastPostProcessBlendWeights.Reset();
+}
+
 void FCameraShakePreviewer::ModifyView(FEditorViewportViewModifierParams& Params)
 {
 	OnModifyView(Params);
@@ -73,35 +132,82 @@ void FCameraShakePreviewer::ModifyView(FEditorViewportViewModifierParams& Params
 
 void FCameraShakePreviewer::OnModifyView(FEditorViewportViewModifierParams& Params)
 {
-	const float DeltaTime = LastDeltaTime.Get(-1.f);
-	if (DeltaTime > 0.f)
+	FMinimalViewInfo& InOutPOV(Params.ViewInfo);
+	const FMinimalViewInfo OriginalPOV(Params.ViewInfo);
+
+	// This is a simpler version of what UCameraModifier_CameraShake does, with extra
+	// support for scrubbing.
+	if (LastDeltaTime.IsSet() || LastScrubTime.IsSet())
 	{
 		LastPostProcessSettings.Reset();
 		LastPostProcessBlendWeights.Reset();
-		PreviewCamera->ResetPostProcessSettings();
-
-		FMinimalViewInfo OriginalPOV(Params.ViewInfo);
-
-		PreviewCameraShake->ModifyCamera(DeltaTime, Params.ViewInfo);
 
-		PreviewCamera->MergePostProcessSettings(LastPostProcessSettings, LastPostProcessBlendWeights);
+		for (FPreviewCameraShakeInfo& ActiveShake : ActiveShakes)
+		{
+			if (ActiveShake.ShakeInstance != nullptr)
+			{
+				float CurShakeAlpha = 1.f;
+
+				if (ActiveShake.SourceComponent.IsValid())
+				{
+					const UCameraShakeSourceComponent* SourceComponent = ActiveShake.SourceComponent.Get();
+					const float AttenuationFactor = SourceComponent->GetAttenuationFactor(InOutPOV.Location);
+					CurShakeAlpha *= AttenuationFactor;
+				}
+
+				if (LastDeltaTime.IsSet())
+				{
+					ActiveShake.ShakeInstance->UpdateAndApplyCameraShake(LastDeltaTime.GetValue(), CurShakeAlpha, InOutPOV);
+				}
+				else if (LastScrubTime.IsSet())
+				{
+					float RelativeScrubTime = LastScrubTime.GetValue() - ActiveShake.StartTime;
+					ActiveShake.ShakeInstance->ScrubAndApplyCameraShake(RelativeScrubTime, CurShakeAlpha, InOutPOV);
+				}
+
+				if (InOutPOV.PostProcessBlendWeight > 0.f)
+				{
+					Params.AddPostProcessBlend(InOutPOV.PostProcessSettings, InOutPOV.PostProcessBlendWeight);
+					LastPostProcessSettings.Add(InOutPOV.PostProcessSettings);
+					LastPostProcessBlendWeights.Add(InOutPOV.PostProcessBlendWeight);
+				}
+				InOutPOV.PostProcessSettings = FPostProcessSettings();
+				InOutPOV.PostProcessBlendWeight = 0.f;
+			}
+		}
 
-		LastLocationModifier = Params.ViewInfo.Location - OriginalPOV.Location;
-		LastRotationModifier = Params.ViewInfo.Rotation - OriginalPOV.Rotation;
-		LastFOVModifier = Params.ViewInfo.FOV - OriginalPOV.FOV;
+		LastLocationModifier = InOutPOV.Location - OriginalPOV.Location;
+		LastRotationModifier = InOutPOV.Rotation - OriginalPOV.Rotation;
+		LastFOVModifier = InOutPOV.FOV - OriginalPOV.FOV;
 
 		LastDeltaTime.Reset();
+		LastScrubTime.Reset();
+
+		// Delete any obsolete shakes.
+		for (int32 i = ActiveShakes.Num() - 1; i >= 0; i--)
+		{
+			const FPreviewCameraShakeInfo& ShakeInfo = ActiveShakes[i];
+			if (ShakeInfo.ShakeInstance == nullptr || ShakeInfo.ShakeInstance->IsFinished() || ShakeInfo.SourceComponent.IsStale())
+			{
+				if (ShakeInfo.ShakeInstance != nullptr)
+				{
+					ShakeInfo.ShakeInstance->TeardownShake();
+				}
+
+				ActiveShakes.RemoveAt(i, 1);
+			}
+		}
 	}
 	else
 	{
-		Params.ViewInfo.Location += LastLocationModifier;
-		Params.ViewInfo.Rotation += LastRotationModifier;
-		Params.ViewInfo.FOV += LastFOVModifier;
-	}
+		InOutPOV.Location += LastLocationModifier;
+		InOutPOV.Rotation += LastRotationModifier;
+		InOutPOV.FOV += LastFOVModifier;
 
-	for (int32 PPIndex = 0; PPIndex < LastPostProcessSettings.Num(); ++PPIndex)
-	{
-		Params.AddPostProcessBlend(LastPostProcessSettings[PPIndex], LastPostProcessBlendWeights[PPIndex]);
+		for (int32 PPIndex = 0; PPIndex < LastPostProcessSettings.Num(); ++PPIndex)
+		{
+			Params.AddPostProcessBlend(LastPostProcessSettings[PPIndex], LastPostProcessBlendWeights[PPIndex]);
+		}
 	}
 }
 
@@ -156,5 +262,32 @@ void FCameraShakePreviewer::OnLevelViewportClientListChanged()
 	}
 }
 
+void FCameraShakePreviewer::AddReferencedObjects(FReferenceCollector& Collector)
+{
+	for (FPreviewCameraShakeInfo& ActiveShake : ActiveShakes)
+	{
+		if (ActiveShake.ShakeInstance)
+		{
+			Collector.AddReferencedObject(ActiveShake.ShakeInstance);
+		}
+	}
+}
+
+void FCameraShakePreviewer::OnObjectsReplaced(const TMap<UObject*, UObject*>& ReplacementMap)
+{
+	const bool bImmediately = true;
+	for (int32 i = ActiveShakes.Num() - 1; i >= 0; i--)
+	{
+		FPreviewCameraShakeInfo& ActiveShake = ActiveShakes[i];
+		if (ReplacementMap.Find(ActiveShake.ShakeInstance))
+		{
+			// If a camera shake gets recompiled, we just stop and discard it.
+			ActiveShake.ShakeInstance->StopShake(bImmediately);
+			ActiveShake.ShakeInstance->TeardownShake();
+			ActiveShakes.RemoveAt(i);
+		}
+	}
+}
+
 #endif
 
diff --git a/Engine/Source/Runtime/MovieSceneTracks/Private/Sections/MovieSceneCameraShakeSourceTriggerSection.cpp b/Engine/Source/Runtime/MovieSceneTracks/Private/Sections/MovieSceneCameraShakeSourceTriggerSection.cpp
--- a/Engine/Source/Runtime/MovieSceneTracks/Private/Sections/MovieSceneCameraShakeSourceTriggerSection.cpp
+++ b/Engine/Source/Runtime/MovieSceneTracks/Private/Sections/MovieSceneCameraShakeSourceTriggerSection.cpp
@@ -54,12 +54,7 @@ void UMovieSceneCameraShakeSourceTriggerSection::ImportEntityImpl(UMovieSceneEnt
 	UMovieSceneCameraShakeInstantiatorSystem* ShakeSystem = EntityLinker->LinkSystem<UMovieSceneCameraShakeInstantiatorSystem>();
 	if (ensure(ShakeSystem))
 	{
-		const FSequenceInstance& ThisInstance = EntityLinker->GetInstanceRegistry()->GetInstance(Params.Sequence.InstanceHandle);
-		const FMovieSceneContext& Context = ThisInstance.GetContext();
-		
-		const FFrameTime ShakeTime = Times[ShakeIndex] * Context.GetSequenceToRootSequenceTransform();
-
-		ShakeSystem->AddShakeTrigger(Params.Sequence.InstanceHandle, Params.GetObjectBindingID(), ShakeTime, Shakes[ShakeIndex]);
+		ShakeSystem->AddShakeTrigger(Params.Sequence.InstanceHandle, Params.GetObjectBindingID(), Times[ShakeIndex], Shakes[ShakeIndex]);
 
 		// Mimic the structure changing in order to ensure that the instantiation phase runs
 		EntityLinker->EntityManager.MimicStructureChanged();
diff --git a/Engine/Source/Runtime/MovieSceneTracks/Private/Systems/MovieSceneCameraShakeSystem.cpp b/Engine/Source/Runtime/MovieSceneTracks/Private/Systems/MovieSceneCameraShakeSystem.cpp
--- a/Engine/Source/Runtime/MovieSceneTracks/Private/Systems/MovieSceneCameraShakeSystem.cpp
+++ b/Engine/Source/Runtime/MovieSceneTracks/Private/Systems/MovieSceneCameraShakeSystem.cpp
@@ -18,6 +18,10 @@
 #include "UObject/Object.h"
 #include "UObject/Package.h"
 
+#if WITH_EDITOR
+#include "Editor.h"
+#endif  // WITH_EDITOR
+
 #include UE_INLINE_GENERATED_CPP_BY_NAME(MovieSceneCameraShakeSystem)
 
 namespace UE::MovieScene
@@ -81,6 +85,16 @@ struct FPreAnimatedCameraSourceShakeTraits : FBoundObjectPreAnimatedStateTraits
 		if (UCameraShakeSourceComponent* ShakeSourceComponent = Cast<UCameraShakeSourceComponent>(InKey.ResolveObjectPtr()))
 		{
 			ShakeSourceComponent->StopAllCameraShakes(true);
+#if WITH_EDITOR
+			FCameraShakePreviewerLinkerExtension* PreviewerExtension = Params.Linker->FindExtension<FCameraShakePreviewerLinkerExtension>();
+			if (PreviewerExtension)
+			{
+				if (FCameraShakePreviewer* Previewer = PreviewerExtension->FindPreviewer(Params.TerminalInstanceHandle))
+				{
+					Previewer->RemoveAllCameraShakesFromSource(ShakeSourceComponent);
+				}
+			}
+#endif  // WITH_EDITOR
 		}
 	}
 };
@@ -139,11 +153,7 @@ FCameraShakePreviewerLinkerExtension::~FCameraShakePreviewerLinkerExtension()
 	for (TPair<FInstanceHandle, FCameraShakePreviewer>& Pair : Previewers)
 	{
 		FCameraShakePreviewer& Previewer = Pair.Value;
-		if (ensure(Previewer.IsInitialized()))
-		{
-			Previewer.UnRegisterViewModifier();
-			Previewer.Teardown();
-		}
+		Previewer.UnRegisterViewModifier();
 	}
 	Previewers.Reset();
 }
@@ -157,12 +167,10 @@ FCameraShakePreviewer& FCameraShakePreviewerLinkerExtension::GetPreviewer(FInsta
 {
 	if (FCameraShakePreviewer* Previewer = Previewers.Find(InstanceHandle))
 	{
-		ensure(Previewer->IsInitialized());
 		return *Previewer;
 	}
 
 	FCameraShakePreviewer& NewPreviewer = Previewers.Add(InstanceHandle);
-	NewPreviewer.Initialize(WeakLinker.Get()->GetWorld());
 	NewPreviewer.RegisterViewModifier();
 	return NewPreviewer;
 }
@@ -176,15 +184,16 @@ void FCameraShakePreviewerLinkerExtension::UpdateAllPreviewers(FInstanceRegistry
 		if (FCameraShakePreviewer* Previewer = Previewers.Find(InstanceHandle))
 		{
 			const FMovieSceneContext& Context = It->GetContext();
-			if (Context.HasJumped())
+			const float DeltaTime = Context.GetFrameRate().AsSeconds(Context.GetDelta());
+			if (DeltaTime > 0.f)
 			{
-				Previewer->GetCameraModifier()->RemoveAllCameraShakes(true);
+				const bool bIsPlaying = Context.GetStatus() == EMovieScenePlayerStatus::Playing;
+				Previewer->Update(DeltaTime, bIsPlaying);
 			}
 			else
 			{
-				const float DeltaTime = Context.GetFrameRate().AsSeconds(Context.GetDelta());
-				const bool bIsPlaying = Context.GetStatus() == EMovieScenePlayerStatus::Playing;
-				Previewer->Update(DeltaTime, bIsPlaying);
+				const float ScrubTime = Context.GetFrameRate().AsSeconds(Context.GetTime());
+				Previewer->Scrub(ScrubTime);
 			}
 		}
 	}
@@ -195,9 +204,7 @@ bool FCameraShakePreviewerLinkerExtension::HasAnyShake() const
 	TArray<FActiveCameraShakeInfo> TempCameraShakes;
 	for (const TPair<FInstanceHandle, FCameraShakePreviewer>& Pair : Previewers)
 	{
-		const UCameraModifier_CameraShake* CameraModifier = Pair.Value.GetCameraModifier();
-		CameraModifier->GetActiveCameraShakes(TempCameraShakes);
-		if (TempCameraShakes.Num() > 0)
+		if (Pair.Value.NumActiveCameraShakes() > 0)
 		{
 			return true;
 		}
@@ -238,8 +245,12 @@ void UMovieSceneCameraShakeInstantiatorSystem::OnLink()
 	PreAnimatedCameraSourceShakeStorage = Linker->PreAnimatedState.GetOrCreateStorage<FPreAnimatedCameraSourceShakeStateStorage>();
 
 #if WITH_EDITOR
-	PreviewerExtension = FCameraShakePreviewerLinkerExtension::GetOrCreateExtension(Linker);
-#endif
+	if (GEditor != nullptr)
+	{
+		// We only need the previewer extension if there's an actual editor.
+		PreviewerExtension = FCameraShakePreviewerLinkerExtension::GetOrCreateExtension(Linker);
+	}
+#endif  // WITH_EDITOR
 }
 
 void UMovieSceneCameraShakeInstantiatorSystem::OnUnlink()
@@ -250,7 +261,7 @@ void UMovieSceneCameraShakeInstantiatorSystem::OnUnlink()
 	// Only the two camera shake systems hold pointers to the extension, so it should delete itself
 	// once both systems are unlinked.
 	PreviewerExtension = nullptr;
-#endif
+#endif  // WITH_EDITOR
 
 	if (!ensure(TriggersByInstance.Num() == 0))
 	{
@@ -270,10 +281,10 @@ void UMovieSceneCameraShakeInstantiatorSystem::OnRun(FSystemTaskPrerequisites& I
 
 	// Create camera shake instances for new shakes, and start them.
 	auto VisitNewShakes = [this, BuiltInComponents, &EntityManager, InstanceRegistry](
-			FMovieSceneEntityID EntityID,
-			FInstanceHandle InstanceHandle, 
-			UObject* BoundObject, 
-			FMovieSceneCameraShakeComponentData& ShakeData)
+		FMovieSceneEntityID EntityID,
+		FInstanceHandle InstanceHandle,
+		UObject* BoundObject,
+		FMovieSceneCameraShakeComponentData& ShakeData)
 	{
 		const FSequenceInstance& Instance = InstanceRegistry->GetInstance(InstanceHandle);
 		const FMovieSceneContext& Context = Instance.GetContext();
@@ -281,7 +292,7 @@ void UMovieSceneCameraShakeInstantiatorSystem::OnRun(FSystemTaskPrerequisites& I
 
 		TSubclassOf<UCameraShakeBase> ShakeClass = ShakeData.SectionData.ShakeClass;
 		UCameraShakeSourceComponent* ShakeSourceComponent = Cast<UCameraShakeSourceComponent>(BoundObject);
-		
+
 		if (ShakeClass.Get() == nullptr)
 		{
 			if (ShakeSourceComponent)
@@ -294,19 +305,26 @@ void UMovieSceneCameraShakeInstantiatorSystem::OnRun(FSystemTaskPrerequisites& I
 			return;
 		}
 
-		// Get the duration of the shake and store it in the component data.
+		// Get the duration of the shake and compare it to the duration of the section, to know
+		// if we need to override it.
+		const FFrameRate& FrameRate = Context.GetFrameRate();
+		TOptional<float> DurationOverride;
+
 		FCameraShakeDuration ShakeDuration;
 		UCameraShakeBase::GetCameraShakeDuration(ShakeClass, ShakeDuration);
-		ShakeData.Duration = ShakeDuration.IsFixed() ? 
-			Context.GetFrameRate().AsFrameTime(ShakeDuration.Get()) :
-			FFrameTime(-1);
-
-		// Get the blend out duration and also store it in the instance data.
-		float BlendIn = 0.f, BlendOut = 0.f;
-		UCameraShakeBase::GetCameraShakeBlendTimes(ShakeClass, BlendIn, BlendOut);
-		ShakeData.BlendOutTime = (BlendOut > 0) ?
-			Context.GetFrameRate().AsFrameTime(BlendOut) :
-			FFrameTime(0);
+		const FFrameTime SectionDurationFrames = (ShakeData.SectionEndTime - ShakeData.SectionStartTime);
+		if (ShakeDuration.IsFixed())
+		{
+			const FFrameTime ShakeDurationFrames = Context.GetFrameRate().AsFrameTime(ShakeDuration.Get());
+			if (ShakeDurationFrames > SectionDurationFrames)
+			{
+				DurationOverride = FrameRate.AsSeconds(SectionDurationFrames);
+			}
+		}
+		else
+		{
+			DurationOverride = FrameRate.AsSeconds(SectionDurationFrames);
+		}
 
 		const bool bWantsRestoreState = EntityManager.HasComponent(EntityID, BuiltInComponents->Tags.RestoreState);
 		const FRootInstanceHandle RootInstanceHandle = Instance.GetRootInstanceHandle();
@@ -318,27 +336,34 @@ void UMovieSceneCameraShakeInstantiatorSystem::OnRun(FSystemTaskPrerequisites& I
 			PreAnimatedCameraSourceShakeStorage->BeginTrackingEntity(EntityID, bWantsRestoreState, RootInstanceHandle, ShakeSourceComponent);
 			PreAnimatedCameraSourceShakeStorage->CachePreAnimatedValue(CacheParams, ShakeSourceComponent);
 
-			ShakeSourceComponent->StartCameraShake(
-					ShakeClass, 
-					ShakeData.SectionData.PlayScale, 
-					ShakeData.SectionData.PlaySpace,
-					ShakeData.SectionData.UserDefinedPlaySpace);
+			FCameraShakeSourceComponentStartParams ComponentParams;
+			ComponentParams.ShakeClass = ShakeClass;
+			ComponentParams.Scale = ShakeData.SectionData.PlayScale;
+			ComponentParams.PlaySpace = ShakeData.SectionData.PlaySpace;
+			ComponentParams.UserPlaySpaceRot = ShakeData.SectionData.UserDefinedPlaySpace;
+			ComponentParams.DurationOverride = DurationOverride;
+			ShakeSourceComponent->StartCameraShake(ComponentParams);
 
 #if WITH_EDITOR
-			// Shake source components start shakes in the world, unlike the other shakes
-			// (in the `else` clause) who directly affect the bound camera. This means that
-			// the shake we have just started won't affect the Sequencer preview unless we
-			// add some shaking ourselves. Let's do that here.
-			FCameraShakePreviewer& Previewer = PreviewerExtension->GetPreviewer(InstanceHandle);
-			UCameraModifier_CameraShake* const PreviewCameraShake = Previewer.GetCameraModifier();
-			
-			FAddCameraShakeParams Params;
-			Params.SourceComponent = ShakeSourceComponent;
-			Params.Scale = ShakeData.SectionData.PlayScale;
-			Params.PlaySpace = ShakeData.SectionData.PlaySpace;
-			Params.UserPlaySpaceRot = ShakeData.SectionData.UserDefinedPlaySpace;
-			PreviewCameraShake->AddCameraShake(ShakeClass, Params);
-#endif
+			if (PreviewerExtension)
+			{
+				// Shake source components start shakes in the world, unlike the other shakes
+				// (in the `else` clause) who directly affect the bound camera. This means that
+				// the shake we have just started won't affect the Sequencer preview unless we
+				// add some shaking ourselves. Let's do that here.
+				FCameraShakePreviewer& Previewer = PreviewerExtension->GetPreviewer(InstanceHandle);
+
+				FCameraShakePreviewerAddParams PreviewParams;
+				PreviewParams.ShakeClass = ShakeClass;
+				PreviewParams.GlobalStartTime = FrameRate.AsSeconds(ShakeData.SectionStartTime);
+				PreviewParams.SourceComponent = ShakeSourceComponent;
+				PreviewParams.Scale = ShakeData.SectionData.PlayScale;
+				PreviewParams.PlaySpace = ShakeData.SectionData.PlaySpace;
+				PreviewParams.UserPlaySpaceRot = ShakeData.SectionData.UserDefinedPlaySpace;
+				PreviewParams.DurationOverride = DurationOverride;
+				Previewer.AddCameraShake(PreviewParams);
+			}
+#endif  // WITH_EDITOR
 		}
 		else if (UCameraComponent* CameraComponent = MovieSceneHelpers::CameraComponentFromRuntimeObject(BoundObject))
 		{
@@ -351,25 +376,24 @@ void UMovieSceneCameraShakeInstantiatorSystem::OnRun(FSystemTaskPrerequisites& I
 			PreAnimatedCameraComponentShakeStorage->BeginTrackingEntity(EntityID, bWantsRestoreState, RootInstanceHandle, CameraComponent);
 			PreAnimatedCameraComponentShakeStorage->CachePreAnimatedValue(CacheParams, CameraComponent);
 
-			ShakeInstance->StartShake(
-					nullptr, 
-					ShakeData.SectionData.PlayScale,
-					ShakeData.SectionData.PlaySpace,
-					ShakeData.SectionData.UserDefinedPlaySpace);
+			FCameraShakeBaseStartParams ShakeParams;
+			ShakeParams.Scale = ShakeData.SectionData.PlayScale;
+			ShakeParams.PlaySpace = ShakeData.SectionData.PlaySpace;
+			ShakeParams.UserPlaySpaceRot = ShakeData.SectionData.UserDefinedPlaySpace;
+			ShakeParams.DurationOverride = DurationOverride;
+			ShakeInstance->StartShake(ShakeParams);
 
 			ShakeData.ShakeInstance = ShakeInstance;
 		}
-
-		ShakeData.Status = EMovieSceneCameraShakeStatus::Started;
 	};
 
 	FEntityTaskBuilder()
-	.ReadEntityIDs()
-	.Read(BuiltInComponents->InstanceHandle)
-	.Read(BuiltInComponents->BoundObject)
-	.Write(TrackComponents->CameraShake)
-	.FilterAll({ BuiltInComponents->Tags.NeedsLink })
-	.Iterate_PerEntity(&Linker->EntityManager, VisitNewShakes);
+		.ReadEntityIDs()
+		.Read(BuiltInComponents->InstanceHandle)
+		.Read(BuiltInComponents->BoundObject)
+		.Write(TrackComponents->CameraShake)
+		.FilterAll({ BuiltInComponents->Tags.NeedsLink })
+		.Iterate_PerEntity(&Linker->EntityManager, VisitNewShakes);
 
 	// We don't need to go over expired shakes from NeedsUnlink entities. Either these shakes
 	// will be stopped be the pre-animated state, and then their component data will be freed
@@ -399,6 +423,12 @@ void UMovieSceneCameraShakeInstantiatorSystem::TriggerOneShotShakes()
 		const FSequenceInstance& Instance = InstanceRegistry->GetInstance(Pair.Key);
 		IMovieScenePlayer* Player = Instance.GetPlayer();
 
+		const FMovieSceneContext& Context = Instance.GetContext();
+		if (Context.GetDirection() != EPlayDirection::Forwards)
+		{
+			return;
+		}
+
 		for (const FTimedTrigger& Trigger : Pair.Value)
 		{
 			TArray<UCameraShakeSourceComponent*> ShakeSourceComponents;
@@ -416,6 +446,8 @@ void UMovieSceneCameraShakeInstantiatorSystem::TriggerOneShotShakes()
 				}
 			}
 
+			const float TriggerTime = Context.GetFrameRate().AsSeconds(Trigger.Time);
+
 			for (UCameraShakeSourceComponent* ShakeSourceComponent : ShakeSourceComponents)
 			{
 				TSubclassOf<UCameraShakeBase> ShakeClass = Trigger.Trigger.ShakeClass;
@@ -428,27 +460,45 @@ void UMovieSceneCameraShakeInstantiatorSystem::TriggerOneShotShakes()
 				{
 					// Start playing the shake.
 					ShakeSourceComponent->StartCameraShake(
-							ShakeClass, 
-							Trigger.Trigger.PlayScale, 
-							Trigger.Trigger.PlaySpace, 
-							Trigger.Trigger.UserDefinedPlaySpace);
+						ShakeClass,
+						Trigger.Trigger.PlayScale,
+						Trigger.Trigger.PlaySpace,
+						Trigger.Trigger.UserDefinedPlaySpace);
 
 #if WITH_EDITOR
-					// Also start playing the shake in our editor preview.
-					FCameraShakePreviewer& Previewer = PreviewerExtension->GetPreviewer(Pair.Key);
-					UCameraModifier_CameraShake* const PreviewCameraShake = Previewer.GetCameraModifier();
-
-					FAddCameraShakeParams Params;
-					Params.SourceComponent = ShakeSourceComponent;
-					Params.Scale = Trigger.Trigger.PlayScale;
-					Params.PlaySpace = Trigger.Trigger.PlaySpace;
-					Params.UserPlaySpaceRot = Trigger.Trigger.UserDefinedPlaySpace;
-					PreviewCameraShake->AddCameraShake(ShakeClass, Params);
-#endif
+					if (PreviewerExtension)
+					{
+						// Also start playing the shake in our editor preview.
+						FCameraShakePreviewer& Previewer = PreviewerExtension->GetPreviewer(Pair.Key);
+
+						FCameraShakePreviewerAddParams PreviewParams;
+						PreviewParams.ShakeClass = ShakeClass;
+						PreviewParams.GlobalStartTime = TriggerTime;
+						PreviewParams.SourceComponent = ShakeSourceComponent;
+						PreviewParams.Scale = Trigger.Trigger.PlayScale;
+						PreviewParams.PlaySpace = Trigger.Trigger.PlaySpace;
+						PreviewParams.UserPlaySpaceRot = Trigger.Trigger.UserDefinedPlaySpace;
+						Previewer.AddCameraShake(PreviewParams);
+					}
+#endif  // WITH_EDITOR
 				}
 			}
 		}
 	}
+
+#if WITH_EDITOR
+	if (PreviewerExtension && !TriggersByInstance.IsEmpty())
+	{
+		// If we have just started new shakes, we might need to forcibly link the shake evaluator system
+		// so that it can keep the shake previewer ticking every frame (in case that system isn't linked
+		// because there are no shake sections active right now).
+		// Once linked, this system will stay relevant and alive for as long as there are active shakes
+		// in the previewer.
+		Linker->LinkSystem<UMovieSceneCameraShakeEvaluatorSystem>();
+	}
+#endif
+
+	TriggersByInstance.Empty();
 }
 
 namespace UE::MovieScene
@@ -499,18 +549,11 @@ private:
 struct FEvaluateCameraShake
 {
 	const FInstanceRegistry* InstanceRegistry;
-#if WITH_EDITOR
-	FCameraShakePreviewerLinkerExtension* PreviewerExtension;
-#endif
-
 	TMap<UCameraComponent*, FAccumulatedShake> AccumulatedShakes;
 
 	FEvaluateCameraShake(UMovieSceneEntitySystemLinker* InLinker)
 		: InstanceRegistry(InLinker->GetInstanceRegistry())
 	{
-#if WITH_EDITOR
-		PreviewerExtension = InLinker->FindExtension<FCameraShakePreviewerLinkerExtension>();
-#endif
 	}
 
 	void ForEachAllocation(const FEntityAllocation* Allocation, TRead<FMovieSceneEntityID> EntityIDs, TRead<FRootInstanceHandle> RootInstanceHandles, TRead<FInstanceHandle> InstanceHandles, TRead<UObject*> BoundObjects, TWrite<FMovieSceneCameraShakeComponentData> ShakeComponents)
@@ -520,125 +563,31 @@ struct FEvaluateCameraShake
 		{
 			FRootInstanceHandle InstanceHandle = RootInstanceHandles[Index];
 			const FSequenceInstance& Instance = InstanceRegistry->GetInstance(InstanceHandle);
-			const FMovieSceneContext& Context = Instance.GetContext();
-			IMovieScenePlayer* Player = Instance.GetPlayer();
-
-			// If we have jumped, we could have ended up anywhere, so we need to start fresh without
-			// any running camera shake previews.
-#if WITH_EDITOR
-			if (Context.HasJumped())
-			{
-				if (FCameraShakePreviewer* Previewer = PreviewerExtension->FindPreviewer(InstanceHandle))
-				{
-					Previewer->GetCameraModifier()->RemoveAllCameraShakes();
-				}
-			}
-#endif
 
 			// Shakes should have been started by the instantiator system.
-			FMovieSceneCameraShakeComponentData& ShakeData = ShakeComponents[Index];
-			ensure(ShakeData.Status != EMovieSceneCameraShakeStatus::NotStarted);
-
-			// Let's see what this camera shake should be doing now.
-			EMovieSceneCameraShakeStatus DesiredNewStatus = EMovieSceneCameraShakeStatus::Started;
-			
-			// See if there is any blend-out and/or end time that we need to watch out for.
-			// There can be 3 situations here:
 			//
-			//   1. The shake has a duration, and our original section is long enough that the shake will
-			//		blend out and finish naturally on its own. In this case we have nothing to do 
-			//		except update our internal status.
-			//	 2. The shake has a duration, but our original section's size is cutting this short. We 
-			//	    will want to start making the shake blend out manually, or stop it abruptely if
-			//	    it doesn't have any blend-out time.
-			//	 3. The shake has no duration, so we need to make it blend out manually near the end of
-			//	    our source section, or end it abruptely at the end of our source section if it 
-			//	    doesn't have any blend-out time.
-			//
-			// Cases 2 and 3 require us to blend-out/stop the shake ourselves.
-			// Let's see which case we are in.
-			const bool bHasDuration = ShakeData.Duration > 0;
-			const bool bHasBlendOut = ShakeData.BlendOutTime > 0;
-
-			const bool bNeedsManualStop = 
-				// Case 3
-				!bHasDuration ||
-				// Case 2
-				(ShakeData.SectionStartTime + ShakeData.Duration) > ShakeData.SectionEndTime;
-
-			if (bNeedsManualStop)
-			{
-				// Let's see if we have reached the time when we need to start blending out, or the
-				// time we need to flat out finish.
-				if (bHasDuration && Context.GetTime() >= ShakeData.SectionStartTime + ShakeData.Duration)
-				{
-					DesiredNewStatus = EMovieSceneCameraShakeStatus::Finished;
-				}
-				else if (bHasBlendOut && Context.GetTime() >= ShakeData.SectionEndTime - ShakeData.BlendOutTime)
-				{
-					DesiredNewStatus = EMovieSceneCameraShakeStatus::BlendingOut;
-				}
-			}
-
-			if (UCameraShakeSourceComponent* ShakeSourceComponent = Cast<UCameraShakeSourceComponent>(BoundObjects[Index]))
-			{
-				EvaluateShakeSourceComponentShake(ShakeSourceComponent, ShakeData, Instance, DesiredNewStatus);
-			}
-			else if (UCameraComponent* CameraComponent = MovieSceneHelpers::CameraComponentFromRuntimeObject(BoundObjects[Index]))
+			// We don't need to evaluate source components' camera shakes here, as they are ticking
+			// along by themselves in both the player camera manager (in the game) and the camera shake
+			// previewer (in the editor). We do however need to tick the camera shakes running directly
+			// onto camera component bindings.
+			FMovieSceneCameraShakeComponentData& ShakeData = ShakeComponents[Index];
+			if (UCameraComponent* CameraComponent = MovieSceneHelpers::CameraComponentFromRuntimeObject(BoundObjects[Index]))
 			{
-				EvaluateCameraComponentShake(CameraComponent, ShakeData, Instance, DesiredNewStatus);
+				EvaluateCameraComponentShake(CameraComponent, ShakeData, Instance);
 			}
-			ShakeData.Status = DesiredNewStatus;
-		}
-	}
-
-	void EvaluateShakeSourceComponentShake(
-			UCameraShakeSourceComponent* ShakeSourceComponent,
-			FMovieSceneCameraShakeComponentData& ShakeData,
-			const FSequenceInstance& Instance,
-			EMovieSceneCameraShakeStatus DesiredNewStatus)
-	{
-		// We don't need to evaluate the source component's camera shakes here, as they are ticking
-		// along by themselves in both the player camera manager (in the game) and the camera shake
-		// previewer (in the editor).
-		// All we need to do is intervene to stop/blend-out the shake if needed.
-
-		if (ShakeData.Status == EMovieSceneCameraShakeStatus::Started &&
-			(DesiredNewStatus == EMovieSceneCameraShakeStatus::BlendingOut ||
-			 DesiredNewStatus == EMovieSceneCameraShakeStatus::Finished))
-		{
-			const bool bImmediately = (DesiredNewStatus == EMovieSceneCameraShakeStatus::Finished);
-
-			// TODO-ludovic: this isn't exactly correct...
-			// We could be stopping other shakes of the same type started by other means.. but doing
-			// the correct thing would require storing multiple weak shake instance pointers mapped 
-			// to multiple player controllers, themselves mapped to multiple bound objects. Let's only
-			// do that if we run into the (quite unlikely) case where we need it?
-			ShakeSourceComponent->StopAllCameraShakesOfType(ShakeData.SectionData.ShakeClass, bImmediately);
-
-#if WITH_EDITOR
-			FCameraShakePreviewer& Previewer = PreviewerExtension->GetPreviewer(Instance.GetInstanceHandle());
-			UCameraModifier_CameraShake* const PreviewCameraShake = Previewer.GetCameraModifier();
-			PreviewCameraShake->RemoveAllCameraShakesOfClassFromSource(
-				ShakeData.SectionData.ShakeClass, ShakeSourceComponent, bImmediately);
-#endif
 		}
 	}
 
 	void EvaluateCameraComponentShake(
 			UCameraComponent* CameraComponent,
 			FMovieSceneCameraShakeComponentData& ShakeData,
-			const FSequenceInstance& Instance,
-			EMovieSceneCameraShakeStatus DesiredNewStatus)
+			const FSequenceInstance& Instance)
 	{
 		FMinimalViewInfo POV;
 		POV.Location = CameraComponent->GetComponentLocation();
 		POV.Rotation = CameraComponent->GetComponentRotation();
 		POV.FOV = CameraComponent->FieldOfView;
 
-		float PostProcessBlendWeight = 0.f;
-		FPostProcessSettings PostProcessSettings;
-
 		// Update shake to the new time.
 		const FMovieSceneContext& Context = Instance.GetContext();
 		const FFrameTime NewShakeTime = Context.GetTime() - ShakeData.SectionStartTime;
@@ -661,10 +610,10 @@ struct FEvaluateCameraShake
 		}
 
 		// Grab post process changes.
-		if (PostProcessBlendWeight > 0.f)
+		if (POV.PostProcessBlendWeight > 0.f)
 		{
 			FAccumulatedShake& AccumulatedShake = AccumulatedShakes.FindOrAdd(CameraComponent);
-			AccumulatedShake.AccumulatePostProcessing(PostProcessSettings, PostProcessBlendWeight);
+			AccumulatedShake.AccumulatePostProcessing(POV.PostProcessSettings, POV.PostProcessBlendWeight);
 		}
 	}
 
@@ -692,36 +641,7 @@ UMovieSceneCameraShakeEvaluatorSystem::UMovieSceneCameraShakeEvaluatorSystem(con
 	}
 }
 
-void UMovieSceneCameraShakeEvaluatorSystem::OnRun(FSystemTaskPrerequisites& InPrerequisites, FSystemSubsequentTasks& Subsequents)
-{
-	using namespace UE::MovieScene;
-
-	const FBuiltInComponentTypes* BuiltInComponents = FBuiltInComponentTypes::Get();
-	const FMovieSceneTracksComponentTypes* TrackComponents = FMovieSceneTracksComponentTypes::Get();
-
-	FEntityTaskBuilder()
-	.ReadEntityIDs()
-	.Read(BuiltInComponents->RootInstanceHandle)
-	.Read(BuiltInComponents->InstanceHandle)
-	.Read(BuiltInComponents->BoundObject)
-	.Write(TrackComponents->CameraShake)
-	.SetDesiredThread(Linker->EntityManager.GetGatherThread())
-	.Dispatch_PerAllocation<FEvaluateCameraShake>(
-		&Linker->EntityManager, InPrerequisites, &Subsequents, Linker);
-}
-
-UMovieSceneCameraShakePreviewerEvaluatorSystem::UMovieSceneCameraShakePreviewerEvaluatorSystem(const FObjectInitializer& ObjInit)
-	: Super(ObjInit)
-{
-	using namespace UE::MovieScene;
-
-	if (HasAnyFlags(RF_ClassDefaultObject))
-	{
-		DefineImplicitPrerequisite(UMovieSceneCameraShakeEvaluatorSystem::StaticClass(), GetClass());
-	}
-}
-
-bool UMovieSceneCameraShakePreviewerEvaluatorSystem::IsRelevantImpl(UMovieSceneEntitySystemLinker* InLinker) const
+bool UMovieSceneCameraShakeEvaluatorSystem::IsRelevantImpl(UMovieSceneEntitySystemLinker* InLinker) const
 {
 	using namespace UE::MovieScene;
 
@@ -730,20 +650,20 @@ bool UMovieSceneCameraShakePreviewerEvaluatorSystem::IsRelevantImpl(UMovieSceneE
 	{
 		return Extension->HasAnyShake();
 	}
-#endif
+#endif  // WITH_EDITOR
 	return false;
 }
 
-void UMovieSceneCameraShakePreviewerEvaluatorSystem::OnLink()
+void UMovieSceneCameraShakeEvaluatorSystem::OnLink()
 {
 	using namespace UE::MovieScene;
 
 #if WITH_EDITOR
 	PreviewerExtension = FCameraShakePreviewerLinkerExtension::GetOrCreateExtension(Linker);
-#endif
+#endif  // WITH_EDITOR
 }
 
-void UMovieSceneCameraShakePreviewerEvaluatorSystem::OnUnlink()
+void UMovieSceneCameraShakeEvaluatorSystem::OnUnlink()
 {
 	using namespace UE::MovieScene;
 
@@ -751,16 +671,35 @@ void UMovieSceneCameraShakePreviewerEvaluatorSystem::OnUnlink()
 	// Only the two camera shake systems hold pointers to the extension, so it should delete itself
 	// once both systems are unlinked.
 	PreviewerExtension = nullptr;
-#endif
+#endif  // WITH_EDITOR
 }
 
-void UMovieSceneCameraShakePreviewerEvaluatorSystem::OnRun(FSystemTaskPrerequisites& InPrerequisites, FSystemSubsequentTasks& Subsequents)
+void UMovieSceneCameraShakeEvaluatorSystem::OnRun(FSystemTaskPrerequisites& InPrerequisites, FSystemSubsequentTasks& Subsequents)
 {
 	using namespace UE::MovieScene;
 
+	const FBuiltInComponentTypes* BuiltInComponents = FBuiltInComponentTypes::Get();
+	const FMovieSceneTracksComponentTypes* TrackComponents = FMovieSceneTracksComponentTypes::Get();
+
+	FEntityTaskBuilder()
+		.ReadEntityIDs()
+		.Read(BuiltInComponents->RootInstanceHandle)
+		.Read(BuiltInComponents->InstanceHandle)
+		.Read(BuiltInComponents->BoundObject)
+		.Write(TrackComponents->CameraShake)
+		.SetDesiredThread(Linker->EntityManager.GetGatherThread())
+		.Dispatch_PerAllocation<FEvaluateCameraShake>(
+			&Linker->EntityManager, InPrerequisites, &Subsequents, Linker);
+
 #if WITH_EDITOR
-	FInstanceRegistry* InstanceRegistry = Linker->GetInstanceRegistry();
-	PreviewerExtension->UpdateAllPreviewers(InstanceRegistry);
-#endif
+	if (PreviewerExtension)
+	{
+		// The previewer only stores the delta time, and only computes shake results when the editor
+		// later processes viewports. We can therefore safely do this in parallel with the shake
+		// evaluation above.
+		FInstanceRegistry* InstanceRegistry = Linker->GetInstanceRegistry();
+		PreviewerExtension->UpdateAllPreviewers(InstanceRegistry);
+	}
+#endif  // WITH_EDITOR
 }
 
diff --git a/Engine/Source/Runtime/MovieSceneTracks/Public/Evaluation/MovieSceneCameraShakePreviewer.h b/Engine/Source/Runtime/MovieSceneTracks/Public/Evaluation/MovieSceneCameraShakePreviewer.h
--- a/Engine/Source/Runtime/MovieSceneTracks/Public/Evaluation/MovieSceneCameraShakePreviewer.h
+++ b/Engine/Source/Runtime/MovieSceneTracks/Public/Evaluation/MovieSceneCameraShakePreviewer.h
@@ -4,74 +4,86 @@
 
 #if WITH_EDITOR
 
+#include "Camera/CameraTypes.h"
 #include "CoreMinimal.h"
 #include "UObject/GCObject.h"
-#include "Camera/PlayerCameraManager.h"
-
-#include "MovieSceneCameraShakePreviewer.generated.h"
 
 class FLevelEditorViewportClient;
-class UCameraModifier_CameraShake;
+class UCameraShakeBase;
+class UCameraShakeSourceComponent;
+struct FActiveCameraShakeInfo;
 struct FEditorViewportViewModifierParams;
 
-/**
- * A mock player camera manager, used to store camera shake post-process settings for previewing.
- */
-UCLASS()
-class APreviewPlayerCameraManager : public APlayerCameraManager
+struct FCameraShakePreviewerAddParams
 {
-	GENERATED_BODY()
+	// The class of the shake.
+	TSubclassOf<UCameraShakeBase> ShakeClass;
 
-public:
-	void ResetPostProcessSettings()
-	{
-		ClearCachedPPBlends();
-	}
+	// Optional shake source.
+	TObjectPtr<const UCameraShakeSourceComponent> SourceComponent;
 
-	void MergePostProcessSettings(TArray<FPostProcessSettings>& InSettings, TArray<float>& InBlendWeights)
-	{
-		InSettings.Append(PostProcessBlendCache);
-		InBlendWeights.Append(PostProcessBlendCacheWeights);
-	}
+	// Start time of the shake, for scrubbing.
+	float GlobalStartTime;
+
+	// Parameters to be passed to the shake's start method.
+	float Scale = 1.f;
+	ECameraShakePlaySpace PlaySpace = ECameraShakePlaySpace::CameraLocal;
+	FRotator UserPlaySpaceRot = FRotator::ZeroRotator;
+	TOptional<float> DurationOverride;
 };
 
 /**
  * A class that owns a gameplay camera shake manager, so that we can us it to preview shakes in editor.
  */
-class FCameraShakePreviewer : public FGCObject
+class MOVIESCENETRACKS_API FCameraShakePreviewer : public FGCObject
 {
 public:
-	MOVIESCENETRACKS_API FCameraShakePreviewer();
-	MOVIESCENETRACKS_API ~FCameraShakePreviewer();
+	FCameraShakePreviewer();
+	~FCameraShakePreviewer();
 
-	MOVIESCENETRACKS_API void Initialize(UWorld* InWorld);
-	bool IsInitialized() const { return PreviewCameraShake != nullptr; }
-	MOVIESCENETRACKS_API void Teardown();
+	void ModifyView(FEditorViewportViewModifierParams& Params);
 
-	APlayerCameraManager* GetCameraManager() const { return PreviewCamera; }
-	UCameraModifier_CameraShake* GetCameraModifier() const { return PreviewCameraShake; }
+	void RegisterViewModifier();
+	void UnRegisterViewModifier();
 
-	MOVIESCENETRACKS_API void ModifyView(FEditorViewportViewModifierParams& Params);
+	void Update(float DeltaTime, bool bIsPlaying);
+	void Scrub(float ScrubTime);
 
-	MOVIESCENETRACKS_API void RegisterViewModifier();
-	MOVIESCENETRACKS_API void UnRegisterViewModifier();
-	MOVIESCENETRACKS_API void Update(float DeltaTime, bool bIsPlaying);
+	UCameraShakeBase* AddCameraShake(const FCameraShakePreviewerAddParams& Params);
+	void RemoveCameraShake(UCameraShakeBase* ShakeInstance);
+	void RemoveAllCameraShakesFromSource(const UCameraShakeSourceComponent* SourceComponent);
+	void RemoveAllCameraShakes();
+
+	int32 NumActiveCameraShakes() const { return ActiveShakes.Num(); }
+	void GetActiveCameraShakes(TArray<FActiveCameraShakeInfo>& ActiveCameraShakes) const;
+
+	void OnObjectsReplaced(const TMap<UObject*, UObject*>& ReplacementMap);
 
 private:
 	// FGCObject interface
-	virtual void AddReferencedObjects(FReferenceCollector& Collector) override { Collector.AddReferencedObject(PreviewCameraShake); }
-	virtual FString GetReferencerName() const override { return TEXT("SCameraShakePreviewer"); }
+	virtual void AddReferencedObjects(FReferenceCollector& Collector) override;
+	virtual FString GetReferencerName() const override { return TEXT("FCameraShakePreviewer"); }
 
 private:
-	MOVIESCENETRACKS_API void OnModifyView(FEditorViewportViewModifierParams& Params);
-	MOVIESCENETRACKS_API void OnLevelViewportClientListChanged();
+	void OnModifyView(FEditorViewportViewModifierParams& Params);
+	void OnLevelViewportClientListChanged();
+
+	void ResetModifiers();
 
 private:
-	APreviewPlayerCameraManager* PreviewCamera;
-	TObjectPtr<UCameraModifier_CameraShake> PreviewCameraShake;
 	TArray<FLevelEditorViewportClient*> RegisteredViewportClients;
 
+	struct FPreviewCameraShakeInfo
+	{
+		TObjectPtr<UCameraShakeBase> ShakeInstance;
+		TWeakObjectPtr<const UCameraShakeSourceComponent> SourceComponent;
+		float StartTime;
+	};
+	TArray<FPreviewCameraShakeInfo> ActiveShakes;
+
 	TOptional<float> LastDeltaTime;
+	TOptional<float> LastScrubTime;
+
 	FVector LastLocationModifier;
 	FRotator LastRotationModifier;
 	float LastFOVModifier;
diff --git a/Engine/Source/Runtime/MovieSceneTracks/Public/MovieSceneTracksComponentTypes.h b/Engine/Source/Runtime/MovieSceneTracks/Public/MovieSceneTracksComponentTypes.h
--- a/Engine/Source/Runtime/MovieSceneTracks/Public/MovieSceneTracksComponentTypes.h
+++ b/Engine/Source/Runtime/MovieSceneTracks/Public/MovieSceneTracksComponentTypes.h
@@ -105,15 +105,6 @@ struct FMovieSceneAudioComponentData
 	TObjectPtr<UMovieSceneAudioSection> Section = nullptr;
 };
 
-UENUM()
-enum EMovieSceneCameraShakeStatus : uint8
-{
-	NotStarted,
-	Started,
-	BlendingOut,
-	Finished
-};
-
 /** Component data for camera shakes */
 USTRUCT()
 struct FMovieSceneCameraShakeComponentData
@@ -136,18 +127,6 @@ struct FMovieSceneCameraShakeComponentData
 	// Shake instance created by the shake evaluation system.
 	UPROPERTY()
 	TObjectPtr<UCameraShakeBase> ShakeInstance;
-	
-	// Status for the shake.
-	UPROPERTY()
-	TEnumAsByte<EMovieSceneCameraShakeStatus> Status = EMovieSceneCameraShakeStatus::NotStarted;
-
-	// Duration of the shake.
-	UPROPERTY()
-	FFrameTime Duration;
-
-	// Time for blending out the shake.
-	UPROPERTY()
-	FFrameTime BlendOutTime;
 
 	FMovieSceneCameraShakeComponentData()
 	{}
diff --git a/Engine/Source/Runtime/MovieSceneTracks/Public/Systems/MovieSceneCameraShakeSystem.h b/Engine/Source/Runtime/MovieSceneTracks/Public/Systems/MovieSceneCameraShakeSystem.h
--- a/Engine/Source/Runtime/MovieSceneTracks/Public/Systems/MovieSceneCameraShakeSystem.h
+++ b/Engine/Source/Runtime/MovieSceneTracks/Public/Systems/MovieSceneCameraShakeSystem.h
@@ -108,17 +108,6 @@ class UMovieSceneCameraShakeEvaluatorSystem : public UMovieSceneEntitySystem
 public:
 	UMovieSceneCameraShakeEvaluatorSystem(const FObjectInitializer& ObjInit);
 
-	virtual void OnRun(FSystemTaskPrerequisites& InPrerequisites, FSystemSubsequentTasks& Subsequents) override;
-};
-
-UCLASS()
-class UMovieSceneCameraShakePreviewerEvaluatorSystem : public UMovieSceneEntitySystem
-{
-	GENERATED_BODY()
-
-public:
-	UMovieSceneCameraShakePreviewerEvaluatorSystem(const FObjectInitializer& ObjInit);
-
 	virtual void OnLink() override;
 	virtual void OnUnlink() override;
 	virtual bool IsRelevantImpl(UMovieSceneEntitySystemLinker* InLinker) const override;
@@ -129,4 +118,3 @@ private:
 	TSharedPtr<UE::MovieScene::FCameraShakePreviewerLinkerExtension> PreviewerExtension;
 #endif
 };
-
