PR #10839: Fix disappearing root motion in template sequences.

The montage data is mapped per MovieSceneSkeletalAnimationSection. The problem is that this object is the same if the template sequence is used for multipl objects, so the montage instance ID gets overwritten every frame. Thus, during the next frame when updating the anim positions, it cannot find the already existing montage that is playing on the anim instance, and creates a new one.

diff --git a/Engine/Source/Runtime/MovieSceneTracks/Private/Systems/MovieSceneSkeletalAnimationSystem.cpp b/Engine/Source/Runtime/MovieSceneTracks/Private/Systems/MovieSceneSkeletalAnimationSystem.cpp
--- a/Engine/Source/Runtime/MovieSceneTracks/Private/Systems/MovieSceneSkeletalAnimationSystem.cpp
+++ b/Engine/Source/Runtime/MovieSceneTracks/Private/Systems/MovieSceneSkeletalAnimationSystem.cpp
@@ -545,7 +545,7 @@ private:
 		}
 		else if (ExistingAnimInstance)
 		{
-			for (const TPair<FObjectKey, FMontagePlayerPerSectionData >& Pair : SystemData->MontageData)
+			for (const TPair<FObjectKey, FMontagePlayerPerSectionData >& Pair : SystemData->MontageData.FindOrAdd(SkeletalMeshComponent))
 			{
 				int32 InstanceId = Pair.Value.MontageInstanceId;
 				FAnimMontageInstance* MontageInstanceToUpdate = ExistingAnimInstance->GetMontageInstanceForID(InstanceId);
@@ -821,7 +821,7 @@ private:
 		}
 		else if (UAnimInstance* AnimInst = GetSourceAnimInstance(Params.SkeletalMeshComponent))
 		{
-			FMontagePlayerPerSectionData* SectionData = SystemData->MontageData.Find(Params.Section);
+			FMontagePlayerPerSectionData* SectionData = SystemData->MontageData.FindOrAdd(Params.SkeletalMeshComponent).Find(Params.Section);
 
 			int32 InstanceId = (SectionData) ? SectionData->MontageInstanceId : INDEX_NONE;
 
@@ -841,7 +841,7 @@ private:
 			UAnimMontage* Montage = WeakMontage.Get();
 			if (Montage)
 			{
-				FMontagePlayerPerSectionData& DataContainer = SystemData->MontageData.FindOrAdd(Params.Section);
+				FMontagePlayerPerSectionData& DataContainer = SystemData->MontageData.FindOrAdd(Params.SkeletalMeshComponent).FindOrAdd(Params.Section);
 				DataContainer.Montage = WeakMontage;
 				DataContainer.MontageInstanceId = InstanceId;
 
@@ -914,7 +914,7 @@ private:
 		}
 		else if (UAnimInstance* AnimInst = GetSourceAnimInstance(Params.SkeletalMeshComponent))
 		{
-			FMontagePlayerPerSectionData* SectionData = SystemData->MontageData.Find(Params.Section);
+			FMontagePlayerPerSectionData* SectionData = SystemData->MontageData.FindOrAdd(Params.SkeletalMeshComponent).Find(Params.Section);
 
 			int32 InstanceId = SectionData ? SectionData->MontageInstanceId : INDEX_NONE;
 		
@@ -935,7 +935,7 @@ private:
 			UAnimMontage* Montage = WeakMontage.Get();
 			if (Montage)
 			{
-				FMontagePlayerPerSectionData& DataContainer = SystemData->MontageData.FindOrAdd(Params.Section);
+				FMontagePlayerPerSectionData& DataContainer = SystemData->MontageData.FindOrAdd(Params.SkeletalMeshComponent).FindOrAdd(Params.Section);
 				DataContainer.Montage = WeakMontage;
 				DataContainer.MontageInstanceId = InstanceId;
 
diff --git a/Engine/Source/Runtime/MovieSceneTracks/Public/Systems/MovieSceneSkeletalAnimationSystem.h b/Engine/Source/Runtime/MovieSceneTracks/Public/Systems/MovieSceneSkeletalAnimationSystem.h
--- a/Engine/Source/Runtime/MovieSceneTracks/Public/Systems/MovieSceneSkeletalAnimationSystem.h
+++ b/Engine/Source/Runtime/MovieSceneTracks/Public/Systems/MovieSceneSkeletalAnimationSystem.h
@@ -79,7 +79,7 @@ struct FSkeletalAnimationSystemData
 	TMap<USkeletalMeshComponent*, FBoundObjectActiveSkeletalAnimations> SkeletalAnimations;
 
 	/** Map of persistent montage data */
-	TMap<FObjectKey, FMontagePlayerPerSectionData> MontageData;
+	TMap<FObjectKey, TMap<FObjectKey, FMontagePlayerPerSectionData>> MontageData;
 };
 
 } // namespace UE::MovieScene
