Sequencer: use pre-animated state API in camera cut track instance, and remove directly use of player API

This change tries to address various problems with the camera cut evaluation logic:

- Use the pre-animated state API to properly handle restoring things, instead of doing it partially "by hand". This removes the need for a few hacks like the carry-over of previous view-targets from one sequence to the next.

- Move the code from FSequencer and ULevelSequencePlayer into the track instance (and some related helper classes). Not only does this remove a bit of code duplication, this also lets us take a first step towards removing the dependency of evaluation code on the player API.

diff --git a/Engine/Plugins/MovieScene/LevelSequenceEditor/Source/LevelSequenceEditor/Private/LevelSequenceEditorBlueprintLibrary.cpp b/Engine/Plugins/MovieScene/LevelSequenceEditor/Source/LevelSequenceEditor/Private/LevelSequenceEditorBlueprintLibrary.cpp
--- a/Engine/Plugins/MovieScene/LevelSequenceEditor/Source/LevelSequenceEditor/Private/LevelSequenceEditorBlueprintLibrary.cpp
+++ b/Engine/Plugins/MovieScene/LevelSequenceEditor/Source/LevelSequenceEditor/Private/LevelSequenceEditorBlueprintLibrary.cpp
@@ -621,27 +621,7 @@ void ULevelSequenceEditorBlueprintLibrary::SetLockCameraCutToViewport(bool bLock
 	if (CurrentSequencer.IsValid())
 	{
 		TSharedPtr<ISequencer> Sequencer = CurrentSequencer.Pin();
-
-		if (bLock)
-		{
-			for(FLevelEditorViewportClient* LevelVC : GEditor->GetLevelViewportClients())
-			{
-				if (LevelVC && LevelVC->AllowsCinematicControl() && LevelVC->GetViewMode() != VMI_Unknown)
-				{
-					LevelVC->SetActorLock(nullptr);
-					LevelVC->bLockedCameraView = false;
-					LevelVC->UpdateViewForLockedActor();
-					LevelVC->Invalidate();
-				}
-			}
-			Sequencer->SetPerspectiveViewportCameraCutEnabled(true);
-		}
-		else
-		{
-			Sequencer->UpdateCameraCut(nullptr, EMovieSceneCameraCutParams());
-			Sequencer->SetPerspectiveViewportCameraCutEnabled(false);
-		}
-
+		Sequencer->SetPerspectiveViewportCameraCutEnabled(bLock);
 		Sequencer->ForceEvaluate();
 	}
 }
diff --git a/Engine/Source/Editor/MovieSceneTools/Private/TrackEditors/CameraCutTrackEditor.cpp b/Engine/Source/Editor/MovieSceneTools/Private/TrackEditors/CameraCutTrackEditor.cpp
--- a/Engine/Source/Editor/MovieSceneTools/Private/TrackEditors/CameraCutTrackEditor.cpp
+++ b/Engine/Source/Editor/MovieSceneTools/Private/TrackEditors/CameraCutTrackEditor.cpp
@@ -570,26 +570,9 @@ ECheckBoxState FCameraCutTrackEditor::IsCameraLocked() const
 
 void FCameraCutTrackEditor::OnLockCameraClicked(ECheckBoxState CheckBoxState)
 {
-	if (CheckBoxState == ECheckBoxState::Checked)
-	{
-		for(FLevelEditorViewportClient* LevelVC : GEditor->GetLevelViewportClients())
-		{
-			if (LevelVC && LevelVC->AllowsCinematicControl() && LevelVC->GetViewMode() != VMI_Unknown)
-			{
-				LevelVC->SetActorLock(nullptr);
-				LevelVC->bLockedCameraView = false;
-				LevelVC->UpdateViewForLockedActor();
-				LevelVC->Invalidate();
-			}
-		}
-		GetSequencer()->SetPerspectiveViewportCameraCutEnabled(true);
-	}
-	else
-	{
-		GetSequencer()->UpdateCameraCut(nullptr, EMovieSceneCameraCutParams());
-		GetSequencer()->SetPerspectiveViewportCameraCutEnabled(false);
-	}
 
+	const bool bEnableCameraCuts = (CheckBoxState == ECheckBoxState::Checked);
+	GetSequencer()->SetPerspectiveViewportCameraCutEnabled(bEnableCameraCuts);
 	GetSequencer()->ForceEvaluate();
 }
 
diff --git a/Engine/Source/Editor/MovieSceneTools/Private/TrackEditors/CinematicShotTrackEditor.cpp b/Engine/Source/Editor/MovieSceneTools/Private/TrackEditors/CinematicShotTrackEditor.cpp
--- a/Engine/Source/Editor/MovieSceneTools/Private/TrackEditors/CinematicShotTrackEditor.cpp
+++ b/Engine/Source/Editor/MovieSceneTools/Private/TrackEditors/CinematicShotTrackEditor.cpp
@@ -298,7 +298,6 @@ void FCinematicShotTrackEditor::OnLockShotsClicked(ECheckBoxState CheckBoxState)
 	}
 	else
 	{
-		GetSequencer()->UpdateCameraCut(nullptr, EMovieSceneCameraCutParams());
 		GetSequencer()->SetPerspectiveViewportCameraCutEnabled(false);
 	}
 
diff --git a/Engine/Source/Editor/Sequencer/Private/Sequencer.cpp b/Engine/Source/Editor/Sequencer/Private/Sequencer.cpp
--- a/Engine/Source/Editor/Sequencer/Private/Sequencer.cpp
+++ b/Engine/Source/Editor/Sequencer/Private/Sequencer.cpp
@@ -155,6 +155,7 @@
 #include "EntitySystem/MovieSceneEntitySystemRunner.h"
 #include "EntitySystem/MovieSceneEntitySystemLinker.h"
 #include "EntitySystem/MovieScenePreAnimatedStateSystem.h"
+#include "EntitySystem/MovieSceneSharedPlaybackState.h"
 #include "Systems/MovieSceneMotionVectorSimulationSystem.h"
 #include "IKeyArea.h"
 #include "Editor/TransBuffer.h"
@@ -284,6 +285,7 @@ namespace UE
 
 void FSequencer::InitSequencer(const FSequencerInitParams& InitParams, const TSharedRef<ISequencerObjectChangeListener>& InObjectChangeListener, const TArray<FOnCreateTrackEditor>& TrackEditorDelegates, const TArray<FOnCreateEditorObjectBinding>& EditorObjectBindingDelegates, const TArray<FOnCreateOutlinerColumn>& OutlinerColumnDelegates)
 {
+	using namespace UE::MovieScene;
 	using namespace UE::Sequencer;
 
 	bIsEditingWithinLevelEditor = InitParams.bEditWithinLevelEditor;
@@ -470,6 +472,13 @@ void FSequencer::InitSequencer(const FSequencerInitParams& InitParams, const TSh
 	InitialValueCache = UE::MovieScene::FInitialValueCache::GetGlobalInitialValues();
 	RootTemplateInstance.GetEntitySystemLinker()->AddExtension(InitialValueCache.Get());
 
+	// Add the camera cut playback capability.
+	FInstanceHandle RootInstanceHandle = RootTemplateInstance.GetRootInstanceHandle();
+	UMovieSceneEntitySystemLinker* Linker = GetEvaluationTemplate().GetEntitySystemLinker();
+	FInstanceRegistry* InstanceRegistry = Linker->GetInstanceRegistry();
+	TSharedRef<FSharedPlaybackState> SharedPlaybackState = InstanceRegistry->GetInstance(RootInstanceHandle).GetSharedPlaybackState();
+	SharedPlaybackState->AddCapabilityRaw<FCameraCutPlaybackCapability>((FCameraCutPlaybackCapability*)this);
+
 	// Create tools and bind them to this sequencer
 	for( int32 DelegateIndex = 0; DelegateIndex < TrackEditorDelegates.Num(); ++DelegateIndex )
 	{
@@ -662,7 +671,6 @@ FSequencer::FSequencer()
 	, bUpdatingExternalSelection( false )
 	, bNeedsEvaluate(false)
 	, bNeedsInvalidateCachedData(false)
-	, bHasPreAnimatedInfo(false)
 {
 	// Exposes the sequencer and curve editor command lists to subscribers from other systems
 	FInputBindingManager::Get().RegisterCommandList(FSequencerCommands::Get().GetContextName(), SequencerCommandBindings);
@@ -1039,6 +1047,8 @@ UMovieSceneSubSection* FSequencer::FindSubSection(FMovieSceneSequenceID Sequence
 
 void FSequencer::ResetToNewRootSequence(UMovieSceneSequence& NewSequence)
 {
+	using namespace UE::MovieScene;
+
 	RemoveNodeGroupsCollectionChangedDelegate();
 
 	const UMovieSceneSequence* PreviousRootSequence = RootSequence.Get();
@@ -1085,6 +1095,13 @@ void FSequencer::ResetToNewRootSequence(UMovieSceneSequence& NewSequence)
 
 	AddNodeGroupsCollectionChangedDelegate();
 
+	// Add the camera cut playback capability.
+	FInstanceHandle RootInstanceHandle = RootTemplateInstance.GetRootInstanceHandle();
+	UMovieSceneEntitySystemLinker* Linker = GetEvaluationTemplate().GetEntitySystemLinker();
+	FInstanceRegistry* InstanceRegistry = Linker->GetInstanceRegistry();
+	TSharedRef<FSharedPlaybackState> SharedPlaybackState = InstanceRegistry->GetInstance(RootInstanceHandle).GetSharedPlaybackState();
+	SharedPlaybackState->AddCapabilityRaw<FCameraCutPlaybackCapability>((FCameraCutPlaybackCapability*)this);
+
 	OnActivateSequenceEvent.Broadcast(ActiveTemplateIDs.Top());
 }
 
@@ -3408,93 +3425,6 @@ void FSequencer::SetPerspectiveViewportCameraCutEnabled(bool bEnabled)
 	}
 
 	bPerspectiveViewportCameraCutEnabled = bEnabled;
-
-	for (FLevelEditorViewportClient* LevelVC : GEditor->GetLevelViewportClients())
-	{
-		if (LevelVC != nullptr && LevelVC->AllowsCinematicControl() && LevelVC->GetViewMode() != VMI_Unknown)
-		{
-			if (bEnabled)
-			{
-				LevelVC->ViewModifiers.AddRaw(this, &FSequencer::ModifyViewportClientView);
-			}
-			else
-			{
-				LevelVC->ViewModifiers.RemoveAll(this);
-			}
-		}
-	}
-}
-
-void FSequencer::ModifyViewportClientView(FEditorViewportViewModifierParams& Params)
-{
-	if (!ViewModifierInfo.bApplyViewModifier)
-	{
-		return;
-	}
-
-	const float BlendFactor = ViewModifierInfo.BlendFactor;
-	AActor* CameraActor = ViewModifierInfo.NextCamera.Get();
-	AActor* PreviousCameraActor = ViewModifierInfo.PreviousCamera.Get();
-	
-	UCameraComponent* CameraComponent = MovieSceneHelpers::CameraComponentFromRuntimeObject(CameraActor);
-	UCameraComponent* PreviousCameraComponent = MovieSceneHelpers::CameraComponentFromRuntimeObject(PreviousCameraActor);
-
-
-	if (CameraActor)
-	{
-		const FVector ViewLocation = CameraComponent ? CameraComponent->GetComponentLocation() : CameraActor->GetActorLocation();
-		const FRotator ViewRotation = CameraComponent ? CameraComponent->GetComponentRotation() : CameraActor->GetActorRotation();
-
-		// If we have no previous camera actor or component, it means we're blending from the original
-		// editor viewport camera transform that we cached.
-		const FVector PreviousViewLocation = PreviousCameraComponent ?
-			PreviousCameraComponent->GetComponentLocation() :
-			(PreviousCameraActor ? PreviousCameraActor->GetActorLocation() : PreAnimatedViewportLocation);
-		const FRotator PreviousViewRotation = PreviousCameraComponent ?
-			PreviousCameraComponent->GetComponentRotation() :
-			(PreviousCameraActor ? PreviousCameraActor->GetActorRotation() : PreAnimatedViewportRotation);
-
-		const FVector BlendedLocation = FMath::Lerp(PreviousViewLocation, ViewLocation, BlendFactor);
-		const FRotator BlendedRotation = FMath::Lerp(PreviousViewRotation, ViewRotation, BlendFactor);
-
-		Params.ViewInfo.Location = BlendedLocation;
-		Params.ViewInfo.Rotation = BlendedRotation;
-	}
-	else
-	{
-		// Blending from a shot back to editor camera.
-
-		const FVector PreviousViewLocation = PreviousCameraComponent ?
-			PreviousCameraComponent->GetComponentLocation() :
-			(PreviousCameraActor ? PreviousCameraActor->GetActorLocation() : PreAnimatedViewportLocation);
-		const FRotator PreviousViewRotation = PreviousCameraComponent ?
-			PreviousCameraComponent->GetComponentRotation() :
-			(PreviousCameraActor ? PreviousCameraActor->GetActorRotation() : PreAnimatedViewportRotation);
-
-		const FVector BlendedLocation = FMath::Lerp(PreviousViewLocation, PreAnimatedViewportLocation, BlendFactor);
-		const FRotator BlendedRotation = FMath::Lerp(PreviousViewRotation, PreAnimatedViewportRotation, BlendFactor);
-
-		Params.ViewInfo.Location = BlendedLocation;
-		Params.ViewInfo.Rotation = BlendedRotation;
-	}
-
-	// Deal with camera properties.
-	if (CameraComponent)
-	{
-		const float PreviousFOV = PreviousCameraComponent != nullptr ?
-			PreviousCameraComponent->FieldOfView : PreAnimatedViewportFOV;
-		const float BlendedFOV = FMath::Lerp(PreviousFOV, CameraComponent->FieldOfView, BlendFactor);
-
-		Params.ViewInfo.FOV = BlendedFOV;
-	}
-	else
-	{
-		const float PreviousFOV = PreviousCameraComponent != nullptr ?
-			PreviousCameraComponent->FieldOfView : PreAnimatedViewportFOV;
-		const float BlendedFOV = FMath::Lerp(PreviousFOV, PreAnimatedViewportFOV, BlendFactor);
-
-		Params.ViewInfo.FOV = BlendedFOV;
-	}
 }
 
 FString FSequencer::GetMovieRendererName() const
@@ -3645,10 +3575,6 @@ void FSequencer::RenderMovieInternal(TRange<FFrameNumber> Range, bool bSetFrameO
 
 void FSequencer::EnterSilentMode()
 {
-	if (SilentModeCount == 0)
-	{
-		CachedViewModifierInfo = ViewModifierInfo;
-	}
 	++SilentModeCount;
 }
 
@@ -3656,10 +3582,6 @@ void FSequencer::ExitSilentMode()
 { 
 	--SilentModeCount;
 	ensure(SilentModeCount >= 0);
-	if (SilentModeCount == 0)
-	{
-		ViewModifierInfo = CachedViewModifierInfo;
-	}
 }
 
 ISequencer::FOnActorAddedToSequencer& FSequencer::OnActorAddedToSequencer()
@@ -3921,165 +3843,11 @@ ISequencerObjectChangeListener& FSequencer::GetObjectChangeListener()
 	return *ObjectChangeListener;
 }
 
-void FSequencer::PossessPIEViewports(UObject* CameraObject, const EMovieSceneCameraCutParams& CameraCutParams)
-{
-	UWorld* World = Cast<UWorld>(CachedPlaybackContext.Get());
-	if (!World || World->WorldType != EWorldType::PIE)
-	{
-		return;
-	}
-	
-	APlayerController* PC = World->GetGameInstance()->GetFirstLocalPlayerController();
-	if (PC == nullptr)
-	{
-		return;
-	}
-
-	TWeakObjectPtr<APlayerController> WeakPC = PC;
-	auto FindViewTarget = [=](const FCachedViewTarget& In){ return In.PlayerController == WeakPC; };
-
-	// skip same view target
-	AActor* ViewTarget = PC->GetViewTarget();
-
-	// save the last view target so that it can be restored when the camera object is null
-	if (!PrePossessionViewTargets.ContainsByPredicate(FindViewTarget))
-	{
-		PrePossessionViewTargets.Add(FCachedViewTarget{ PC, ViewTarget });
-	}
-
-	UCameraComponent* CameraComponent = MovieSceneHelpers::CameraComponentFromRuntimeObject(CameraObject);
-	if (CameraComponent && CameraComponent->GetOwner() != CameraObject)
-	{
-		CameraObject = CameraComponent->GetOwner();
-	}
-
-	if (CameraObject == ViewTarget)
-	{
-		if (CameraCutParams.bJumpCut)
-		{
-			if (PC->PlayerCameraManager)
-			{
-				PC->PlayerCameraManager->SetGameCameraCutThisFrame();
-			}
-
-			if (CameraComponent)
-			{
-				CameraComponent->NotifyCameraCut();
-			}
-
-			if (UMovieSceneMotionVectorSimulationSystem* MotionVectorSim = RootTemplateInstance.GetEntitySystemLinker()->FindSystem<UMovieSceneMotionVectorSimulationSystem>())
-			{
-				MotionVectorSim->SimulateAllTransforms();
-			}
-		}
-		return;
-	}
-
-	// skip unlocking if the current view target differs
-	AActor* UnlockIfCameraActor = Cast<AActor>(CameraCutParams.UnlockIfCameraObject);
-
-	// if unlockIfCameraActor is valid, release lock if currently locked to object
-	if (CameraObject == nullptr && UnlockIfCameraActor != nullptr && UnlockIfCameraActor != ViewTarget)
-	{
-		return;
-	}
-
-	// override the player controller's view target
-	AActor* CameraActor = Cast<AActor>(CameraObject);
-
-	// if the camera object is null, use the last view target so that it is restored to the state before the sequence takes control
-	if (CameraActor == nullptr)
-	{
-		if (const FCachedViewTarget* CachedTarget = PrePossessionViewTargets.FindByPredicate(FindViewTarget))
-		{
-			CameraActor = CachedTarget->ViewTarget.Get();
-		}
-	}
-
-	FViewTargetTransitionParams TransitionParams;
-	TransitionParams.BlendTime = FMath::Max(0.f, CameraCutParams.BlendTime);
-	PC->SetViewTarget(CameraActor, TransitionParams);
-
-	if (CameraComponent)
-	{
-		CameraComponent->NotifyCameraCut();
-	}
-
-	if (PC->PlayerCameraManager)
-	{
-		PC->PlayerCameraManager->bClientSimulatingViewTarget = (CameraActor != nullptr);
-		PC->PlayerCameraManager->SetGameCameraCutThisFrame();
-	}
-
-	if (UMovieSceneMotionVectorSimulationSystem* MotionVectorSim = RootTemplateInstance.GetEntitySystemLinker()->FindSystem<UMovieSceneMotionVectorSimulationSystem>())
-	{
-		MotionVectorSim->SimulateAllTransforms();
-	}
-}
-
 TSharedPtr<class ITimeSlider> FSequencer::GetTopTimeSliderWidget() const
 {
 	return SequencerWidget->GetTopTimeSliderWidget();
 }
 
-void FSequencer::UpdateCameraCut(UObject* CameraObject, const EMovieSceneCameraCutParams& CameraCutParams)
-{
-	OnCameraCutEvent.Broadcast(CameraObject, CameraCutParams.bJumpCut);
-
-	if (!IsPerspectiveViewportCameraCutEnabled())
-	{
-		return;
-	}
-
-	PossessPIEViewports(CameraObject, CameraCutParams);
-
-	// If the previous camera is null it means we are cutting from the editor camera, in which case
-	// we want to cache the current viewport's pre-animated info.
-	bool bShouldCachePreAnimatedViewportInfo = (
-			!bHasPreAnimatedInfo &&
-			(CameraObject == nullptr || CameraCutParams.PreviousCameraObject == nullptr) &&
-			!IsInSilentMode());
-
-	AActor* UnlockIfCameraActor = Cast<AActor>(CameraCutParams.UnlockIfCameraObject);
-
-	for (FLevelEditorViewportClient* LevelVC : GEditor->GetLevelViewportClients())
-	{
-		if ((LevelVC == nullptr) || !LevelVC->AllowsCinematicControl())
-		{
-			continue;
-		}
-
-		if (CameraObject == nullptr && UnlockIfCameraActor != nullptr && !LevelVC->IsLockedToActor(UnlockIfCameraActor))
-		{
-			continue;
-		}
-
-		if (bShouldCachePreAnimatedViewportInfo)
-		{
-			PreAnimatedViewportLocation = LevelVC->GetViewLocation();
-			PreAnimatedViewportRotation = LevelVC->GetViewRotation();
-			PreAnimatedViewportFOV = LevelVC->ViewFOV;
-			bHasPreAnimatedInfo = true;
-
-			// We end-up only caching the first cinematic viewport's info, which means that
-			// if we are previewing the sequence on 2 different viewports, the second viewport
-			// will blend back to the same camera position as the first viewport, even if they
-			// started at different positions (which is very likely). It's a small downside to
-			// pay for a much simpler piece of code, and for a use-case that is frankly 
-			// probably very uncommon.
-			bShouldCachePreAnimatedViewportInfo = false;
-		}
-
-		UpdatePreviewLevelViewportClientFromCameraCut(*LevelVC, CameraObject, CameraCutParams);
-	}
-
-	// Clear pre-animated info when we exit any sequencer camera.
-	if (CameraObject == nullptr && CameraCutParams.BlendTime < 0.f)
-	{
-		bHasPreAnimatedInfo = false;
-	}
-}
-
 void FSequencer::UpdateLevelViewportClientsActorLocks()
 {
 	// Nothing to do if we are not editing level sequence, as these are the only kinds of sequences right now
@@ -5730,113 +5498,6 @@ void FSequencer::OnNewActorsDropped(const TArray<UObject*>& DroppedObjects, cons
 	}
 }
 
-
-void FSequencer::UpdatePreviewLevelViewportClientFromCameraCut(FLevelEditorViewportClient& InViewportClient, UObject* InCameraObject, const EMovieSceneCameraCutParams& CameraCutParams)
-{
-	AActor* CameraActor = Cast<AActor>(InCameraObject);
-	AActor* PreviousCameraActor = Cast<AActor>(CameraCutParams.PreviousCameraObject);
-
-	const float BlendFactor = FMath::Clamp(CameraCutParams.PreviewBlendFactor, 0.f, 1.f);
-
-	const bool bIsBlending = (
-			(CameraCutParams.bCanBlend) &&
-			(CameraCutParams.BlendTime > 0.f) &&
-			(BlendFactor < 1.f - SMALL_NUMBER) &&
-			(CameraActor != nullptr || PreviousCameraActor != nullptr));
-
-	// To preview blending we'll have to offset the viewport camera using the view modifiers API.
-	ViewModifierInfo.bApplyViewModifier = bIsBlending && !IsInSilentMode();
-	ViewModifierInfo.BlendFactor = BlendFactor;
-	ViewModifierInfo.NextCamera = CameraActor;
-	ViewModifierInfo.PreviousCamera = PreviousCameraActor;
-
-	bool bCameraHasBeenCut = CameraCutParams.bJumpCut;
-
-	// When possible, let's get values from the camera components instead of the actor itself.
-	UCameraComponent* CameraComponent = MovieSceneHelpers::CameraComponentFromRuntimeObject(InCameraObject);
-	UCameraComponent* PreviousCameraComponent = MovieSceneHelpers::CameraComponentFromRuntimeObject(CameraCutParams.PreviousCameraObject);
-
-	if (CameraActor)
-	{
-		bCameraHasBeenCut = bCameraHasBeenCut || !InViewportClient.IsLockedToActor(CameraActor);
-
-		const FVector ViewLocation = CameraComponent ? CameraComponent->GetComponentLocation() : CameraActor->GetActorLocation();
-		const FRotator ViewRotation = CameraComponent ? CameraComponent->GetComponentRotation() : CameraActor->GetActorRotation();
-
-		InViewportClient.SetViewLocation(ViewLocation);
-		InViewportClient.SetViewRotation(ViewRotation);
-	}
-	else
-	{
-		if (CameraCutParams.bCanBlend && bHasPreAnimatedInfo)
-		{
-			InViewportClient.SetViewLocation(PreAnimatedViewportLocation);
-			InViewportClient.SetViewRotation(PreAnimatedViewportRotation);
-		}
-	}
-
-	if (bCameraHasBeenCut)
-	{
-		InViewportClient.SetIsCameraCut();
-
-		if (UMovieSceneMotionVectorSimulationSystem* MotionVectorSim = RootTemplateInstance.GetEntitySystemLinker()->FindSystem<UMovieSceneMotionVectorSimulationSystem>())
-		{
-			MotionVectorSim->SimulateAllTransforms();
-		}
-	}
-
-	// Set the actor lock.
-	InViewportClient.SetCinematicActorLock(CameraActor);
-	InViewportClient.bLockedCameraView = (CameraActor != nullptr);
-	InViewportClient.RemoveCameraRoll();
-
-	// Deal with camera properties.
-	if (CameraComponent)
-	{
-		if (bCameraHasBeenCut)
-		{
-			// tell the camera we cut
-			CameraComponent->NotifyCameraCut();
-		}
-
-		// enforce aspect ratio.
-		if (CameraComponent->AspectRatio == 0)
-		{
-			InViewportClient.AspectRatio = 1.7f;
-		}
-		else
-		{
-			InViewportClient.AspectRatio = CameraComponent->AspectRatio;
-		}
-
-		// enforce viewport type.
-		if (CameraComponent->ProjectionMode == ECameraProjectionMode::Type::Perspective)
-		{
-			if (InViewportClient.GetViewportType() != LVT_Perspective)
-			{
-				InViewportClient.SetViewportType(LVT_Perspective);
-			}
-		}
-
-		// don't stop the camera from zooming when not playing back
-		InViewportClient.ViewFOV = CameraComponent->FieldOfView;
-
-		// If there are selected actors, invalidate the viewports hit proxies, otherwise they won't be selectable afterwards
-		if (InViewportClient.Viewport && GEditor->GetSelectedActorCount() > 0)
-		{
-			InViewportClient.Viewport->InvalidateHitProxy();
-		}
-	}
-	else
-	{
-		InViewportClient.ViewFOV = InViewportClient.FOVAngle;
-	}
-
-	// Update ControllingActorViewInfo, so it is in sync with the updated viewport
-	InViewportClient.UpdateViewForLockedActor();
-}
-
-
 void FSequencer::SetShowCurveEditor(bool bInShowCurveEditor)
 {
 	SequencerWidget->OnCurveEditorVisibilityChanged(bInShowCurveEditor);
@@ -11849,5 +11510,15 @@ UE::Sequencer::FSequencerSelection& FSequencer::GetSelection()
 	return *ViewModel->GetSelection();
 }
 
+float FSequencer::GetCameraBlendPlayRate()
+{
+	return PlaybackSpeed;
+}
+
+void FSequencer::OnCameraCutUpdated(const UE::MovieScene::FOnCameraCutUpdatedParams& Params)
+{
+	OnCameraCutEvent.Broadcast(Params.ViewTarget, Params.bIsJumpCut);
+}
+
 #undef LOCTEXT_NAMESPACE
 
diff --git a/Engine/Source/Editor/Sequencer/Private/Sequencer.h b/Engine/Source/Editor/Sequencer/Private/Sequencer.h
--- a/Engine/Source/Editor/Sequencer/Private/Sequencer.h
+++ b/Engine/Source/Editor/Sequencer/Private/Sequencer.h
@@ -30,6 +30,7 @@
 #include "SequencerSelectionPreview.h"
 #include "SequencerCustomizationManager.h"
 #include "ITransportControl.h"
+#include "Evaluation/CameraCutPlaybackCapability.h"
 #include "Evaluation/MovieSceneSequenceTransform.h"
 #include "Evaluation/MovieScenePlayback.h"
 #include "Evaluation/MovieSceneEvaluationTemplateInstance.h"
@@ -108,6 +109,7 @@ class FSequencer final
 	, public FGCObject
 	, public FEditorUndoClient
 	, public FTickableEditorObject
+	, public UE::MovieScene::FCameraCutPlaybackCapability
 {
 	using FViewModel = UE::Sequencer::FViewModel;
 
@@ -813,7 +815,6 @@ public:
 
 	// IMovieScenePlayer interface
 
-	virtual void UpdateCameraCut(UObject* CameraObject, const EMovieSceneCameraCutParams& CameraCutParams) override;
 	virtual void NotifyBindingsChanged() override;
 	virtual void SetViewportSettings(const TMap<FViewportClient*, EMovieSceneViewportParams>& ViewportParamsMap) override;
 	virtual void GetViewportSettings(TMap<FViewportClient*, EMovieSceneViewportParams>& ViewportParamsMap) const override;
@@ -834,6 +835,17 @@ public:
 	 */
 	TRange<FFrameNumber> GetTimeBounds() const;
 
+protected:
+
+	// FCameraCutPlaybackCapability interface
+
+	virtual bool ShouldUpdateCameraCut() override { return IsPerspectiveViewportCameraCutEnabled(); }
+	virtual float GetCameraBlendPlayRate() override;
+	virtual void OnCameraCutUpdated(const UE::MovieScene::FOnCameraCutUpdatedParams& Params) override;
+#if WITH_EDITOR
+	virtual bool ShouldPreviewCameraCutBlends() override { return !IsInSilentMode(); }
+#endif
+
 protected:
 
 	/** Reevaluate the sequence at the current time */
@@ -1042,9 +1054,6 @@ public:
 
 private:
 
-	/** Updates a viewport client from camera cut data */
-	void UpdatePreviewLevelViewportClientFromCameraCut(FLevelEditorViewportClient& InViewportClient, UObject* InCameraObject, const EMovieSceneCameraCutParams& CameraCutParams);
-
 	/** Updates viewport clients' actor locks if they relate to sequencer cameras */
 	void UpdateLevelViewportClientsActorLocks();
 
@@ -1090,9 +1099,6 @@ private:
 	/** Create record transport control */
 	TSharedRef<SWidget> OnCreateTransportRecord();
 
-	/** Possess PIE viewports with the specified camera settings (a mirror of level viewport possession, but for game viewport clients) */
-	void PossessPIEViewports(UObject* CameraObject, const EMovieSceneCameraCutParams& CameraCutParams);
-
 	/** Update the locked subsequence range (displayed as playback range for subsequences), and root to local transform */
 	void UpdateSubSequenceData();
 
@@ -1143,9 +1149,6 @@ private:
 	/** Update the time bases for the current movie scene */
 	void UpdateTimeBases();
 
-	/** View modifier for level editor viewports. */
-	void ModifyViewportClientView(FEditorViewportViewModifierParams& Params);
-
 	/** User-supplied settings object for this sequencer */
 	TObjectPtr<USequencerSettings> Settings;
 
@@ -1454,25 +1457,8 @@ private:
 
 	FCachedViewState CachedViewState;
 	
-	struct FViewModifierInfo
-	{
-		bool bApplyViewModifier = false;
-		float BlendFactor = 1.f;
-		TWeakObjectPtr<AActor> PreviousCamera;
-		TWeakObjectPtr<AActor> NextCamera;
-	};
 	/** Frame time last evaluated at, needed to make sure we evaluate at the right time when interrogating values*/
 	FFrameTime LastEvaluatedLocalTime;
-	/** Information for previewing camera cut blends. This will be applied to the editor viewport during blends. */
-	FViewModifierInfo ViewModifierInfo;
-	/** Information cached before entering silent mode, so we can restore it afterwards. */
-	FViewModifierInfo CachedViewModifierInfo;
-	
-	/** Original editor camera info, for when previewing a sequence with a blend from/to gameplay. */
-	bool bHasPreAnimatedInfo;
-	FVector PreAnimatedViewportLocation;
-	FRotator PreAnimatedViewportRotation;
-	float PreAnimatedViewportFOV;
 
 	TOptional<FMovieSceneSequenceID> ScrubPositionParent;
 	/** Cache of all bound cameras in the sequence hierarchy */
diff --git a/Engine/Source/Runtime/LevelSequence/Private/LevelSequence.cpp b/Engine/Source/Runtime/LevelSequence/Private/LevelSequence.cpp
--- a/Engine/Source/Runtime/LevelSequence/Private/LevelSequence.cpp
+++ b/Engine/Source/Runtime/LevelSequence/Private/LevelSequence.cpp
@@ -607,7 +607,6 @@ FGuid ULevelSequence::FindOrAddBinding(UObject* InObject)
 		public:
 			FMovieSceneRootEvaluationTemplateInstance Template;
 			virtual FMovieSceneRootEvaluationTemplateInstance& GetEvaluationTemplate() override { return Template; }
-			virtual void UpdateCameraCut(UObject* CameraObject, const EMovieSceneCameraCutParams& CameraCutParams) override {}
 			virtual void SetViewportSettings(const TMap<FViewportClient*, EMovieSceneViewportParams>& ViewportParamsMap) override {}
 			virtual void GetViewportSettings(TMap<FViewportClient*, EMovieSceneViewportParams>& ViewportParamsMap) const override {}
 			virtual EMovieScenePlayerStatus::Type GetPlaybackStatus() const { return EMovieScenePlayerStatus::Stopped; }
diff --git a/Engine/Source/Runtime/LevelSequence/Private/LevelSequencePlayer.cpp b/Engine/Source/Runtime/LevelSequence/Private/LevelSequencePlayer.cpp
--- a/Engine/Source/Runtime/LevelSequence/Private/LevelSequencePlayer.cpp
+++ b/Engine/Source/Runtime/LevelSequence/Private/LevelSequencePlayer.cpp
@@ -23,12 +23,14 @@
 #include "Sections/MovieSceneCinematicShotSection.h"
 #include "Systems/MovieSceneMotionVectorSimulationSystem.h"
 #include "EntitySystem/MovieSceneEntitySystemLinker.h"
+#include "EntitySystem/MovieSceneSharedPlaybackState.h"
 #include "LevelSequenceActor.h"
 #include "Modules/ModuleManager.h"
 #include "LevelUtils.h"
 #include "ProfilingDebugging/CsvProfiler.h"
 #include "LevelSequenceModule.h"
 #include "Generators/MovieSceneEasingCurves.h"
+#include "Evaluation/CameraCutPlaybackCapability.h"
 
 #include UE_INLINE_GENERATED_CPP_BY_NAME(LevelSequencePlayer)
 
@@ -85,6 +87,8 @@ ULevelSequencePlayer* ULevelSequencePlayer::CreateLevelSequencePlayer(UObject* W
 
 void ULevelSequencePlayer::Initialize(ULevelSequence* InLevelSequence, ULevel* InLevel, const FLevelSequenceCameraSettings& InCameraSettings)
 {
+	using namespace UE::MovieScene;
+
 	// Never use the level to resolve bindings unless we're playing back within a streamed or instanced level
 	StreamedLevelAssetPath = FTopLevelAssetPath();
 
@@ -128,7 +132,15 @@ void ULevelSequencePlayer::Initialize(ULevelSequence* InLevelSequence, ULevel* I
 	}
 
 	SpawnRegister = MakeShareable(new FLevelSequenceSpawnRegister);
+
 	UMovieSceneSequencePlayer::Initialize(InLevelSequence);
+
+	// Add the camera cut playback capability.
+	FInstanceHandle RootInstanceHandle = RootTemplateInstance.GetRootInstanceHandle();
+	UMovieSceneEntitySystemLinker* Linker = GetEvaluationTemplate().GetEntitySystemLinker();
+	FInstanceRegistry* InstanceRegistry = Linker->GetInstanceRegistry();
+	TSharedRef<FSharedPlaybackState> SharedPlaybackState = InstanceRegistry->GetInstance(RootInstanceHandle).GetSharedPlaybackState();
+	SharedPlaybackState->AddCapabilityRaw<FCameraCutPlaybackCapability>((FCameraCutPlaybackCapability*)this);
 }
 
 void ULevelSequencePlayer::ResolveBoundObjects(const FGuid& InBindingId, FMovieSceneSequenceID SequenceID, UMovieSceneSequence& InSequence, UObject* ResolutionContext, TArray<UObject*, TInlineAllocator<1>>& OutObjects) const
@@ -189,11 +201,6 @@ void ULevelSequencePlayer::OnStopped()
 			}
 		}
 	}
-
-	LastCameraObject = nullptr;
-	LastViewTarget.Reset();
-	LastLocalPlayer.Reset();
-	LastAspectRatioAxisConstraint.Reset();
 }
 
 void ULevelSequencePlayer::UpdateMovieSceneInstance(FMovieSceneEvaluationRange InRange, EMovieScenePlayerStatus::Type PlayerStatus, const FMovieSceneUpdateArgs& Args)
@@ -213,322 +220,37 @@ void ULevelSequencePlayer::UpdateMovieSceneInstance(FMovieSceneEvaluationRange I
 	PreviousSnapshot = NewSnapshot;
 }
 
-/* IMovieScenePlayer interface
+/* FCameraCutPlaybackCapability interface
  *****************************************************************************/
 
-TTuple<EViewTargetBlendFunction, float> BuiltInEasingTypeToBlendFunction(EMovieSceneBuiltInEasing EasingType)
+bool ULevelSequencePlayer::ShouldUpdateCameraCut()
 {
-	using Return = TTuple<EViewTargetBlendFunction, float>;
-	switch (EasingType)
-	{
-		case EMovieSceneBuiltInEasing::Linear:
-			return Return(EViewTargetBlendFunction::VTBlend_Linear, 1.f);
-
-		case EMovieSceneBuiltInEasing::QuadIn:
-			return Return(EViewTargetBlendFunction::VTBlend_EaseIn, 2);
-		case EMovieSceneBuiltInEasing::QuadOut:
-			return Return(EViewTargetBlendFunction::VTBlend_EaseOut, 2);
-		case EMovieSceneBuiltInEasing::QuadInOut:
-			return Return(EViewTargetBlendFunction::VTBlend_EaseInOut, 2);
-
-		case EMovieSceneBuiltInEasing::CubicIn:
-			return Return(EViewTargetBlendFunction::VTBlend_EaseIn, 3);
-		case EMovieSceneBuiltInEasing::CubicOut:
-			return Return(EViewTargetBlendFunction::VTBlend_EaseOut, 3);
-		case EMovieSceneBuiltInEasing::CubicInOut:
-			return Return(EViewTargetBlendFunction::VTBlend_EaseInOut, 3);
-
-		case EMovieSceneBuiltInEasing::QuartIn:
-			return Return(EViewTargetBlendFunction::VTBlend_EaseIn, 4);
-		case EMovieSceneBuiltInEasing::QuartOut:
-			return Return(EViewTargetBlendFunction::VTBlend_EaseOut, 4);
-		case EMovieSceneBuiltInEasing::QuartInOut:
-			return Return(EViewTargetBlendFunction::VTBlend_EaseInOut, 4);
-
-		case EMovieSceneBuiltInEasing::QuintIn:
-			return Return(EViewTargetBlendFunction::VTBlend_EaseIn, 5);
-		case EMovieSceneBuiltInEasing::QuintOut:
-			return Return(EViewTargetBlendFunction::VTBlend_EaseOut, 5);
-		case EMovieSceneBuiltInEasing::QuintInOut:
-			return Return(EViewTargetBlendFunction::VTBlend_EaseInOut, 5);
-
-		// UNSUPPORTED
-		case EMovieSceneBuiltInEasing::SinIn:
-		case EMovieSceneBuiltInEasing::SinOut:
-		case EMovieSceneBuiltInEasing::SinInOut:
-		case EMovieSceneBuiltInEasing::CircIn:
-		case EMovieSceneBuiltInEasing::CircOut:
-		case EMovieSceneBuiltInEasing::CircInOut:
-		case EMovieSceneBuiltInEasing::ExpoIn:
-		case EMovieSceneBuiltInEasing::ExpoOut:
-		case EMovieSceneBuiltInEasing::ExpoInOut:
-			break;
-	}
-	return Return(EViewTargetBlendFunction::VTBlend_Linear, 1.f);
+	return !PlaybackSettings.bDisableCameraCuts;
 }
 
-void ULevelSequencePlayer::ValidateLastViewTarget(UObject* CameraObject, AActor* ViewTarget)
+float ULevelSequencePlayer::GetCameraBlendPlayRate()
 {
-	if (!LastViewTarget.IsValid())
-	{
-		// If we had another level sequence playing before us, we don't want to cache its last camera. Let's see if
-		// we can instead ask it what its LastViewTarget was, so that we get what the actual previous view target was
-		// before any sequence started to play.
-		TArray<IMovieScenePlayer*> ActivePlayers;
-		const bool bOnlyEvaluatingPlayers = true;
-		IMovieScenePlayer::Get(ActivePlayers, bOnlyEvaluatingPlayers);
-		for (IMovieScenePlayer* ActivePlayer : ActivePlayers)
-		{
-			UObject* PlayerObject = ActivePlayer->AsUObject();
-			if (!PlayerObject)
-			{
-				continue;
-			}
-
-
-			ULevelSequencePlayer* PreviousLevelSequencePlayer = Cast<ULevelSequencePlayer>(PlayerObject);
-			if (!PreviousLevelSequencePlayer || PreviousLevelSequencePlayer == this)
-			{
-				continue;
-			}
-
-			if (PreviousLevelSequencePlayer->LastViewTarget.IsValid())
-			{
-				LastViewTarget = PreviousLevelSequencePlayer->LastViewTarget;
-				break;
-			}
-		}
-	}
-
-	if (!LastViewTarget.IsValid() || (ViewTarget != nullptr && ViewTarget != CameraObject && ViewTarget != LastCameraObject))
-	{
-		LastViewTarget = ViewTarget;
-	}
+	return PlaybackSettings.PlayRate;
 }
 
-void ULevelSequencePlayer::UpdateCameraCut(UObject* CameraObject, const EMovieSceneCameraCutParams& CameraCutParams)
+TOptional<EAspectRatioAxisConstraint> ULevelSequencePlayer::GetAspectRatioAxisConstraintOverride()
 {
-	UCameraComponent* CameraComponent = MovieSceneHelpers::CameraComponentFromRuntimeObject(CameraObject);
-	if (CameraComponent && CameraComponent->GetOwner() != CameraObject)
-	{
-		CameraObject = CameraComponent->GetOwner();
-	}
-
-	CachedCameraComponent = CameraComponent;
-	
-	if (World == nullptr || World->GetGameInstance() == nullptr)
-	{
-		return;
-	}
-
-	// If the player controller is missing, there is no camera manager for us to manage the view target
-	// so we just bail out.
-	APlayerController* PC = World->GetGameInstance()->GetFirstLocalPlayerController();
-	APlayerCameraManager* CameraManager = (PC != nullptr) ? PC->PlayerCameraManager : nullptr;
-
-	if (CameraManager == nullptr)
-	{
-		// If we are releasing view target control but can't actually do it because game state has been
-		// torn down, see if we can at least restore aspect ratio axis constraint. This fixes a bug
-		// where the wrong constraint carries over to a new level if the level change was executed
-		// while we were playing the sequence.
-		if (CameraObject == nullptr && LastAspectRatioAxisConstraint.IsSet())
-		{
-			if (ULocalPlayer* LocalPlayer = LastLocalPlayer.Get())
-			{
-				LocalPlayer->AspectRatioAxisConstraint = LastAspectRatioAxisConstraint.GetValue();
-			}
-		}
-		return;
-	}
-
-	// If we don't want to update camera cuts, let's also bail out.
-	if (!CanUpdateCameraCut())
-	{
-		return;
-	}
-
-	// Let's get the current view target directly from the player camera manager first. This is because
-	// we don't want to go through GetViewTarget, which checks if the current view target is valid, and
-	// re-assigns it to the player controller if it's not. We don't want this to happen, especially since
-	// a spawnable camera might have just been unspawned, causing UpdateCameraCut to be called, and we
-	// need to handle this properly.
-	AActor* ViewTarget = CameraManager->PendingViewTarget.Target;
-	if (!ViewTarget)
-	{
-		ViewTarget = CameraManager->ViewTarget.Target;
-	}
-
-	// If the view target isn't really changing, we don't have much to do.
-	if (CameraObject == ViewTarget)
-	{
-		if (CameraCutParams.bJumpCut)
-		{
-			if (CameraManager)
-			{
-				CameraManager->SetGameCameraCutThisFrame();
-			}
-
-			if (CameraComponent)
-			{
-				CameraComponent->NotifyCameraCut();
-			}
-
-			if (UMovieSceneMotionVectorSimulationSystem* MotionVectorSim = RootTemplateInstance.GetEntitySystemLinker()->FindSystem<UMovieSceneMotionVectorSimulationSystem>())
-			{
-				MotionVectorSim->SimulateAllTransforms();
-			}
-		}
-
-		// If we just started a new sequence that happens to give control to the same camera actor as
-		// the previous sequence, we end up here for an early out, but we also miss on setting out LastViewTarget,
-		// so let's check we have a valid one.
-		ValidateLastViewTarget(CameraObject, ViewTarget);
-
-		return;
-	}
-
-	// skip unlocking if the current view target differs
-	AActor* UnlockIfCameraActor = Cast<AActor>(CameraCutParams.UnlockIfCameraObject);
-
-	// if unlockIfCameraActor is valid, release lock if currently locked to object
-	if (CameraObject == nullptr && UnlockIfCameraActor != nullptr && UnlockIfCameraActor != ViewTarget)
-	{
-		return;
-	}
-
-	// override the player controller's view target
-	AActor* CameraActor = Cast<AActor>(CameraObject);
-	ULocalPlayer* LocalPlayer = (PC != nullptr) ? PC->GetLocalPlayer() : nullptr;
-
-	// if the camera object is null, use the last view target so that it is restored to the state before the sequence takes control
-	bool bRestoreAspectRatioConstraint = false;
-	if (CameraActor == nullptr)
-	{
-		CameraActor = LastViewTarget.Get();
-		CameraObject = CameraActor;
-
-		// Skip if the last view target is the same as the current view target so that there's no additional camera cut
-		// In this case, restore the aspect ratio axis constraint right away before exiting.
-		if (CameraActor == ViewTarget)
-		{
-			if (LocalPlayer && LastAspectRatioAxisConstraint.IsSet())
-			{
-				LocalPlayer->AspectRatioAxisConstraint = LastAspectRatioAxisConstraint.GetValue();
-			}
-			return;
-		}
-		// Otherwise, remember to restore the aspect ratio axis constraint later in this function.
-		bRestoreAspectRatioConstraint = true;
-	}
-
-	// Save the last view target/aspect ratio constraint/etc. so that it can all be restored when the camera object is null.
-	ValidateLastViewTarget(CameraObject, ViewTarget);
-
-	if (!LastLocalPlayer.IsValid() || (LocalPlayer != LastLocalPlayer))
-	{
-		LastLocalPlayer = LocalPlayer;
-	}
-	if (!LastAspectRatioAxisConstraint.IsSet())
-	{
-		if (LocalPlayer != nullptr)
-		{
-			LastAspectRatioAxisConstraint = LocalPlayer->AspectRatioAxisConstraint;
-		}
-	}
-
-	LastCameraObject = CameraObject;
-
-	bool bDoSetViewTarget = true;
-	FViewTargetTransitionParams TransitionParams;
-	if (CameraCutParams.BlendType.IsSet())
-	{
-		UE_LOG(LogLevelSequence, Log, TEXT("Blending into new camera cut: '%s' -> '%s' (blend time: %f)"),
-			(ViewTarget ? *ViewTarget->GetName() : TEXT("None")),
-			(CameraObject ? *CameraObject->GetName() : TEXT("None")),
-			TransitionParams.BlendTime);
-
-		const float PlayRateFactor = PlaybackSettings.PlayRate != 0.f ? (1.f / PlaybackSettings.PlayRate) : 1.f;
-
-		// Convert known easing functions to their corresponding view target blend parameters.
-		TTuple<EViewTargetBlendFunction, float> BlendFunctionAndExp = BuiltInEasingTypeToBlendFunction(CameraCutParams.BlendType.GetValue());
-		TransitionParams.BlendTime = CameraCutParams.BlendTime * PlayRateFactor;
-		TransitionParams.bLockOutgoing = CameraCutParams.bLockPreviousCamera;
-		TransitionParams.BlendFunction = BlendFunctionAndExp.Get<0>();
-		TransitionParams.BlendExp = BlendFunctionAndExp.Get<1>();
-
-		// Calling SetViewTarget on a camera that we are currently transitioning to will 
-		// result in that transition being aborted, and the view target being set immediately.
-		// We want to avoid that, so let's leave the transition running if it's the case.
-		if (CameraManager != nullptr)
-		{
-			const AActor* PendingViewTarget = CameraManager->PendingViewTarget.Target;
-			if (CameraActor != nullptr && PendingViewTarget == CameraActor)
-			{
-				UE_LOG(LogLevelSequence, Log, TEXT("Camera transition aborted, we are already blending towards the intended camera"));
-				bDoSetViewTarget = false;
-			}
-		}
-	}
-	else
-	{
-		UE_LOG(LogLevelSequence, Log, TEXT("Starting new camera cut: '%s'"),
-			(CameraObject ? *CameraObject->GetName() : TEXT("None")));
-	}
-	if (bDoSetViewTarget && ensureMsgf(PC, TEXT("Can't set view target when there is no player controller!")))
-	{
-		PC->SetViewTarget(CameraActor, TransitionParams);
-	}
-
-	// Set or restore the aspect ratio constraint if we were overriding it for this sequence.
-	if (LocalPlayer != nullptr && CameraSettings.bOverrideAspectRatioAxisConstraint)
-	{
-		if (bRestoreAspectRatioConstraint)
-		{
-			check(LastAspectRatioAxisConstraint.IsSet());
-			if (LastAspectRatioAxisConstraint.IsSet())
-			{
-				LocalPlayer->AspectRatioAxisConstraint = LastAspectRatioAxisConstraint.GetValue();
-			}
-		}
-		else
-		{
-			LocalPlayer->AspectRatioAxisConstraint = CameraSettings.AspectRatioAxisConstraint;
-		}
-	}
-
-	// we want to notify of cuts on hard cuts and time jumps, but not on blend cuts
-	const bool bIsStraightCut = !CameraCutParams.BlendType.IsSet() || CameraCutParams.bJumpCut;
-
-	if (CameraComponent && bIsStraightCut)
-	{
-		CameraComponent->NotifyCameraCut();
-	}
-
-	if (CameraManager)
-	{
-		CameraManager->bClientSimulatingViewTarget = (CameraActor != nullptr);
-
-		if (bIsStraightCut)
-		{
-			CameraManager->SetGameCameraCutThisFrame();
-		}
-	}
+	return CameraSettings.bOverrideAspectRatioAxisConstraint ?
+		TOptional<EAspectRatioAxisConstraint>(CameraSettings.AspectRatioAxisConstraint) :
+		TOptional<EAspectRatioAxisConstraint>();
+}
 
-	if (bIsStraightCut)
+void ULevelSequencePlayer::OnCameraCutUpdated(const UE::MovieScene::FOnCameraCutUpdatedParams& Params)
+{
+	if (OnCameraCut.IsBound())
 	{
-		if (UMovieSceneMotionVectorSimulationSystem* MotionVectorSim = RootTemplateInstance.GetEntitySystemLinker()->FindSystem<UMovieSceneMotionVectorSimulationSystem>())
-		{
-			MotionVectorSim->SimulateAllTransforms();
-		}
-
-		if (OnCameraCut.IsBound())
-		{
-			OnCameraCut.Broadcast(CameraComponent);
-		}
+		OnCameraCut.Broadcast(Params.ViewTargetCamera);
 	}
 }
 
+/* IMovieScenePlayer interface
+ *****************************************************************************/
+
 UObject* ULevelSequencePlayer::GetPlaybackContext() const
 {
 	return World.Get();
diff --git a/Engine/Source/Runtime/LevelSequence/Public/LevelSequencePlayer.h b/Engine/Source/Runtime/LevelSequence/Public/LevelSequencePlayer.h
--- a/Engine/Source/Runtime/LevelSequence/Public/LevelSequencePlayer.h
+++ b/Engine/Source/Runtime/LevelSequence/Public/LevelSequencePlayer.h
@@ -7,6 +7,7 @@
 #include "UObject/ObjectMacros.h"
 #include "UObject/Object.h"
 #include "UObject/ScriptMacros.h"
+#include "Evaluation/CameraCutPlaybackCapability.h"
 #include "Evaluation/MovieScenePlayback.h"
 #include "Evaluation/MovieSceneEvaluationTemplateInstance.h"
 #include "Evaluation/PersistentEvaluationData.h"
@@ -22,7 +23,7 @@ class FLevelSequenceSpawnRegister;
 class FViewportClient;
 class UCameraComponent;
 
-DECLARE_DYNAMIC_MULTICAST_DELEGATE_OneParam(FOnLevelSequencePlayerCameraCutEvent, UCameraComponent*, CameraComponent);
+DECLARE_DYNAMIC_MULTICAST_DELEGATE_OneParam(FOnLevelSequencePlayerCameraCutEvent , UCameraComponent*, CameraComponent);
 
 /**
  * Frame snapshot information for a level sequence
@@ -80,6 +81,7 @@ private:
 UCLASS(BlueprintType, MinimalAPI)
 class ULevelSequencePlayer
 	: public UMovieSceneSequencePlayer
+	, public UE::MovieScene::FCameraCutPlaybackCapability
 {
 public:
 	LEVELSEQUENCE_API ULevelSequencePlayer(const FObjectInitializer&);
@@ -127,7 +129,6 @@ public:
 protected:
 
 	// IMovieScenePlayer interface
-	LEVELSEQUENCE_API virtual void UpdateCameraCut(UObject* CameraObject, const EMovieSceneCameraCutParams& CameraCutParams) override;
 	LEVELSEQUENCE_API virtual void ResolveBoundObjects(const FGuid& InBindingId, FMovieSceneSequenceID SequenceID, UMovieSceneSequence& InSequence, UObject* ResolutionContext, TArray<UObject*, TInlineAllocator<1>>& OutObjects) const override;
 
 	//~ UMovieSceneSequencePlayer interface
@@ -136,6 +137,12 @@ protected:
 	LEVELSEQUENCE_API virtual void OnStopped() override;
 	LEVELSEQUENCE_API virtual void UpdateMovieSceneInstance(FMovieSceneEvaluationRange InRange, EMovieScenePlayerStatus::Type PlayerStatus, const FMovieSceneUpdateArgs& Args) override;
 
+	//~ FCameraCutPlaybackCapability interface
+	LEVELSEQUENCE_API virtual bool ShouldUpdateCameraCut() override;
+	LEVELSEQUENCE_API virtual float GetCameraBlendPlayRate() override;
+	LEVELSEQUENCE_API virtual TOptional<EAspectRatioAxisConstraint> GetAspectRatioAxisConstraintOverride() override;
+	LEVELSEQUENCE_API virtual void OnCameraCutUpdated(const UE::MovieScene::FOnCameraCutUpdatedParams& Params) override;
+
 public:
 
 	/** Populate the specified array with any given event contexts for the specified world */
@@ -151,9 +158,6 @@ private:
 
 	LEVELSEQUENCE_API void EnableCinematicMode(bool bEnable);
 
-	// Save the last view target, so that it can all be restored when the camera object is null.
-	void ValidateLastViewTarget(UObject* CameraObject, AActor* ViewTarget);
-
 private:
 
 	/** The world this player will spawn actors in, if needed */
@@ -171,18 +175,6 @@ private:
 	/** The camera settings to use when playing the sequence */
 	FLevelSequenceCameraSettings CameraSettings;
 
-	/** The last view target to reset to when updating camera cuts to null */
-	TWeakObjectPtr<AActor> LastViewTarget;
-
-	/** The last player on which to reset the aspect ratio axis constraint in case the level is changing */
-	TWeakObjectPtr<ULocalPlayer> LastLocalPlayer;
-
-	/** The last aspect ratio axis constraint to reset to when the camera cut is null */
-	TOptional<EAspectRatioAxisConstraint> LastAspectRatioAxisConstraint;
-
-	/** The last camera cut object received during a cinematic */
-	UObject* LastCameraObject = nullptr;
-
 protected:
 
 	TOptional<int32> SnapshotOffsetTime;
diff --git a/Engine/Source/Runtime/MovieScene/Private/MovieSceneSequence.cpp b/Engine/Source/Runtime/MovieScene/Private/MovieSceneSequence.cpp
--- a/Engine/Source/Runtime/MovieScene/Private/MovieSceneSequence.cpp
+++ b/Engine/Source/Runtime/MovieScene/Private/MovieSceneSequence.cpp
@@ -299,7 +299,6 @@ FGuid UMovieSceneSequence::FindPossessableObjectId(UObject& Object, UObject* Con
 	public:
 		FMovieSceneRootEvaluationTemplateInstance Template;
 		virtual FMovieSceneRootEvaluationTemplateInstance& GetEvaluationTemplate() override { return Template; }
-		virtual void UpdateCameraCut(UObject* CameraObject, const EMovieSceneCameraCutParams& CameraCutParams) override {}
 		virtual void SetViewportSettings(const TMap<FViewportClient*, EMovieSceneViewportParams>& ViewportParamsMap) override {}
 		virtual void GetViewportSettings(TMap<FViewportClient*, EMovieSceneViewportParams>& ViewportParamsMap) const override {}
 		virtual EMovieScenePlayerStatus::Type GetPlaybackStatus() const { return EMovieScenePlayerStatus::Stopped; }
diff --git a/Engine/Source/Runtime/MovieScene/Private/Tests/MovieSceneBlendingTests.cpp b/Engine/Source/Runtime/MovieScene/Private/Tests/MovieSceneBlendingTests.cpp
--- a/Engine/Source/Runtime/MovieScene/Private/Tests/MovieSceneBlendingTests.cpp
+++ b/Engine/Source/Runtime/MovieScene/Private/Tests/MovieSceneBlendingTests.cpp
@@ -56,7 +56,6 @@ struct FNullPlayer : IMovieScenePlayer
 		return RootInstance;
 	}
 
-	virtual void UpdateCameraCut(UObject* CameraObject, const EMovieSceneCameraCutParams& CameraCutParams) {}
 	virtual void SetViewportSettings(const TMap<FViewportClient*, EMovieSceneViewportParams>& ViewportParamsMap) {}
 	virtual void GetViewportSettings(TMap<FViewportClient*, EMovieSceneViewportParams>& ViewportParamsMap) const {}
 	virtual EMovieScenePlayerStatus::Type GetPlaybackStatus() const { return EMovieScenePlayerStatus::Stopped; }
diff --git a/Engine/Source/Runtime/MovieScene/Private/Tests/MovieSceneCompilerTests.cpp b/Engine/Source/Runtime/MovieScene/Private/Tests/MovieSceneCompilerTests.cpp
--- a/Engine/Source/Runtime/MovieScene/Private/Tests/MovieSceneCompilerTests.cpp
+++ b/Engine/Source/Runtime/MovieScene/Private/Tests/MovieSceneCompilerTests.cpp
@@ -51,7 +51,6 @@ bool FMovieSceneCompilerPerfTest::RunTest(const FString& Parameters)
 	{
 		FMovieSceneRootEvaluationTemplateInstance RootInstance;
 		virtual FMovieSceneRootEvaluationTemplateInstance& GetEvaluationTemplate() override { return RootInstance; }
-		virtual void UpdateCameraCut(UObject* CameraObject, const EMovieSceneCameraCutParams& CameraCutParams) override {}
 		virtual void SetViewportSettings(const TMap<FViewportClient*, EMovieSceneViewportParams>& ViewportParamsMap) override {}
 		virtual void GetViewportSettings(TMap<FViewportClient*, EMovieSceneViewportParams>& ViewportParamsMap) const override {}
 		virtual EMovieScenePlayerStatus::Type GetPlaybackStatus() const override { return EMovieScenePlayerStatus::Playing; }
diff --git a/Engine/Source/Runtime/MovieScene/Private/Tests/MovieScenePreAnimatedStateTests.cpp b/Engine/Source/Runtime/MovieScene/Private/Tests/MovieScenePreAnimatedStateTests.cpp
--- a/Engine/Source/Runtime/MovieScene/Private/Tests/MovieScenePreAnimatedStateTests.cpp
+++ b/Engine/Source/Runtime/MovieScene/Private/Tests/MovieScenePreAnimatedStateTests.cpp
@@ -53,7 +53,6 @@ namespace Impl
 	{
 		FMovieSceneRootEvaluationTemplateInstance Template;
 		virtual FMovieSceneRootEvaluationTemplateInstance& GetEvaluationTemplate() override { return Template; }
-		virtual void UpdateCameraCut(UObject* CameraObject, const EMovieSceneCameraCutParams& CameraCutParams) override {}
 		virtual void SetViewportSettings(const TMap<FViewportClient*, EMovieSceneViewportParams>& ViewportParamsMap) override {}
 		virtual void GetViewportSettings(TMap<FViewportClient*, EMovieSceneViewportParams>& ViewportParamsMap) const override {}
 		virtual EMovieScenePlayerStatus::Type GetPlaybackStatus() const override { return EMovieScenePlayerStatus::Playing; }
diff --git a/Engine/Source/Runtime/MovieScene/Public/Evaluation/PreAnimatedState/MovieScenePreAnimatedStateStorage.h b/Engine/Source/Runtime/MovieScene/Public/Evaluation/PreAnimatedState/MovieScenePreAnimatedStateStorage.h
--- a/Engine/Source/Runtime/MovieScene/Public/Evaluation/PreAnimatedState/MovieScenePreAnimatedStateStorage.h
+++ b/Engine/Source/Runtime/MovieScene/Public/Evaluation/PreAnimatedState/MovieScenePreAnimatedStateStorage.h
@@ -546,6 +546,24 @@ public:
 		KeyToStorageIndex.Add(NewKey, StorageIndex);
 	}
 
+	/**
+	 * Get the cached value associated with a particular storage index
+	 */
+	const StorageType& GetCachedValue(FPreAnimatedStorageIndex StorageIndex) const
+	{
+		static const StorageType DefaultValue = StorageType();
+
+		if (ensure(PreAnimatedStorage.IsValidIndex(StorageIndex.Value)))
+		{
+			const FCachedData& CachedData = PreAnimatedStorage[StorageIndex.Value];
+			if (CachedData.bInitialized)
+			{
+				return CachedData.Value;
+			}
+		}
+		return DefaultValue;
+	}
+
 	/**
 	 * Look at any entity with the specified component types, and set up new associations with storage indices for those entities
 	 * The provided component values are put together to make up the storage key
diff --git a/Engine/Source/Runtime/MovieScene/Public/IMovieScenePlayer.h b/Engine/Source/Runtime/MovieScene/Public/IMovieScenePlayer.h
--- a/Engine/Source/Runtime/MovieScene/Public/IMovieScenePlayer.h
+++ b/Engine/Source/Runtime/MovieScene/Public/IMovieScenePlayer.h
@@ -55,7 +55,7 @@ struct EMovieSceneViewportParams
 };
 
 /** Camera cut parameters */
-struct EMovieSceneCameraCutParams
+struct FMovieSceneCameraCutParams
 {
 	/** If this is not null, release actor lock only if currently locked to this object */
 	UObject* UnlockIfCameraObject = nullptr;
@@ -78,6 +78,9 @@ struct EMovieSceneCameraCutParams
 #endif
 };
 
+/** Backwards compatibility to old struct name with typo */
+using EMovieSceneCameraCutParams = FMovieSceneCameraCutParams; 
+
 /**
  * Interface for movie scene players
  * Provides information for playback of a movie scene
@@ -104,34 +107,6 @@ public:
 	 */
 	virtual UObject* AsUObject() { return nullptr; }
 
-	/**
-	 * Whether this player can update the camera cut
-	 */
-	virtual bool CanUpdateCameraCut() const { return true; }
-
-	/**
-	 * Updates the perspective viewports with the actor to view through
-	 *
-	 * @param CameraObject The object, probably a camera, that the viewports should lock to
-	 * @param UnlockIfCameraObject If this is not nullptr, release actor lock only if currently locked to this object.
-	 * @param bJumpCut Whether this is a jump cut, ie. the cut jumps from one shot to another shot
-	 */
-	void UpdateCameraCut(UObject* CameraObject, UObject* UnlockIfCameraObject = nullptr, bool bJumpCut = false)
-	{
-		EMovieSceneCameraCutParams CameraCutParams;
-		CameraCutParams.UnlockIfCameraObject = UnlockIfCameraObject;
-		CameraCutParams.bJumpCut = bJumpCut;
-		UpdateCameraCut(CameraObject, CameraCutParams);
-	}
-
-	/**
-	 * Updates the perspective viewports with the actor to view through
-	 *
-	 * @param CameraObject The object, probably a camera, that the viewports should lock to
-	 * @param CameraCutParams The parameters for this camera cut.
-	 */
-	virtual void UpdateCameraCut(UObject* CameraObject, const EMovieSceneCameraCutParams& CameraCutParams) = 0;
-
 	/*
 	 * Set the perspective viewport settings
 	 *
@@ -359,6 +334,25 @@ public:
 		return UniqueIndex;
 	}
 
+public:
+
+	UE_DEPRECATED(5.4, "Camera cut management has moved to UMovieSceneCameraCutTrackInstance")
+	virtual bool CanUpdateCameraCut() const { return true; }
+
+	UE_DEPRECATED(5.4, "Camera cut management has moved to UMovieSceneCameraCutTrackInstance")
+	virtual void UpdateCameraCut(UObject* CameraObject, UObject* UnlockIfCameraObject = nullptr, bool bJumpCut = false)
+	{
+		EMovieSceneCameraCutParams CameraCutParams;
+		CameraCutParams.UnlockIfCameraObject = UnlockIfCameraObject;
+		CameraCutParams.bJumpCut = bJumpCut;
+		PRAGMA_DISABLE_DEPRECATION_WARNINGS
+		UpdateCameraCut(CameraObject, CameraCutParams);
+		PRAGMA_ENABLE_DEPRECATION_WARNINGS
+	}
+
+	UE_DEPRECATED(5.4, "Camera cut management has moved to UMovieSceneCameraCutTrackInstance")
+	virtual void UpdateCameraCut(UObject* CameraObject, const EMovieSceneCameraCutParams& CameraCutParams) {}
+
 public:
 
 	/** Evaluation state that stores global state to do with the playback operation */
diff --git a/Engine/Source/Runtime/MovieScene/Public/MovieSceneSequencePlayer.h b/Engine/Source/Runtime/MovieScene/Public/MovieSceneSequencePlayer.h
--- a/Engine/Source/Runtime/MovieScene/Public/MovieSceneSequencePlayer.h
+++ b/Engine/Source/Runtime/MovieScene/Public/MovieSceneSequencePlayer.h
@@ -530,8 +530,6 @@ protected:
 	virtual void SetPlaybackStatus(EMovieScenePlayerStatus::Type InPlaybackStatus) override {}
 	virtual void SetViewportSettings(const TMap<FViewportClient*, EMovieSceneViewportParams>& ViewportParamsMap) override {}
 	virtual void GetViewportSettings(TMap<FViewportClient*, EMovieSceneViewportParams>& ViewportParamsMap) const override {}
-	virtual bool CanUpdateCameraCut() const override { return !PlaybackSettings.bDisableCameraCuts; }
-	virtual void UpdateCameraCut(UObject* CameraObject, const EMovieSceneCameraCutParams& CameraCutParams) override {}
 	MOVIESCENE_API virtual void ResolveBoundObjects(const FGuid& InBindingId, FMovieSceneSequenceID SequenceID, UMovieSceneSequence& Sequence, UObject* ResolutionContext, TArray<UObject*, TInlineAllocator<1>>& OutObjects) const override;
 	virtual IMovieScenePlaybackClient* GetPlaybackClient() override { return PlaybackClient ? &*PlaybackClient : nullptr; }
 	MOVIESCENE_API virtual bool IsDisablingEventTriggers(FFrameTime& DisabledUntilTime) const override;
diff --git a/Engine/Source/Runtime/MovieSceneTracks/Private/Evaluation/CameraCutPlaybackCapability.cpp b/Engine/Source/Runtime/MovieSceneTracks/Private/Evaluation/CameraCutPlaybackCapability.cpp
new file mode 100644
--- /dev/null
+++ b/Engine/Source/Runtime/MovieSceneTracks/Private/Evaluation/CameraCutPlaybackCapability.cpp
@@ -0,0 +1,11 @@
+// Copyright Epic Games, Inc. All Rights Reserved.
+
+#include "Evaluation/CameraCutPlaybackCapability.h"
+
+namespace UE::MovieScene
+{
+
+TPlaybackCapabilityID<FCameraCutPlaybackCapability> FCameraCutPlaybackCapability::ID = TPlaybackCapabilityID<FCameraCutPlaybackCapability>::Register();
+
+} // namespace UE::MovieScene
+
diff --git a/Engine/Source/Runtime/MovieSceneTracks/Private/TrackInstances/MovieSceneCameraCutEditorHandler.cpp b/Engine/Source/Runtime/MovieSceneTracks/Private/TrackInstances/MovieSceneCameraCutEditorHandler.cpp
new file mode 100644
--- /dev/null
+++ b/Engine/Source/Runtime/MovieSceneTracks/Private/TrackInstances/MovieSceneCameraCutEditorHandler.cpp
@@ -0,0 +1,321 @@
+// Copyright Epic Games, Inc. All Rights Reserved.
+
+#include "TrackInstances/MovieSceneCameraCutEditorHandler.h"
+
+#if WITH_EDITOR
+
+#include "Evaluation/CameraCutPlaybackCapability.h"
+#include "IMovieScenePlayer.h"
+#include "LevelEditorViewport.h"
+#include "Systems/MovieSceneMotionVectorSimulationSystem.h"
+#include "TrackInstances/MovieSceneCameraCutTrackInstance.h"
+#include "TrackInstances/MovieSceneCameraCutViewportPreviewer.h"
+
+namespace UE::MovieScene
+{
+
+FPreAnimatedCameraCutEditorState FPreAnimatedCameraCutEditorTraits::CachePreAnimatedValue(
+		FLevelEditorViewportClient* InKey)
+{
+	FPreAnimatedCameraCutEditorState CachedValue;
+	CachedValue.ViewportLocation = InKey->GetViewLocation();
+	CachedValue.ViewportRotation = InKey->GetViewRotation();
+	CachedValue.ViewportFOV = InKey->ViewFOV;
+	return CachedValue;
+}
+
+void FPreAnimatedCameraCutEditorTraits::RestorePreAnimatedValue(
+		FLevelEditorViewportClient* InKey, 
+		const FPreAnimatedCameraCutEditorState& CachedValue, 
+		const FRestoreStateParams& Params)
+{
+	// Check that our pointer is still valid by searching it in active viewports.
+	if (GEditor->GetLevelViewportClients().Find(InKey) == INDEX_NONE)
+	{
+		return;
+	}
+
+	// If the viewport wasn't locked to cinematics anyway, don't mess it up.
+	if (!InKey->IsLockedToCinematic())
+	{
+		return;
+	}
+	
+	InKey->SetViewLocation(CachedValue.ViewportLocation);
+	InKey->SetViewRotation(CachedValue.ViewportRotation);
+	InKey->ViewFOV = CachedValue.ViewportFOV;
+
+	InKey->SetCinematicActorLock(nullptr);
+	InKey->bLockedCameraView = false;
+	InKey->UpdateViewForLockedActor();
+	InKey->Invalidate();
+}
+
+TAutoRegisterPreAnimatedStorageID<FPreAnimatedCameraCutEditorStorage> FPreAnimatedCameraCutEditorStorage::StorageID;
+
+void FCameraCutEditorHandler::CachePreAnimatedValue(
+		UMovieSceneEntitySystemLinker* Linker,
+		const FSequenceInstance& SequenceInstance)
+{
+	if (!GEditor)
+	{
+		return;
+	}
+
+	TSharedPtr<FPreAnimatedCameraCutEditorStorage> PreAnimatedStorage = Linker->PreAnimatedState.GetOrCreateStorage<FPreAnimatedCameraCutEditorStorage>();
+
+	for (FLevelEditorViewportClient* LevelVC : GEditor->GetLevelViewportClients())
+	{
+		PreAnimatedStorage->CachePreAnimatedValue(
+				LevelVC,
+				[](FLevelEditorViewportClient* InKey) { return FPreAnimatedCameraCutEditorTraits::CachePreAnimatedValue(InKey); },
+				EPreAnimatedCaptureSourceTracking::AlwaysCache);
+	}
+}
+
+FCameraCutEditorHandler::FCameraCutEditorHandler(
+		UMovieSceneEntitySystemLinker* InLinker,
+		const FSequenceInstance& InSequenceInstance,
+		FCameraCutViewportPreviewer& InViewportPreviewer)
+	: Linker(InLinker)
+	, SequenceInstance(InSequenceInstance)
+	, ViewportPreviewer(InViewportPreviewer)
+{
+}
+
+void FCameraCutEditorHandler::SetCameraCut(
+		UObject* CameraObject, 
+		const FMovieSceneCameraCutParams& CameraCutParams)
+{
+	FCameraCutPlaybackCapabilityCompatibilityWrapper Wrapper(SequenceInstance);
+
+	// If we don't want to update camera cuts, let's remember it and release the viewports
+	// in case they were still locked to cinematics.
+	const bool bPreviewCameraCuts = Wrapper.ShouldUpdateCameraCut();
+
+	// Make sure our viewport modifiers are correctly registered/unregistered.
+	ViewportPreviewer.ToggleViewportPreviewModifiers(bPreviewCameraCuts);
+
+	// Gather the viewports we want to affect.
+	TArray<FLevelEditorViewportClient*> CinematicVCs;
+
+	AActor* CameraActor = Cast<AActor>(CameraObject);
+	AActor* UnlockIfCameraActor = Cast<AActor>(CameraCutParams.UnlockIfCameraObject);
+	UCameraComponent* CameraComponent = MovieSceneHelpers::CameraComponentFromRuntimeObject(CameraObject);
+
+	for (FLevelEditorViewportClient* LevelVC : GEditor->GetLevelViewportClients())
+	{
+		if (LevelVC == nullptr)
+		{
+			continue;
+		}
+
+		if (!bPreviewCameraCuts || !LevelVC->AllowsCinematicControl())
+		{
+			ReleaseCameraCutForViewport(*LevelVC);
+			continue;
+		}
+
+		if (CameraActor == nullptr && 
+				UnlockIfCameraActor != nullptr && 
+				!LevelVC->IsLockedToActor(UnlockIfCameraActor))
+		{
+			continue;
+		}
+
+		SetCameraCutForViewport(*LevelVC, CameraActor, CameraComponent, CameraCutParams);
+	}
+
+	// Trigger any callback.
+	FOnCameraCutUpdatedParams CameraCutUpdatedParams;
+	CameraCutUpdatedParams.ViewTarget = CameraActor;
+	CameraCutUpdatedParams.ViewTargetCamera = CameraComponent;
+	CameraCutUpdatedParams.bIsJumpCut = true;
+	Wrapper.OnCameraCutUpdated(CameraCutUpdatedParams);
+}
+
+void FCameraCutEditorHandler::SetCameraCutForViewport(
+		FLevelEditorViewportClient& ViewportClient, 
+		AActor* CameraActor, 
+		UCameraComponent* CameraComponent,
+		const FMovieSceneCameraCutParams& CameraCutParams)
+{
+	FVector ViewLocation = ViewportClient.GetViewLocation();
+	FRotator ViewRotation = ViewportClient.GetViewRotation();
+	float ViewFOV = ViewportClient.ViewFOV;
+	bool bCameraHasBeenCut = CameraCutParams.bJumpCut;
+
+	TSharedPtr<FPreAnimatedCameraCutEditorStorage> PreAnimatedStorage = Linker->PreAnimatedState.FindStorage(
+			FPreAnimatedCameraCutEditorStorage::StorageID);
+
+	// If that viewport wasn't locked to cinematics, see if we need to re-cache pre-animated state.
+	// This is necessary if the user released control with the camera button on the camera cut track.
+	// (see ReleaseCameraCutForViewport)
+	if (!ViewportClient.IsLockedToCinematic())
+	{
+		FPreAnimatedStorageIndex StorageIndex = PreAnimatedStorage->FindStorageIndex(&ViewportClient);
+		if (!StorageIndex.IsValid())
+		{
+			PreAnimatedStorage->CachePreAnimatedValue(
+					&ViewportClient,
+					[](FLevelEditorViewportClient* InKey) { return FPreAnimatedCameraCutEditorTraits::CachePreAnimatedValue(InKey); },
+					EPreAnimatedCaptureSourceTracking::AlwaysCache);
+		}
+	}
+
+	if (CameraActor)
+	{
+		// When possible, let's get values from the camera components instead of the actor itself.
+		ViewLocation = CameraComponent ? CameraComponent->GetComponentLocation() : CameraActor->GetActorLocation();
+		ViewRotation = CameraComponent ? CameraComponent->GetComponentRotation() : CameraActor->GetActorRotation();
+		ViewFOV = CameraComponent ? CameraComponent->FieldOfView : ViewportClient.FOVAngle;
+		bCameraHasBeenCut = bCameraHasBeenCut || !ViewportClient.IsLockedToActor(CameraActor);
+	}
+	else
+	{
+		// If CameraActor is null, we are releasing camera control. This can happen here instead of
+		// inside pre-animated state if we are *blending* out of the cinematic. In this case, let's
+		// restore the pre-animated viewport location ourselves and enable blending.
+		FPreAnimatedStorageIndex StorageIndex = PreAnimatedStorage->FindStorageIndex(&ViewportClient);
+		if (ensureMsgf(StorageIndex.IsValid(), TEXT("Blending camera back to editor but can't find pre-animated viewport info!")))
+		{
+			FPreAnimatedCameraCutEditorState CachedValue = PreAnimatedStorage->GetCachedValue(StorageIndex);
+			ViewLocation = CachedValue.ViewportLocation;
+			ViewRotation = CachedValue.ViewportRotation;
+			ViewFOV = CachedValue.ViewportFOV;
+		}
+	}
+
+	// Set viewport properties.
+	ViewportClient.SetViewLocation(ViewLocation);
+	ViewportClient.SetViewRotation(ViewRotation);
+	ViewportClient.ViewFOV = ViewFOV;
+
+	if (bCameraHasBeenCut)
+	{
+		ViewportClient.SetIsCameraCut();
+
+		if (UMovieSceneMotionVectorSimulationSystem* MotionVectorSim = Linker->FindSystem<UMovieSceneMotionVectorSimulationSystem>())
+		{
+			MotionVectorSim->SimulateAllTransforms();
+		}
+	}
+
+	// Set the actor lock.
+	ViewportClient.SetCinematicActorLock(CameraActor);
+	ViewportClient.SetActorLock(nullptr);
+	ViewportClient.bLockedCameraView = (CameraActor != nullptr);
+	ViewportClient.RemoveCameraRoll();
+
+	// Deal with extra camera properties.
+	if (CameraComponent)
+	{
+		if (bCameraHasBeenCut)
+		{
+			// tell the camera we cut
+			CameraComponent->NotifyCameraCut();
+		}
+
+		// enforce aspect ratio.
+		if (CameraComponent->AspectRatio == 0)
+		{
+			ViewportClient.AspectRatio = 1.7f;
+		}
+		else
+		{
+			ViewportClient.AspectRatio = CameraComponent->AspectRatio;
+		}
+
+		// enforce viewport type.
+		if (CameraComponent->ProjectionMode == ECameraProjectionMode::Type::Perspective)
+		{
+			if (ViewportClient.GetViewportType() != LVT_Perspective)
+			{
+				ViewportClient.SetViewportType(LVT_Perspective);
+			}
+		}
+
+		// If there are selected actors, invalidate the viewports hit proxies, otherwise they won't be selectable afterwards
+		if (ViewportClient.Viewport && GEditor->GetSelectedActorCount() > 0)
+		{
+			ViewportClient.Viewport->InvalidateHitProxy();
+		}
+	}
+
+	// Setup blending preview.
+	const bool bIsBlending = (
+			(CameraCutParams.bCanBlend) &&
+			(CameraCutParams.BlendTime > 0.f) &&
+			(CameraCutParams.PreviewBlendFactor < 1.f - SMALL_NUMBER));
+	if (bIsBlending)
+	{
+		AActor* PreviousCameraActor = Cast<AActor>(CameraCutParams.PreviousCameraObject);
+		UCameraComponent* PreviousCameraComponent = MovieSceneHelpers::CameraComponentFromRuntimeObject(CameraCutParams.PreviousCameraObject);
+
+		FCameraCutViewportPreviewerTarget FromPreviewTarget;
+		if (PreviousCameraActor)
+		{
+			FromPreviewTarget.CameraActor = PreviousCameraActor;
+			FromPreviewTarget.CameraComponent = PreviousCameraComponent;
+		}
+		else
+		{
+			// We have no "from" camera, so let's blend from the original viewport position.
+			FromPreviewTarget.PreAnimatedStorage = PreAnimatedStorage;
+		}
+
+		FCameraCutViewportPreviewerTarget ToPreviewTarget;
+		if (CameraActor)
+		{
+			ToPreviewTarget.CameraActor = CameraActor;
+			ToPreviewTarget.CameraComponent = CameraComponent;
+		}
+		else
+		{
+			// We have no "to" camera, so let's blend back to the original viewport position.
+			FromPreviewTarget.PreAnimatedStorage = PreAnimatedStorage;
+		}
+
+		ViewportPreviewer.SetupBlend(FromPreviewTarget, ToPreviewTarget, CameraCutParams.PreviewBlendFactor);
+	}
+	else
+	{
+		ViewportPreviewer.TeardownBlend();
+	}
+
+	// Update ControllingActorViewInfo, so it is in sync with the updated viewport
+	ViewportClient.UpdateViewForLockedActor();
+	ViewportClient.Invalidate();
+}
+
+void FCameraCutEditorHandler::ReleaseCameraCutForViewport(
+		FLevelEditorViewportClient& ViewportClient)
+{
+	// If the viewport was already released, we have nothing to do.
+	if (!ViewportClient.IsLockedToCinematic())
+	{
+		return;
+	}
+	
+	// Restore the viewport to its pre-animated state.
+	TSharedPtr<FPreAnimatedCameraCutEditorStorage> PreAnimatedStorage = Linker->PreAnimatedState.FindStorage(FPreAnimatedCameraCutEditorStorage::StorageID);
+	FPreAnimatedStorageIndex StorageIndex = PreAnimatedStorage->FindStorageIndex(&ViewportClient);
+	if (StorageIndex.IsValid())
+	{
+		FPreAnimatedCameraCutEditorState CachedValue = PreAnimatedStorage->GetCachedValue(StorageIndex);
+		ViewportClient.SetViewLocation(CachedValue.ViewportLocation);
+		ViewportClient.SetViewRotation(CachedValue.ViewportRotation);
+		ViewportClient.ViewFOV = CachedValue.ViewportFOV;
+	}
+
+	// Actually release control.
+	ViewportClient.SetCinematicActorLock(nullptr);
+	ViewportClient.bLockedCameraView = false;
+	ViewportClient.UpdateViewForLockedActor();
+	ViewportClient.Invalidate();
+}
+
+}  // namespace UE::MovieScene
+ 
+#endif  // WITH_EDITOR
+
diff --git a/Engine/Source/Runtime/MovieSceneTracks/Private/TrackInstances/MovieSceneCameraCutEditorHandler.h b/Engine/Source/Runtime/MovieSceneTracks/Private/TrackInstances/MovieSceneCameraCutEditorHandler.h
new file mode 100644
--- /dev/null
+++ b/Engine/Source/Runtime/MovieSceneTracks/Private/TrackInstances/MovieSceneCameraCutEditorHandler.h
@@ -0,0 +1,98 @@
+// Copyright Epic Games, Inc. All Rights Reserved.
+
+#pragma once
+
+#if WITH_EDITOR
+
+#include "CoreMinimal.h"
+#include "Evaluation/PreAnimatedState/MovieScenePreAnimatedStateStorage.h"
+#include "Evaluation/PreAnimatedState/MovieScenePreAnimatedStateTypes.h"
+#include "Evaluation/PreAnimatedState/MovieScenePreAnimatedStorageID.inl"
+
+class AActor;
+class FLevelEditorViewportClient;
+class UCameraComponent;
+class UMovieSceneEntitySystemLinker;
+struct FMovieSceneCameraCutParams;
+
+namespace UE::MovieScene
+{
+
+class FCameraCutViewportPreviewer;
+struct FCameraCutPlaybackCapability;
+struct FPreAnimatedCameraCutEditorStorage;
+struct FSequenceInstance;
+
+/** Pre-animated editor viewpoint */
+struct FPreAnimatedCameraCutEditorState
+{
+	FVector ViewportLocation;
+	FRotator ViewportRotation;
+	float ViewportFOV = 0.f;
+};
+
+/** Pre-animated traits for in-editor camera cuts */
+struct FPreAnimatedCameraCutEditorTraits : FPreAnimatedStateTraits
+{
+	// Key type is the viewport to restore location/rotation/FOV on.
+	using KeyType = FLevelEditorViewportClient*;
+	using StorageType = FPreAnimatedCameraCutEditorState;
+
+	static StorageType CachePreAnimatedValue(KeyType InKey);
+
+	void RestorePreAnimatedValue(KeyType InKey, const StorageType& CachedValue, const FRestoreStateParams& Params);
+};
+
+// Make it possible to declare our storage below with a forward-declared FLevelEditorViewportClient.
+template<>
+struct THasAddReferencedObjectForComponent<FLevelEditorViewportClient*>
+{
+	static constexpr bool Value = false;
+};
+
+/** Pre-animated state storage for in-editor camera cuts */
+struct FPreAnimatedCameraCutEditorStorage : TPreAnimatedStateStorage<FPreAnimatedCameraCutEditorTraits>
+{
+	static TAutoRegisterPreAnimatedStorageID<FPreAnimatedCameraCutEditorStorage> StorageID;
+};
+
+/**
+ * Utility class for executing camera cuts in the editor viewports.
+ */
+struct FCameraCutEditorHandler
+{
+	FCameraCutEditorHandler(
+			UMovieSceneEntitySystemLinker* InLinker,
+			const FSequenceInstance& InSequenceInstance,
+			FCameraCutViewportPreviewer& InViewportPreviewer);
+
+	/** Sets the given camera cut in all editor viewport with cinematic control enabled */
+	void SetCameraCut(
+			UObject* CameraObject, 
+			const FMovieSceneCameraCutParams& CameraCutParams);
+
+	/** Cache any pre-animated values required for handling camera cuts in editor. */
+	static void CachePreAnimatedValue(
+			UMovieSceneEntitySystemLinker* Linker,
+			const FSequenceInstance& SequenceInstance);
+
+private:
+
+	void SetCameraCutForViewport(
+			FLevelEditorViewportClient& ViewportClient, 
+			AActor* CameraActor, 
+			UCameraComponent* CameraComponent,
+			const FMovieSceneCameraCutParams& CameraCutParams);
+	void ReleaseCameraCutForViewport(
+			FLevelEditorViewportClient& ViewportClient);
+
+private:
+	UMovieSceneEntitySystemLinker* Linker;
+	const FSequenceInstance& SequenceInstance;
+	FCameraCutViewportPreviewer& ViewportPreviewer;
+};
+
+}  // namespace UE::MovieScene
+
+#endif  // WITH_EDITOR
+
diff --git a/Engine/Source/Runtime/MovieSceneTracks/Private/TrackInstances/MovieSceneCameraCutGameHandler.cpp b/Engine/Source/Runtime/MovieSceneTracks/Private/TrackInstances/MovieSceneCameraCutGameHandler.cpp
new file mode 100644
--- /dev/null
+++ b/Engine/Source/Runtime/MovieSceneTracks/Private/TrackInstances/MovieSceneCameraCutGameHandler.cpp
@@ -0,0 +1,350 @@
+// Copyright Epic Games, Inc. All Rights Reserved.
+
+#include "TrackInstances/MovieSceneCameraCutGameHandler.h"
+
+#include "Camera/CameraComponent.h"
+#include "Camera/PlayerCameraManager.h"
+#include "Engine/Engine.h"
+#include "Engine/GameInstance.h"
+#include "Engine/LocalPlayer.h"
+#include "Engine/World.h"
+#include "Evaluation/CameraCutPlaybackCapability.h"
+#include "Generators/MovieSceneEasingCurves.h"
+#include "IMovieScenePlayer.h"
+#include "MovieSceneCommonHelpers.h"
+#include "Systems/MovieSceneMotionVectorSimulationSystem.h"
+#include "TrackInstances/MovieSceneCameraCutTrackInstance.h"
+
+namespace UE::MovieScene
+{
+
+FPreAnimatedCameraCutState FPreAnimatedCameraCutTraits::CachePreAnimatedValue(
+		IMovieScenePlayer* Player, 
+		uint8 InKey)
+{
+	UObject* PlaybackContext = Player->GetPlaybackContext();
+	UWorld* World = PlaybackContext ? PlaybackContext->GetWorld() : nullptr;
+	if (World && World->GetGameInstance())
+	{
+		APlayerController* PC = World->GetGameInstance()->GetFirstLocalPlayerController();
+
+		// Save previous view target.
+		APlayerCameraManager* CameraManager = (PC != nullptr) ? PC->PlayerCameraManager : nullptr;
+		AActor* ViewTarget = CameraManager ? CameraManager->GetViewTarget() : nullptr;
+
+		// Save previous aspect ratio axis constraint.
+		ULocalPlayer* LocalPlayer = (PC != nullptr) ? PC->GetLocalPlayer() : nullptr;
+		EAspectRatioAxisConstraint AspectRatioAxisConstraint = LocalPlayer->AspectRatioAxisConstraint;
+
+		return StorageType{ LocalPlayer, ViewTarget, AspectRatioAxisConstraint };
+	}
+	return StorageType();
+}
+
+void FPreAnimatedCameraCutTraits::RestorePreAnimatedValue(
+		uint8 InKey, 
+		const FPreAnimatedCameraCutState& CachedValue, 
+		const FRestoreStateParams& Params)
+{
+	IMovieScenePlayer* Player = Params.GetTerminalPlayer();
+	if (!Player)
+	{
+		return;
+	}
+
+	UObject* PlaybackContext = Player->GetPlaybackContext();
+	UWorld* World = GEngine->GetWorldFromContextObject(PlaybackContext, EGetWorldErrorMode::LogAndReturnNull);
+	if (!World || !World->GetGameInstance())
+	{
+		return;
+	}
+
+	APlayerController* PC = World->GetGameInstance()->GetFirstLocalPlayerController();
+	APlayerCameraManager* CameraManager = (PC != nullptr) ? PC->PlayerCameraManager : nullptr;
+
+	// Restore previous view target.
+	AActor* PreviousViewTarget = Cast<AActor>(CachedValue.LastViewTarget.ResolveObjectPtr());
+	if (CameraManager && PreviousViewTarget)
+	{
+		CameraManager->SetViewTarget(PreviousViewTarget);
+	}
+
+	// Restore previous aspect ratio axis constraint. Use the cached local player if there's no local player
+	// to be found, which can happen if pre-animated state is restored during level loads and such.
+	ULocalPlayer* LocalPlayer = (PC != nullptr) ? 
+		PC->GetLocalPlayer() : 
+		Cast<ULocalPlayer>(CachedValue.LastLocalPlayer.ResolveObjectPtr());
+	if (LocalPlayer)
+	{
+		LocalPlayer->AspectRatioAxisConstraint = CachedValue.LastAspectRatioAxisConstraint;
+	}
+}
+
+TAutoRegisterPreAnimatedStorageID<FPreAnimatedCameraCutStorage> FPreAnimatedCameraCutStorage::StorageID;
+
+/** Utility function for converting sequence blend curves to camera manager blend types */
+static TTuple<EViewTargetBlendFunction, float> BuiltInEasingTypeToBlendFunction(EMovieSceneBuiltInEasing EasingType)
+{
+	using Return = TTuple<EViewTargetBlendFunction, float>;
+	switch (EasingType)
+	{
+		case EMovieSceneBuiltInEasing::Linear:
+			return Return(EViewTargetBlendFunction::VTBlend_Linear, 1.f);
+
+		case EMovieSceneBuiltInEasing::QuadIn:
+			return Return(EViewTargetBlendFunction::VTBlend_EaseIn, 2);
+		case EMovieSceneBuiltInEasing::QuadOut:
+			return Return(EViewTargetBlendFunction::VTBlend_EaseOut, 2);
+		case EMovieSceneBuiltInEasing::QuadInOut:
+			return Return(EViewTargetBlendFunction::VTBlend_EaseInOut, 2);
+
+		case EMovieSceneBuiltInEasing::Cubic:
+		case EMovieSceneBuiltInEasing::HermiteCubicInOut:
+			return Return(EViewTargetBlendFunction::VTBlend_Cubic, 3);
+
+		case EMovieSceneBuiltInEasing::CubicIn:
+			return Return(EViewTargetBlendFunction::VTBlend_EaseIn, 3);
+		case EMovieSceneBuiltInEasing::CubicOut:
+			return Return(EViewTargetBlendFunction::VTBlend_EaseOut, 3);
+		case EMovieSceneBuiltInEasing::CubicInOut:
+			return Return(EViewTargetBlendFunction::VTBlend_EaseInOut, 3);
+
+		case EMovieSceneBuiltInEasing::QuartIn:
+			return Return(EViewTargetBlendFunction::VTBlend_EaseIn, 4);
+		case EMovieSceneBuiltInEasing::QuartOut:
+			return Return(EViewTargetBlendFunction::VTBlend_EaseOut, 4);
+		case EMovieSceneBuiltInEasing::QuartInOut:
+			return Return(EViewTargetBlendFunction::VTBlend_EaseInOut, 4);
+
+		case EMovieSceneBuiltInEasing::QuintIn:
+			return Return(EViewTargetBlendFunction::VTBlend_EaseIn, 5);
+		case EMovieSceneBuiltInEasing::QuintOut:
+			return Return(EViewTargetBlendFunction::VTBlend_EaseOut, 5);
+		case EMovieSceneBuiltInEasing::QuintInOut:
+			return Return(EViewTargetBlendFunction::VTBlend_EaseInOut, 5);
+
+		// UNSUPPORTED
+		case EMovieSceneBuiltInEasing::SinIn:
+		case EMovieSceneBuiltInEasing::SinOut:
+		case EMovieSceneBuiltInEasing::SinInOut:
+		case EMovieSceneBuiltInEasing::CircIn:
+		case EMovieSceneBuiltInEasing::CircOut:
+		case EMovieSceneBuiltInEasing::CircInOut:
+		case EMovieSceneBuiltInEasing::ExpoIn:
+		case EMovieSceneBuiltInEasing::ExpoOut:
+		case EMovieSceneBuiltInEasing::ExpoInOut:
+		case EMovieSceneBuiltInEasing::Custom:
+			break;
+	}
+	return Return(EViewTargetBlendFunction::VTBlend_Linear, 1.f);
+}
+
+void FCameraCutGameHandler::CachePreAnimatedValue(
+		UMovieSceneEntitySystemLinker* Linker,
+		const FSequenceInstance& SequenceInstance)
+{
+	TSharedPtr<FPreAnimatedCameraCutStorage> PreAnimatedStorage = Linker->PreAnimatedState.GetOrCreateStorage<FPreAnimatedCameraCutStorage>();
+
+	IMovieScenePlayer* Player = SequenceInstance.GetPlayer();
+	PreAnimatedStorage->CachePreAnimatedValue(
+			(uint8)0,  // Later this can be an index for split-screen player
+			[Player](uint8 InKey) { return FPreAnimatedCameraCutTraits::CachePreAnimatedValue(Player, InKey); },
+			EPreAnimatedCaptureSourceTracking::AlwaysCache);
+}
+
+FCameraCutGameHandler::FCameraCutGameHandler(
+		UMovieSceneEntitySystemLinker* InLinker,
+		const FSequenceInstance& InSequenceInstance)
+	: Linker(InLinker)
+	, SequenceInstance(InSequenceInstance)
+{
+}
+
+void FCameraCutGameHandler::SetCameraCut(
+		UObject* CameraObject,
+		const FMovieSceneCameraCutParams& CameraCutParams)
+{
+	FCameraCutPlaybackCapabilityCompatibilityWrapper Wrapper(SequenceInstance);
+
+	// If we don't want to update camera cuts, let's bail out.
+	if (!Wrapper.ShouldUpdateCameraCut())
+	{
+		return;
+	}
+
+	IMovieScenePlayer* Player = SequenceInstance.GetPlayer();
+	UObject* PlaybackContext = Player->GetPlaybackContext();
+	UWorld* World = PlaybackContext ? PlaybackContext->GetWorld() : nullptr;
+
+	// Also bail out if we don't have a world running any sort of game.
+	if (World == nullptr || World->GetGameInstance() == nullptr)
+	{
+		return;
+	}
+
+	APlayerController* PC = World->GetGameInstance()->GetFirstLocalPlayerController();
+	APlayerCameraManager* CameraManager = (PC != nullptr) ? PC->PlayerCameraManager : nullptr;
+
+	// If the player controller is missing, there is no camera manager for us to manage the view target
+	// so, again, we bail out.
+	if (CameraManager == nullptr)
+	{
+		return;
+	}
+
+	// Let's get the current view target directly from the player camera manager first. This is because
+	// we don't want to go through GetViewTarget, which checks if the current view target is valid, and
+	// re-assigns it to the player controller if it's not. We don't want this to happen, especially since
+	// a spawnable camera might have just been unspawned, causing SetCameraCut to be called, and we
+	// need to handle this properly.
+	AActor* ViewTarget = CameraManager->PendingViewTarget.Target;
+	if (!ViewTarget)
+	{
+		ViewTarget = CameraManager->ViewTarget.Target;
+	}
+
+	// If we find a camera component inside the provided object, let's make sure we are going to set
+	// its owner as the next view target, and not some component (including the camera component itself).
+	UCameraComponent* CameraComponent = MovieSceneHelpers::CameraComponentFromRuntimeObject(CameraObject);
+	if (CameraComponent && CameraComponent->GetOwner() != CameraObject)
+	{
+		CameraObject = CameraComponent->GetOwner();
+	}
+
+	// If the view target isn't really changing, we don't have much to do.
+	if (CameraObject == ViewTarget)
+	{
+		if (CameraCutParams.bJumpCut)
+		{
+			if (CameraManager)
+			{
+				CameraManager->SetGameCameraCutThisFrame();
+			}
+
+			if (CameraComponent)
+			{
+				CameraComponent->NotifyCameraCut();
+			}
+
+			if (UMovieSceneMotionVectorSimulationSystem* MotionVectorSim = Linker->FindSystem<UMovieSceneMotionVectorSimulationSystem>())
+			{
+				MotionVectorSim->SimulateAllTransforms();
+			}
+		}
+
+		return;
+	}
+
+	// If UnlockIfCameraActor is valid, release lock only if currently locked to the specified object.
+	AActor* UnlockIfCameraActor = Cast<AActor>(CameraCutParams.UnlockIfCameraObject);
+	if (CameraObject == nullptr && UnlockIfCameraActor != nullptr && UnlockIfCameraActor != ViewTarget)
+	{
+		return;
+	}
+
+	AActor* CameraActor = Cast<AActor>(CameraObject);
+	ULocalPlayer* LocalPlayer = (PC != nullptr) ? PC->GetLocalPlayer() : nullptr;
+
+	TOptional<EAspectRatioAxisConstraint> OverrideAspectRatioAxisConstraint;
+	if (Wrapper.CameraCutCapability)
+	{
+		OverrideAspectRatioAxisConstraint = Wrapper.CameraCutCapability->GetAspectRatioAxisConstraintOverride();
+	}
+
+	// CameraObject is null if we need to release control, which can happen here (instead of via pre-animated 
+	// state restore) if we are *blending* back to gameplay, and not cutting back to it at the end of a camera 
+	// cut section. Let's get the pre-animated value and blend back towards it.
+	if (CameraObject == nullptr)
+	{
+		TSharedPtr<FPreAnimatedCameraCutStorage> PreAnimatedStorage = Linker->PreAnimatedState.FindStorage(FPreAnimatedCameraCutStorage::StorageID);
+		FPreAnimatedStorageIndex StorageIndex = PreAnimatedStorage->FindStorageIndex(0);
+		if (ensureMsgf(StorageIndex.IsValid(), TEXT("Blending camera back to gameplay but can't find pre-animated camera info!")))
+		{
+			FPreAnimatedCameraCutState CachedValue = PreAnimatedStorage->GetCachedValue(StorageIndex);
+			CameraObject = CachedValue.LastViewTarget.ResolveObjectPtr();
+			CameraActor = Cast<AActor>(CameraObject);
+			CameraComponent = MovieSceneHelpers::CameraComponentFromRuntimeObject(CameraActor);
+			OverrideAspectRatioAxisConstraint = CachedValue.LastAspectRatioAxisConstraint;
+		}
+	}
+
+	// Time to set the camera cut! How we do it depends on whether we need to do some blending, or a straight cut.
+	bool bDoSetViewTarget = true;
+	FViewTargetTransitionParams TransitionParams;
+	if (CameraCutParams.BlendType.IsSet())
+	{
+		UE_LOG(LogMovieScene, Log, TEXT("Blending into new camera cut: '%s' -> '%s' (blend time: %f)"),
+			(ViewTarget ? *ViewTarget->GetName() : TEXT("None")),
+			(CameraActor ? *CameraActor->GetName() : TEXT("None")),
+			TransitionParams.BlendTime);
+
+		const float PlayRateFactor = Wrapper.CameraCutCapability ? Wrapper.CameraCutCapability->GetCameraBlendPlayRate() : 1.f;
+
+		// Convert known easing functions to their corresponding view target blend parameters.
+		TTuple<EViewTargetBlendFunction, float> BlendFunctionAndExp = BuiltInEasingTypeToBlendFunction(CameraCutParams.BlendType.GetValue());
+		TransitionParams.BlendTime = CameraCutParams.BlendTime * PlayRateFactor;
+		TransitionParams.bLockOutgoing = CameraCutParams.bLockPreviousCamera;
+		TransitionParams.BlendFunction = BlendFunctionAndExp.Get<0>();
+		TransitionParams.BlendExp = BlendFunctionAndExp.Get<1>();
+
+		// Calling SetViewTarget on a camera that we are currently transitioning to will 
+		// result in that transition being aborted, and the view target being set immediately.
+		// We want to avoid that, so let's leave the transition running if it's the case.
+		const AActor* PendingViewTarget = CameraManager->PendingViewTarget.Target;
+		if (CameraActor != nullptr && PendingViewTarget == CameraActor)
+		{
+			UE_LOG(LogMovieScene, Log, TEXT("Camera transition aborted, we are already blending towards the intended camera"));
+			bDoSetViewTarget = false;
+		}
+	}
+	else
+	{
+		UE_LOG(LogMovieScene, Log, TEXT("Starting new camera cut: '%s'"),
+			(CameraActor ? *CameraActor->GetName() : TEXT("None")));
+	}
+	if (bDoSetViewTarget && ensureMsgf(PC, TEXT("Can't set view target when there is no player controller!")))
+	{
+		PC->SetViewTarget(CameraActor, TransitionParams);
+	}
+
+	// Override the aspect ratio constraint if this sequence requires it.
+	if (LocalPlayer != nullptr && OverrideAspectRatioAxisConstraint.IsSet())
+	{
+		LocalPlayer->AspectRatioAxisConstraint = OverrideAspectRatioAxisConstraint.GetValue();
+	}
+
+	// We want to notify of cuts on hard cuts and time jumps, but not on blend cuts.
+	const bool bIsStraightCut = !CameraCutParams.BlendType.IsSet() || CameraCutParams.bJumpCut;
+
+	if (CameraComponent && bIsStraightCut)
+	{
+		CameraComponent->NotifyCameraCut();
+	}
+
+	if (CameraManager)
+	{
+		CameraManager->bClientSimulatingViewTarget = (CameraActor != nullptr);
+
+		if (bIsStraightCut)
+		{
+			CameraManager->SetGameCameraCutThisFrame();
+		}
+	}
+
+	if (bIsStraightCut)
+	{
+		if (UMovieSceneMotionVectorSimulationSystem* MotionVectorSim = Linker->FindSystem<UMovieSceneMotionVectorSimulationSystem>())
+		{
+			MotionVectorSim->SimulateAllTransforms();
+		}
+
+		FOnCameraCutUpdatedParams CameraCutUpdatedParams;
+		CameraCutUpdatedParams.ViewTarget = CameraActor;
+		CameraCutUpdatedParams.ViewTargetCamera = CameraComponent;
+		CameraCutUpdatedParams.bIsJumpCut = true;
+		Wrapper.OnCameraCutUpdated(CameraCutUpdatedParams);
+	}
+}
+
+}  // namespace UE::MovieScene
+
diff --git a/Engine/Source/Runtime/MovieSceneTracks/Private/TrackInstances/MovieSceneCameraCutGameHandler.h b/Engine/Source/Runtime/MovieSceneTracks/Private/TrackInstances/MovieSceneCameraCutGameHandler.h
new file mode 100644
--- /dev/null
+++ b/Engine/Source/Runtime/MovieSceneTracks/Private/TrackInstances/MovieSceneCameraCutGameHandler.h
@@ -0,0 +1,70 @@
+// Copyright Epic Games, Inc. All Rights Reserved.
+
+#pragma once
+
+#include "CoreMinimal.h"
+#include "Evaluation/PreAnimatedState/MovieScenePreAnimatedStateStorage.h"
+#include "Evaluation/PreAnimatedState/MovieScenePreAnimatedStateTypes.h"
+#include "Evaluation/PreAnimatedState/MovieScenePreAnimatedStorageID.inl"
+
+class UMovieSceneEntitySystemLinker;
+struct FMovieSceneCameraCutParams;
+
+namespace UE::MovieScene
+{
+
+struct FCameraCutPlaybackCapability;
+struct FPreAnimatedCameraCutStorage;
+struct FSequenceInstance;
+
+/** Pre-animated view target info */
+struct FPreAnimatedCameraCutState
+{
+	FObjectKey LastLocalPlayer;
+	FObjectKey LastViewTarget;
+	EAspectRatioAxisConstraint LastAspectRatioAxisConstraint;
+};
+
+/** Pre-animated traits for in-game camera cuts */
+struct FPreAnimatedCameraCutTraits : FPreAnimatedStateTraits
+{
+	// Key type is an integer, to be maybe later used as the splitscreen index.
+	using KeyType = uint8;
+	using StorageType = FPreAnimatedCameraCutState;
+
+	static StorageType CachePreAnimatedValue(IMovieScenePlayer* Player, uint8 InKey);
+
+	void RestorePreAnimatedValue(uint8 InKey, const StorageType& CachedValue, const FRestoreStateParams& Params);
+};
+
+/** Pre-aniamted state storage for in-game camera cuts */
+struct FPreAnimatedCameraCutStorage : TPreAnimatedStateStorage<FPreAnimatedCameraCutTraits>
+{
+	static TAutoRegisterPreAnimatedStorageID<FPreAnimatedCameraCutStorage> StorageID;
+};
+
+/**
+ * Utility class for executing camera cuts in game worlds (including PIE).
+ */
+struct FCameraCutGameHandler
+{
+	FCameraCutGameHandler(
+			UMovieSceneEntitySystemLinker* InLinker,
+			const FSequenceInstance& InSequenceInstance);
+
+	/** Sets the given camera cut in all qualifying game worlds. */
+	void SetCameraCut(
+			UObject* CameraObject, 
+			const FMovieSceneCameraCutParams& CameraCutParams);
+
+	/** Cache any pre-animated values required for handling camera cuts in game. */
+	static void CachePreAnimatedValue(
+			UMovieSceneEntitySystemLinker* Linker,
+			const FSequenceInstance& SequenceInstance);
+
+private:
+	UMovieSceneEntitySystemLinker* Linker;
+	const FSequenceInstance& SequenceInstance;
+};
+
+}  // namespace UE::MovieScene
diff --git a/Engine/Source/Runtime/MovieSceneTracks/Private/TrackInstances/MovieSceneCameraCutTrackInstance.cpp b/Engine/Source/Runtime/MovieSceneTracks/Private/TrackInstances/MovieSceneCameraCutTrackInstance.cpp
--- a/Engine/Source/Runtime/MovieSceneTracks/Private/TrackInstances/MovieSceneCameraCutTrackInstance.cpp
+++ b/Engine/Source/Runtime/MovieSceneTracks/Private/TrackInstances/MovieSceneCameraCutTrackInstance.cpp
@@ -1,18 +1,28 @@
 // Copyright Epic Games, Inc. All Rights Reserved.
 
 #include "TrackInstances/MovieSceneCameraCutTrackInstance.h"
+
+#include "Camera/CameraComponent.h"
 #include "ContentStreaming.h"
+#include "Engine/Engine.h"
 #include "Engine/World.h"
+#include "EngineGlobals.h"
+#include "EntitySystem/MovieSceneEntitySystemLinker.h"
+#include "EntitySystem/MovieSceneInstanceRegistry.h"
+#include "EntitySystem/MovieSceneSharedPlaybackState.h"
+#include "Evaluation/CameraCutPlaybackCapability.h"
 #include "Evaluation/MovieSceneEvaluation.h"
+#include "Evaluation/PreAnimatedState/MovieScenePreAnimatedCaptureSource.h"
 #include "GameFramework/Actor.h"
+#include "GameFramework/PlayerController.h"
 #include "Generators/MovieSceneEasingCurves.h"
 #include "IMovieScenePlayer.h"
+#include "MovieSceneCommonHelpers.h"
 #include "MovieSceneTimeHelpers.h"
 #include "Sections/MovieSceneCameraCutSection.h"
+#include "TrackInstances/MovieSceneCameraCutEditorHandler.h"
+#include "TrackInstances/MovieSceneCameraCutGameHandler.h"
 #include "Tracks/MovieSceneCameraCutTrack.h"
-#include "EntitySystem/MovieSceneEntitySystemLinker.h"
-#include "EntitySystem/MovieSceneInstanceRegistry.h"
-#include "MovieSceneCommonHelpers.h"
 
 #include UE_INLINE_GENERATED_CPP_BY_NAME(MovieSceneCameraCutTrackInstance)
 
@@ -23,186 +33,234 @@ namespace UE
 namespace MovieScene
 {
 
-	/** Information about a camera cut's easing (in or out) */
-	struct FBlendedCameraCutEasingInfo
-	{
-		float RootBlendTime = -1.f;
-		TOptional<EMovieSceneBuiltInEasing> BlendType;
-
-		FBlendedCameraCutEasingInfo() {}
-		FBlendedCameraCutEasingInfo(float InRootBlendTime, const TScriptInterface<IMovieSceneEasingFunction>& EasingFunction)
-		{
-			RootBlendTime = InRootBlendTime;
+FCameraCutPlaybackCapabilityCompatibilityWrapper::FCameraCutPlaybackCapabilityCompatibilityWrapper(const FSequenceInstance& SequenceInstance)
+{
+	TSharedRef<FSharedPlaybackState> PlaybackState = SequenceInstance.GetSharedPlaybackState();
+	CameraCutCapability = PlaybackState->FindCapability<FCameraCutPlaybackCapability>();
+	Player = SequenceInstance.GetPlayer();
+}
 
-			// If it's a built-in easing function, get the curve type. We'll try to convert it to what the
-			// player controller knows later, in the movie scene player.
-			const UObject* EaseInScript = EasingFunction.GetObject();
-			if (const UMovieSceneBuiltInEasingFunction* BuiltInEaseIn = Cast<UMovieSceneBuiltInEasingFunction>(EaseInScript))
-			{
-				BlendType = BuiltInEaseIn->Type;
-			}
-		}
-	};
+bool FCameraCutPlaybackCapabilityCompatibilityWrapper::ShouldUpdateCameraCut()
+{
+	if (CameraCutCapability)
+	{
+		return CameraCutCapability->ShouldUpdateCameraCut();
+	}
+	else if (Player)
+	{
+		PRAGMA_DISABLE_DEPRECATION_WARNINGS
+		return Player->CanUpdateCameraCut();
+		PRAGMA_ENABLE_DEPRECATION_WARNINGS
+	}
+	return true;
+}
 
-	/** Camera cut info struct. */
-	struct FBlendedCameraCut
+void FCameraCutPlaybackCapabilityCompatibilityWrapper::OnCameraCutUpdated(const FOnCameraCutUpdatedParams& Params)
+{
+	if (CameraCutCapability)
 	{
-		FInstanceHandle InstanceHandle;
-		TObjectPtr<UMovieSceneSection> Section;
-
-		FMovieSceneObjectBindingID CameraBindingID;
-		FMovieSceneSequenceID OperandSequenceID;
-
-		FFrameNumber LocalStartTime;
-		FFrameNumber LocalEaseInEndTime;
-		FFrameTime LocalContextTime;
-		FFrameNumber LocalEaseOutStartTime;
-		FFrameNumber LocalEndTime;
-
-		FBlendedCameraCutEasingInfo EaseIn;
-		FBlendedCameraCutEasingInfo EaseOut;
-		bool bLockPreviousCamera = false;
-
-		FMovieSceneObjectBindingID PreviousCameraBindingID;
-		FMovieSceneSequenceID PreviousOperandSequenceID;
-
-		float PreviewBlendFactor = -1.f;
-		bool bCanBlend = false;
-
-		FBlendedCameraCut()
-		{}
-		FBlendedCameraCut(const FMovieSceneTrackInstanceInput& InInput, FMovieSceneObjectBindingID InCameraBindingID, FMovieSceneSequenceID InOperandSequenceID) 
-			: InstanceHandle(InInput.InstanceHandle)
-			, Section(InInput.Section)
-			, CameraBindingID(InCameraBindingID)
-			, OperandSequenceID(InOperandSequenceID)
-		{}
-	};
-
-	/** Pre-roll camera cut info struct. */
-	struct FPreRollCameraCut
+		CameraCutCapability->OnCameraCutUpdated(Params);
+	}
+	else if (Player)
 	{
-		FInstanceHandle InstanceHandle;
-		FMovieSceneObjectBindingID CameraBindingID;
-		FTransform CutTransform;
-		bool bHasCutTransform;
-	};
-
-	/** A movie scene pre-animated token that stores a pre-animated camera cut */
-	struct FCameraCutPreAnimatedToken : IMovieScenePreAnimatedGlobalToken
+		// Not quite correct?
+		PRAGMA_DISABLE_DEPRECATION_WARNINGS
+		Player->UpdateCameraCut(Params.ViewTarget, nullptr, Params.bIsJumpCut);
+		PRAGMA_ENABLE_DEPRECATION_WARNINGS
+	}
+}
+
+/** Information about a camera cut's easing (in or out) */
+struct FBlendedCameraCutEasingInfo
+{
+	float RootBlendTime = -1.f;
+	TOptional<EMovieSceneBuiltInEasing> BlendType;
+
+	FBlendedCameraCutEasingInfo() {}
+	FBlendedCameraCutEasingInfo(float InRootBlendTime, const TScriptInterface<IMovieSceneEasingFunction>& EasingFunction)
 	{
-		static FMovieSceneAnimTypeID GetAnimTypeID()
-		{
-			return TMovieSceneAnimTypeID<FCameraCutPreAnimatedToken>();
-		}
+		RootBlendTime = InRootBlendTime;
 
-		virtual void RestoreState(const FRestoreStateParams& RestoreParams) override
+		// If it's a built-in easing function, get the curve type. We'll try to convert it to what the
+		// player controller knows later, in the movie scene player.
+		const UObject* EaseInScript = EasingFunction.GetObject();
+		if (const UMovieSceneBuiltInEasingFunction* BuiltInEaseIn = Cast<UMovieSceneBuiltInEasingFunction>(EaseInScript))
 		{
-			IMovieScenePlayer* Player = RestoreParams.GetTerminalPlayer();
-			if (!ensure(Player))
-			{
-				return;
-			}
-			
-			EMovieSceneCameraCutParams Params;
-			Player->UpdateCameraCut(nullptr, Params);
+			BlendType = BuiltInEaseIn->Type;
 		}
-	};
+	}
+};
 
-	/** The producer class for the pre-animated token above */
-	struct FCameraCutPreAnimatedTokenProducer : IMovieScenePreAnimatedGlobalTokenProducer
+/** Camera cut info struct. */
+struct FBlendedCameraCut
+{
+	FInstanceHandle InstanceHandle;
+	TObjectPtr<UMovieSceneSection> Section;
+
+	FMovieSceneObjectBindingID CameraBindingID;
+	FMovieSceneSequenceID OperandSequenceID;
+
+	FFrameNumber LocalStartTime;
+	FFrameNumber LocalEaseInEndTime;
+	FFrameTime LocalContextTime;
+	FFrameNumber LocalEaseOutStartTime;
+	FFrameNumber LocalEndTime;
+
+	FBlendedCameraCutEasingInfo EaseIn;
+	FBlendedCameraCutEasingInfo EaseOut;
+	bool bLockPreviousCamera = false;
+
+	FMovieSceneObjectBindingID PreviousCameraBindingID;
+	FMovieSceneSequenceID PreviousOperandSequenceID;
+
+	float PreviewBlendFactor = -1.f;
+	bool bCanBlend = false;
+
+	FBlendedCameraCut()
+	{}
+	FBlendedCameraCut(const FMovieSceneTrackInstanceInput& InInput, FMovieSceneObjectBindingID InCameraBindingID, FMovieSceneSequenceID InOperandSequenceID) 
+		: InstanceHandle(InInput.InstanceHandle)
+		, Section(InInput.Section)
+		, CameraBindingID(InCameraBindingID)
+		, OperandSequenceID(InOperandSequenceID)
+	{}
+};
+
+/** Pre-roll camera cut info struct. */
+struct FPreRollCameraCut
+{
+	FInstanceHandle InstanceHandle;
+	FMovieSceneObjectBindingID CameraBindingID;
+	FTransform CutTransform;
+	bool bHasCutTransform;
+};
+
+/** Utility class for executing camera cuts */
+struct FCameraCutAnimator
+{
+	using FCameraCutCache = UMovieSceneCameraCutTrackInstance::FCameraCutCache;
+
+public:
+
+	static UObject* FindBoundObject(FMovieSceneObjectBindingID BindingID, FMovieSceneSequenceID SequenceID, IMovieScenePlayer& Player)
 	{
-		virtual IMovieScenePreAnimatedGlobalTokenPtr CacheExistingState() const override
+		TArrayView<TWeakObjectPtr<>> Objects = BindingID.ResolveBoundObjects(SequenceID, Player);
+		if (Objects.Num() > 0)
 		{
-			return FCameraCutPreAnimatedToken();
+			return Objects[0].Get();
 		}
-	};
+		return nullptr;
+	}
 
-	struct FCameraCutAnimator
+	static bool MatchesCameraCutCache(UObject* CameraActor, const FBlendedCameraCut& Params, const FCameraCutCache& CameraCutCache)
 	{
-		static UObject* FindBoundObject(FMovieSceneObjectBindingID BindingID, FMovieSceneSequenceID SequenceID, IMovieScenePlayer& Player)
-		{
-			TArrayView<TWeakObjectPtr<>> Objects = BindingID.ResolveBoundObjects(SequenceID, Player);
-			if (Objects.Num() > 0)
-			{
-				return Objects[0].Get();
-			}
-			return nullptr;
-		}
+		return CameraActor == CameraCutCache.LastLockedCamera.Get() &&
+			Params.InstanceHandle == CameraCutCache.LastInstanceHandle &&
+			Params.Section == CameraCutCache.LastSection;
+	}
 
-		static bool MatchesCameraCutCache(UObject* CameraActor, const FBlendedCameraCut& Params, const UMovieSceneCameraCutTrackInstance::FCameraCutCache& CameraCutCache)
-		{
-			return CameraActor == CameraCutCache.LastLockedCamera.Get() &&
-				Params.InstanceHandle == CameraCutCache.LastInstanceHandle &&
-				Params.Section == CameraCutCache.LastSection;
-		}
+	static void UpdateCameraCutCache(UObject* CameraActor, const FBlendedCameraCut& Params, FCameraCutCache& OutCameraCutCache)
+	{
+		OutCameraCutCache.LastLockedCamera = CameraActor;
+		OutCameraCutCache.LastInstanceHandle = Params.InstanceHandle;
+		OutCameraCutCache.LastSection = Params.Section;
+	}
+	
+public:
+
+	FCameraCutCache& CameraCutCache;
+#if WITH_EDITOR
+	FCameraCutViewportPreviewer* ViewportPreviewer = nullptr;
+#endif
+
+	FCameraCutAnimator(FCameraCutCache& InCameraCutCache)
+		: CameraCutCache(InCameraCutCache)
+	{
+	}
+
+#if WITH_EDITOR
+	void SetViewportPreviewer(FCameraCutViewportPreviewer* InViewportPreviewer)
+	{
+		ViewportPreviewer = InViewportPreviewer;
+	}
+#endif
+
+public:
 
-		static void UpdateCameraCutCache(UObject* CameraActor, const FBlendedCameraCut& Params, UMovieSceneCameraCutTrackInstance::FCameraCutCache& OutCameraCutCache)
+	void AnimatePreRoll(const FPreRollCameraCut& Params, const FMovieSceneSequenceID& SequenceID, IMovieScenePlayer& Player)
+	{
+		if (Params.bHasCutTransform)
 		{
-			OutCameraCutCache.LastLockedCamera = CameraActor;
-			OutCameraCutCache.LastInstanceHandle = Params.InstanceHandle;
-			OutCameraCutCache.LastSection = Params.Section;
+			FVector Location = Params.CutTransform.GetLocation();
+			IStreamingManager::Get().AddViewLocation(Location);
 		}
-
-		static void AnimatePreRoll(const FPreRollCameraCut& Params, const FMovieSceneContext& Context, const FMovieSceneSequenceID& SequenceID, IMovieScenePlayer& Player)
+		else
 		{
-			if (Params.bHasCutTransform)
+			UObject* CameraObject = FindBoundObject(Params.CameraBindingID, SequenceID, Player);
+
+			if (AActor* Actor = Cast<AActor>(CameraObject))
 			{
-				FVector Location = Params.CutTransform.GetLocation();
+				FVector Location = Actor->GetActorLocation();
 				IStreamingManager::Get().AddViewLocation(Location);
 			}
-			else
-			{
-				UObject* CameraObject = FindBoundObject(Params.CameraBindingID, SequenceID, Player);
-
-				if (AActor* Actor = Cast<AActor>(CameraObject))
-				{
-					FVector Location = Actor->GetActorLocation();
-					IStreamingManager::Get().AddViewLocation(Location);
-				}
-			}
 		}
+	}
 
-		static bool AnimateBlendedCameraCut(const FBlendedCameraCut& Params, UMovieSceneCameraCutTrackInstance::FCameraCutCache& CameraCutCache, const FMovieSceneContext& Context, IMovieScenePlayer& Player)
-		{
-			UObject* CameraActor = FindBoundObject(Params.CameraBindingID, Params.OperandSequenceID, Player);
+	bool AnimateBlendedCameraCut(
+			const FBlendedCameraCut& Params, 
+			UMovieSceneEntitySystemLinker* Linker,
+			const FSequenceInstance& SequenceInstance)
+	{
+		IMovieScenePlayer& Player = *SequenceInstance.GetPlayer();
+		const FMovieSceneContext& Context = SequenceInstance.GetContext();
+
+		UObject* CameraActor = FindBoundObject(Params.CameraBindingID, Params.OperandSequenceID, Player);
 
-			EMovieSceneCameraCutParams CameraCutParams;
-			CameraCutParams.bJumpCut = Context.HasJumped();
-			CameraCutParams.BlendTime = Params.EaseIn.RootBlendTime;
-			CameraCutParams.BlendType = Params.EaseIn.BlendType;
-			CameraCutParams.bLockPreviousCamera = Params.bLockPreviousCamera;
+		FMovieSceneCameraCutParams CameraCutParams;
+		CameraCutParams.bJumpCut = Context.HasJumped();
+		CameraCutParams.BlendTime = Params.EaseIn.RootBlendTime;
+		CameraCutParams.BlendType = Params.EaseIn.BlendType;
+		CameraCutParams.bLockPreviousCamera = Params.bLockPreviousCamera;
 
 #if WITH_EDITOR
-			UObject* PreviousCameraActor = FindBoundObject(Params.PreviousCameraBindingID, Params.PreviousOperandSequenceID, Player);
-			CameraCutParams.PreviousCameraObject = PreviousCameraActor;
-			CameraCutParams.PreviewBlendFactor = Params.PreviewBlendFactor;
-			CameraCutParams.bCanBlend = Params.bCanBlend;
+		UObject* PreviousCameraActor = FindBoundObject(Params.PreviousCameraBindingID, Params.PreviousOperandSequenceID, Player);
+		CameraCutParams.PreviousCameraObject = PreviousCameraActor;
+		CameraCutParams.PreviewBlendFactor = Params.PreviewBlendFactor;
+		CameraCutParams.bCanBlend = Params.bCanBlend;
 #endif
 
-			static const FMovieSceneAnimTypeID CameraAnimTypeID = FMovieSceneAnimTypeID::Unique();
+		const bool bMatchesCache = MatchesCameraCutCache(CameraActor, Params, CameraCutCache);
+		if (!bMatchesCache)
+		{
+			CameraCutParams.UnlockIfCameraObject = CameraCutCache.LastLockedCamera.Get();
+			SetCameraCut(CameraActor, CameraCutParams, Linker, SequenceInstance);
+			UpdateCameraCutCache(CameraActor, Params, CameraCutCache);
+			return true;
+		}
+		else if (CameraActor || CameraCutParams.BlendTime > 0.f)
+		{
+			SetCameraCut(CameraActor, CameraCutParams, Linker, SequenceInstance);
+			return true;
+		}
 
-			const bool bMatchesCache = MatchesCameraCutCache(CameraActor, Params, CameraCutCache);
-			if (!bMatchesCache)
-			{
-				Player.SavePreAnimatedState(CameraAnimTypeID, FCameraCutPreAnimatedTokenProducer());
+		return false;
+	}
 
-				CameraCutParams.UnlockIfCameraObject = CameraCutCache.LastLockedCamera.Get();
-				Player.UpdateCameraCut(CameraActor, CameraCutParams);
-				UpdateCameraCutCache(CameraActor, Params, CameraCutCache);
-				return true;
-			}
-			else if (CameraActor || CameraCutParams.BlendTime > 0.f)
-			{
-				Player.SavePreAnimatedState(CameraAnimTypeID, FCameraCutPreAnimatedTokenProducer());
-	
-				Player.UpdateCameraCut(CameraActor, CameraCutParams);
-				return true;
-			}
+	void SetCameraCut(
+			UObject* CameraObject, 
+			const FMovieSceneCameraCutParams& CameraCutParams, 
+			UMovieSceneEntitySystemLinker* Linker,
+			const FSequenceInstance& SequenceInstance)
+	{
+		FCameraCutGameHandler GameHandler(Linker, SequenceInstance);
+		GameHandler.SetCameraCut(CameraObject, CameraCutParams);
 
-			return false;
-		}
-	};
+#if WITH_EDITOR
+		FCameraCutEditorHandler EditorHandler(Linker, SequenceInstance, *ViewportPreviewer);
+		EditorHandler.SetCameraCut(CameraObject, CameraCutParams);
+#endif
+	}
+};
 
 }  // namespace MovieScene
 }  // namespace UE
@@ -215,14 +273,15 @@ void UMovieSceneCameraCutTrackInstance::OnAnimate()
 	// Gather active camera cuts, and triage pre-rolls from actual cuts.
 	TArray<FPreRollCameraCut> CameraCutPreRolls;
 	TArray<FBlendedCameraCut> CameraCutParams;
-	const FInstanceRegistry* InstanceRegistry = GetLinker()->GetInstanceRegistry();
+
+	UMovieSceneEntitySystemLinker* Linker = GetLinker();
+	const FInstanceRegistry* InstanceRegistry = Linker->GetInstanceRegistry();
 
 	for (const FCameraCutInputInfo& InputInfo : SortedInputInfos)
 	{
 		const FMovieSceneTrackInstanceInput& Input = InputInfo.Input;
 		const FSequenceInstance& SequenceInstance = InstanceRegistry->GetInstance(Input.InstanceHandle);
 		const FMovieSceneContext& Context = SequenceInstance.GetContext();
-		IMovieScenePlayer* Player = SequenceInstance.GetPlayer();
 
 		const UMovieSceneCameraCutSection* Section = Cast<const UMovieSceneCameraCutSection>(Input.Section);
 		const FMovieSceneObjectBindingID CameraBindingID = Section->GetCameraBindingID();
@@ -285,15 +344,19 @@ void UMovieSceneCameraCutTrackInstance::OnAnimate()
 		}
 	}
 
+	FCameraCutAnimator Animator(CameraCutCache);
+#if WITH_EDITOR
+	Animator.SetViewportPreviewer(&ViewportPreviewer);
+#endif
+
 	// For now we only support one pre-roll.
 	if (CameraCutPreRolls.Num() > 0)
 	{
 		FPreRollCameraCut& CameraCutPreRoll = CameraCutPreRolls.Last();
 		const FSequenceInstance& SequenceInstance = InstanceRegistry->GetInstance(CameraCutPreRoll.InstanceHandle);
-		const FMovieSceneContext& Context = SequenceInstance.GetContext();
 		IMovieScenePlayer* Player = SequenceInstance.GetPlayer();
 		const FMovieSceneSequenceID SequenceID = SequenceInstance.GetSequenceID();
-		FCameraCutAnimator::AnimatePreRoll(CameraCutPreRoll, Context, SequenceID, *Player);
+		Animator.AnimatePreRoll(CameraCutPreRoll, SequenceID, *Player);
 	}
 
 	// For now we only support 2 active camera cuts at most (with blending between them).
@@ -376,19 +439,7 @@ void UMovieSceneCameraCutTrackInstance::OnAnimate()
 	if (CameraCutParams.Num() > 0)
 	{
 		const FSequenceInstance& SequenceInstance = InstanceRegistry->GetInstance(FinalCameraCut.InstanceHandle);
-		const FMovieSceneContext& Context = SequenceInstance.GetContext();
-		IMovieScenePlayer* Player = SequenceInstance.GetPlayer();
-		if (FCameraCutAnimator::AnimateBlendedCameraCut(FinalCameraCut, CameraCutCache, Context, *Player))
-		{
-			// Track whether this ever evaluated to take control. If so, we'll want to remove control in OnDestroyed.
-			if (FCameraCutUseData* PlayerUseCount = PlayerUseCounts.Find(Player))
-			{
-				PlayerUseCount->bValid = true;
-				// Remember whether we had blending support the last time we took control of the viewport. This is also
-				// for OnDestroyed.
-				PlayerUseCount->bCanBlend = FinalCameraCut.bCanBlend;
-			}
-		}
+		Animator.AnimateBlendedCameraCut(FinalCameraCut, Linker, SequenceInstance);
 	}
 }
 
@@ -396,30 +447,15 @@ void UMovieSceneCameraCutTrackInstance::OnInputAdded(const FMovieSceneTrackInsta
 {
 	using namespace UE::MovieScene;
 
-	const FInstanceRegistry* InstanceRegistry = GetLinker()->GetInstanceRegistry();
+	UMovieSceneEntitySystemLinker* Linker = GetLinker();
+	const FInstanceRegistry* InstanceRegistry = Linker->GetInstanceRegistry();
 	const FSequenceInstance& SequenceInstance = InstanceRegistry->GetInstance(InInput.InstanceHandle);
-	IMovieScenePlayer* Player = SequenceInstance.GetPlayer();
 
-	int32& UseCount = PlayerUseCounts.FindOrAdd(Player, FCameraCutUseData()).UseCount;
-	++UseCount;
-}
-
-void UMovieSceneCameraCutTrackInstance::OnInputRemoved(const FMovieSceneTrackInstanceInput& InInput)
-{
-	using namespace UE::MovieScene;
-
-	const FInstanceRegistry* InstanceRegistry = GetLinker()->GetInstanceRegistry();
-	const FSequenceInstance& SequenceInstance = InstanceRegistry->GetInstance(InInput.InstanceHandle);
-	IMovieScenePlayer* Player = SequenceInstance.GetPlayer();
-
-	if (FCameraCutUseData* PlayerUseCount = PlayerUseCounts.Find(Player))
-	{
-		PlayerUseCount->UseCount--;
-		if (PlayerUseCount->UseCount == 0)
-		{
-			PlayerUseCounts.Remove(Player);
-		}
-	}
+	FScopedPreAnimatedCaptureSource CaptureSource(Linker, InInput);
+	FCameraCutGameHandler::CachePreAnimatedValue(Linker, SequenceInstance);
+#if WITH_EDITOR
+	FCameraCutEditorHandler::CachePreAnimatedValue(Linker, SequenceInstance);
+#endif
 }
 
 void UMovieSceneCameraCutTrackInstance::OnEndUpdateInputs()
@@ -490,71 +526,9 @@ void UMovieSceneCameraCutTrackInstance::OnEndUpdateInputs()
 
 void UMovieSceneCameraCutTrackInstance::OnDestroyed()
 {
-	using namespace UE::MovieScene;
-
-	const FInstanceRegistry* InstanceRegistry = GetLinker()->GetInstanceRegistry();
-
-	bool bRestoreCamera = false;
-	for (const FCameraCutInputInfo& InputInfo : SortedInputInfos)
-	{
-		const FMovieSceneTrackInstanceInput& Input = InputInfo.Input;
-		const FMovieSceneContext& Context = InstanceRegistry->GetInstance(Input.InstanceHandle).GetContext();
-
-		const UMovieSceneCameraCutSection* Section = Cast<const UMovieSceneCameraCutSection>(Input.Section);
-		if (Context.IsPreRoll())
-		{
-			continue;
-		}
-
-		EMovieSceneCompletionMode CompletionMode = EMovieSceneCompletionMode::KeepState;
-
-		if (Section->EvalOptions.CompletionMode == EMovieSceneCompletionMode::ProjectDefault)
-		{
-			if (const UMovieSceneSequence* OuterSequence = Section->GetTypedOuter<const UMovieSceneSequence>())
-			{
-				CompletionMode = OuterSequence->DefaultCompletionMode;
-			}
-		}
-		else
-		{
-			CompletionMode = Section->EvalOptions.CompletionMode;
-		}
-
-		if (CompletionMode == EMovieSceneCompletionMode::RestoreState)
-		{
-			bRestoreCamera = true;
-			break;
-		}
-	}
-
-	if (bRestoreCamera)
-	{
-		// All sequencer players actually point to the same player controller and view target in a given world,
-		// so we only need to restore the pre-animated state on one sequencer player, like, say, the first one
-		// we still have in use. And we only do that when we have no more inputs active (if we still have some
-		// inputs active, regardless of what sequencer player they belong to, they still have control of the
-		// player controller's view target, so we don't want to mess that up).
-		//
-		// TODO-ludovic: when we have proper splitscreen support, this should be changed heavily.
-		//
-		for (const TPair<IMovieScenePlayer*, FCameraCutUseData>& PlayerUseCount : PlayerUseCounts)
-		{
-			// Restore only if we ever took control
-			if (PlayerUseCount.Value.bValid)
-			{
-				EMovieSceneCameraCutParams Params;
 #if WITH_EDITOR
-				Params.bCanBlend = PlayerUseCount.Value.bCanBlend;
+	// Make sure we don't have any viewport modifiers registered anymore.
+	ViewportPreviewer.ToggleViewportPreviewModifiers(false);
 #endif
-				if (PlayerUseCount.Key)
-				{
-					PlayerUseCount.Key->UpdateCameraCut(nullptr, Params);
-				}
-				break;  // Only do it on the first one.
-			}
-		}
-	}
-
-	PlayerUseCounts.Reset();
 }
 
diff --git a/Engine/Source/Runtime/MovieSceneTracks/Private/TrackInstances/MovieSceneCameraCutTrackInstance.h b/Engine/Source/Runtime/MovieSceneTracks/Private/TrackInstances/MovieSceneCameraCutTrackInstance.h
--- a/Engine/Source/Runtime/MovieSceneTracks/Private/TrackInstances/MovieSceneCameraCutTrackInstance.h
+++ b/Engine/Source/Runtime/MovieSceneTracks/Private/TrackInstances/MovieSceneCameraCutTrackInstance.h
@@ -5,12 +5,37 @@
 #include "CoreMinimal.h"
 #include "EntitySystem/BuiltInComponentTypes.h"
 #include "EntitySystem/TrackInstance/MovieSceneTrackInstance.h"
+#include "TrackInstances/MovieSceneCameraCutViewportPreviewer.h"
 #include "UObject/ObjectMacros.h"
+
 #include "MovieSceneCameraCutTrackInstance.generated.h"
 
+class IMovieScenePlayer;
 class UMovieSceneCameraCutSection;
-namespace UE { namespace MovieScene { struct FCameraCutAnimator; } }
 
+namespace UE::MovieScene
+{ 
+	struct FCameraCutAnimator; 
+	struct FCameraCutPlaybackCapability;
+	struct FOnCameraCutUpdatedParams;
+	struct FSequenceInstance;
+
+	// Backwards compatibilty wrapper for camera cut playback capability.
+	struct FCameraCutPlaybackCapabilityCompatibilityWrapper
+	{
+		FCameraCutPlaybackCapabilityCompatibilityWrapper(const FSequenceInstance& SequenceInstance);
+
+		bool ShouldUpdateCameraCut();
+		void OnCameraCutUpdated(const FOnCameraCutUpdatedParams& Params);
+
+		FCameraCutPlaybackCapability* CameraCutCapability;
+		IMovieScenePlayer* Player;
+	};
+}
+
+/**
+ * Track instance used to animate camera cuts.
+ */
 UCLASS()
 class UMovieSceneCameraCutTrackInstance : public UMovieSceneTrackInstance
 {
@@ -19,11 +44,14 @@ class UMovieSceneCameraCutTrackInstance : public UMovieSceneTrackInstance
 private:
 	virtual void OnAnimate() override;
 	virtual void OnInputAdded(const FMovieSceneTrackInstanceInput& InInput) override;
-	virtual void OnInputRemoved(const FMovieSceneTrackInstanceInput& InInput) override;
 	virtual void OnEndUpdateInputs() override;
 	virtual void OnDestroyed() override;
 
 private:
+	/**
+	 * Stores information about the last set camera in order to differentiate
+	 * between new and pre-existing cuts.
+	 */
 	struct FCameraCutCache
 	{
 		TWeakObjectPtr<> LastLockedCamera;
@@ -31,23 +59,26 @@ private:
 		TObjectPtr<UMovieSceneSection> LastSection;
 	};
 
+	/**
+	 * Track instance input qualified with the global start time of its corresponding
+	 * section, used for sorting inputs and prioritizing more "recent" camera cuts
+	 * over "older" ones.
+	 */
 	struct FCameraCutInputInfo
 	{
 		FMovieSceneTrackInstanceInput Input;
 		float GlobalStartTime = 0.f;
 	};
 
-	struct FCameraCutUseData
-	{
-		int32 UseCount = 0;
-		bool bValid = false;
-		bool bCanBlend = false;
-	};
-
 	FCameraCutCache CameraCutCache;
-	TMap<IMovieScenePlayer*, FCameraCutUseData> PlayerUseCounts;
 	TArray<FCameraCutInputInfo> SortedInputInfos;
 
+#if WITH_EDITOR
+	UE::MovieScene::FCameraCutViewportPreviewer ViewportPreviewer;
+#endif
+
+private:
+
 	friend struct UE::MovieScene::FCameraCutAnimator;
 };
 
diff --git a/Engine/Source/Runtime/MovieSceneTracks/Private/TrackInstances/MovieSceneCameraCutViewportPreviewer.cpp b/Engine/Source/Runtime/MovieSceneTracks/Private/TrackInstances/MovieSceneCameraCutViewportPreviewer.cpp
new file mode 100644
--- /dev/null
+++ b/Engine/Source/Runtime/MovieSceneTracks/Private/TrackInstances/MovieSceneCameraCutViewportPreviewer.cpp
@@ -0,0 +1,164 @@
+// Copyright Epic Games, Inc. All Rights Reserved.
+
+#include "TrackInstances/MovieSceneCameraCutViewportPreviewer.h"
+
+#if WITH_EDITOR
+
+#include "Evaluation/CameraCutPlaybackCapability.h"
+#include "IMovieScenePlayer.h"
+#include "LevelEditorViewport.h"
+#include "MovieSceneCommonHelpers.h"
+#include "TrackInstances/MovieSceneCameraCutEditorHandler.h"
+
+namespace UE::MovieScene
+{
+
+void FCameraCutViewportPreviewerTarget::Get(FLevelEditorViewportClient* InClient, FVector& OutLocation, FRotator& OutRotation, float& OutFOV) const
+{
+	if (CameraComponent)
+	{
+		OutLocation = CameraComponent->GetComponentLocation();
+		OutRotation = CameraComponent->GetComponentRotation();
+		OutFOV = CameraComponent->FieldOfView;
+		return;
+	}
+
+	if (CameraActor)
+	{
+		OutLocation = CameraActor->GetActorLocation();
+		OutRotation = CameraActor->GetActorRotation();
+		OutFOV = -1.f;
+		return;
+	}
+
+	if (PreAnimatedStorage)
+	{
+		FPreAnimatedStorageIndex StorageIndex = PreAnimatedStorage->FindStorageIndex(InClient);
+		if (ensureMsgf(StorageIndex.IsValid(), TEXT("Blending camera to or from editor but can't find pre-animated viewport info!")))
+		{
+			FPreAnimatedCameraCutEditorState CachedValue = PreAnimatedStorage->GetCachedValue(StorageIndex);
+			OutLocation = CachedValue.ViewportLocation;
+			OutRotation = CachedValue.ViewportRotation;
+			OutFOV = CachedValue.ViewportFOV;
+			return;
+		}
+	}
+	
+	// Provide sensible defaults even if we somehow have incorrect data.
+	ensureMsgf(CameraComponent || CameraActor || PreAnimatedStorage, TEXT("Invalid viewport preview target: nothing was set!"));
+	OutLocation = InClient->GetViewLocation();
+	OutRotation = InClient->GetViewRotation();
+	OutFOV = InClient->ViewFOV;
+}
+
+FCameraCutViewportPreviewer::FCameraCutViewportPreviewer()
+{
+}
+
+FCameraCutViewportPreviewer::~FCameraCutViewportPreviewer()
+{
+	if (!ensure(!bViewportModifiersRegistered))
+	{
+		ToggleViewportPreviewModifiers(false);
+	}
+}
+
+void FCameraCutViewportPreviewer::ToggleViewportPreviewModifiers(bool bEnabled)
+{
+	if (GEditor == nullptr || bViewportModifiersRegistered == bEnabled)
+	{
+		return;
+	}
+
+	bViewportModifiersRegistered = bEnabled;
+
+	if (bEnabled)
+	{
+		// We aren't registered with any viewport client, let's grab them all.
+		RegisteredViewportClients = GEditor->GetLevelViewportClients();
+		for (FLevelEditorViewportClient* LevelVC : RegisteredViewportClients)
+		{
+			LevelVC->ViewModifiers.AddRaw(this, &FCameraCutViewportPreviewer::ModifyViewportClientView);
+		}
+
+		// Also listen to viewports changing.
+		GEditor->OnLevelViewportClientListChanged().AddRaw(this, &FCameraCutViewportPreviewer::OnLevelViewportClientListChanged);
+	}
+	else
+	{
+		// Unregister from all the viewport clients we know of.
+		for (FLevelEditorViewportClient* LevelVC : RegisteredViewportClients)
+		{
+			LevelVC->ViewModifiers.RemoveAll(this);
+		}
+		RegisteredViewportClients.Reset();
+
+		// Stop listening to viewports changing.
+		GEditor->OnLevelViewportClientListChanged().RemoveAll(this);
+	}
+}
+
+void FCameraCutViewportPreviewer::SetupBlend(const FCameraCutViewportPreviewerTarget& From, const FCameraCutViewportPreviewerTarget& To, float InBlendFactor)
+{
+	FromTarget = From;
+	ToTarget = To;
+	BlendFactor = InBlendFactor;
+	bApplyViewModifier = true;
+}
+
+void FCameraCutViewportPreviewer::TeardownBlend()
+{
+	bApplyViewModifier = false;
+}
+
+void FCameraCutViewportPreviewer::ModifyViewportClientView(FEditorViewportViewModifierParams& Params)
+{
+	if (!bApplyViewModifier)
+	{
+		return;
+	}
+
+	if (!Params.ViewportClient->AllowsCinematicControl() || Params.ViewportClient->GetViewMode() == VMI_Unknown)
+	{
+		return;
+	}
+
+	FVector FromViewLocation, ToViewLocation;
+	FRotator FromViewRotation, ToViewRotation;
+	float FromViewFOV, ToViewFOV;
+	FLevelEditorViewportClient* ViewportClient = static_cast<FLevelEditorViewportClient*>(Params.ViewportClient);
+
+	FromTarget.Get(ViewportClient, FromViewLocation, FromViewRotation, FromViewFOV);
+	ToTarget.Get(ViewportClient, ToViewLocation, ToViewRotation, ToViewFOV);
+
+	const FVector BlendedLocation = FMath::Lerp(FromViewLocation, ToViewLocation, BlendFactor);
+	const FRotator BlendedRotation = FMath::Lerp(FromViewRotation, ToViewRotation, BlendFactor);
+	const float BlendedFOV = FMath::Lerp(FromViewFOV, ToViewFOV, BlendFactor);
+
+	Params.ViewInfo.Location = BlendedLocation;
+	Params.ViewInfo.Rotation = BlendedRotation;
+	Params.ViewInfo.FOV = BlendedFOV;
+}
+
+void FCameraCutViewportPreviewer::OnLevelViewportClientListChanged()
+{
+	TSet<FLevelEditorViewportClient*> NewVCs(GEditor->GetLevelViewportClients());
+	TSet<FLevelEditorViewportClient*> OldVCs(RegisteredViewportClients);
+
+	// Register our callback on the new clients, remove it from the retired clients.
+	for (FLevelEditorViewportClient* NewVC : NewVCs.Difference(OldVCs))
+	{
+		NewVC->ViewModifiers.AddRaw(this, &FCameraCutViewportPreviewer::ModifyViewportClientView);
+	}
+	for (FLevelEditorViewportClient* OldVC : OldVCs.Difference(NewVCs))
+	{
+		OldVC->ViewModifiers.RemoveAll(this);
+	}
+
+	RegisteredViewportClients = NewVCs.Array();
+}
+
+}  // namespace UE::MovieScene
+
+#endif  // WITH_EDITOR
+
diff --git a/Engine/Source/Runtime/MovieSceneTracks/Private/TrackInstances/MovieSceneCameraCutViewportPreviewer.h b/Engine/Source/Runtime/MovieSceneTracks/Private/TrackInstances/MovieSceneCameraCutViewportPreviewer.h
new file mode 100644
--- /dev/null
+++ b/Engine/Source/Runtime/MovieSceneTracks/Private/TrackInstances/MovieSceneCameraCutViewportPreviewer.h
@@ -0,0 +1,96 @@
+// Copyright Epic Games, Inc. All Rights Reserved.
+
+#pragma once
+
+#if WITH_EDITOR
+
+#include "CoreMinimal.h"
+#include "UObject/WeakObjectPtr.h"
+
+class AActor;
+class FLevelEditorViewportClient;
+class UCameraComponent;
+struct FEditorViewportViewModifierParams;
+struct FMovieSceneCameraCutParams;
+
+namespace UE::MovieScene
+{
+	struct FCameraCutPlaybackCapability;
+	struct FPreAnimatedCameraCutEditorStorage;
+
+	/**
+	 * View target for the camera viewport previewer.
+	 */
+	struct FCameraCutViewportPreviewerTarget
+	{
+		// Either an actor with an optional camera component...
+		AActor* CameraActor = nullptr;
+		UCameraComponent* CameraComponent = nullptr;
+
+		// ...or a pre-animated viewport position, to be found in the given storage.
+		TSharedPtr<FPreAnimatedCameraCutEditorStorage> PreAnimatedStorage;
+
+		/** Returns the target's location, rotation, and FOV */
+		void Get(FLevelEditorViewportClient* InClient, FVector& OutLocation, FRotator& OutRotation, float& OutFOV) const;
+	};
+
+	/**
+	 * Utility class for previewing camera cut blends in editor.
+	 *
+	 * This class is able to register view modifiers on editor viewports and, if needed,
+	 * move the viewport's point-of-view into a blend between two targets. This is used 
+	 * by FCameraCutEditorHandler to preview blending, but should be inactive when no
+	 * blending is needed.
+	 */
+	class FCameraCutViewportPreviewer
+	{
+	public:
+		FCameraCutViewportPreviewer();
+		~FCameraCutViewportPreviewer();
+
+		/**
+		 * Toggle whether this previewer is running a blend preview.
+		 * This will register/unregister the view modifiers on all editor viewports that
+		 * have cinematic control enabled.
+		 */
+		void ToggleViewportPreviewModifiers(bool bEnable);
+
+		/**
+		 * Sets up a blend to be previewed between the specified targets, using the specified
+		 * blend factor. This will only do something if ToggleViewportPreviewModifiers(true) has
+		 * been called.
+		 */
+		void SetupBlend(const FCameraCutViewportPreviewerTarget& From, const FCameraCutViewportPreviewerTarget& To, float InBlendFactor);
+
+		/**
+		 * Stops previewing any blend. This will make the view modifier dormant but won't
+		 * unregister it (that's only done with ToggleViewportPreviewModifiers(false)).
+		 */
+		void TeardownBlend();
+
+	private:
+		void UpdatePreviewLevelViewportClientFromCameraCut(FLevelEditorViewportClient& InViewportClient, UObject* InCameraObject, const FMovieSceneCameraCutParams& CameraCutParams);
+		void ModifyViewportClientView(FEditorViewportViewModifierParams& Params);
+		void OnLevelViewportClientListChanged();
+
+	private:
+		/** Target from which to blend */
+		FCameraCutViewportPreviewerTarget FromTarget;
+		/** Target towards which to blend */
+		FCameraCutViewportPreviewerTarget ToTarget;
+		/** Blend factor between both targets */
+		float BlendFactor;
+
+		/** Whether to apply the modifier with the blended view information */
+		bool bApplyViewModifier;
+
+		/** Whether we have registered viewport modifiers in the editor. */
+		bool bViewportModifiersRegistered = false;
+		/** Which viewport clients we have registered with */
+		TArray<FLevelEditorViewportClient*> RegisteredViewportClients;
+	};
+
+}  // namespace UE::MovieScene
+
+#endif  // WITH_EDITOR
+
diff --git a/Engine/Source/Runtime/MovieSceneTracks/Public/Evaluation/CameraCutPlaybackCapability.h b/Engine/Source/Runtime/MovieSceneTracks/Public/Evaluation/CameraCutPlaybackCapability.h
new file mode 100644
--- /dev/null
+++ b/Engine/Source/Runtime/MovieSceneTracks/Public/Evaluation/CameraCutPlaybackCapability.h
@@ -0,0 +1,49 @@
+// Copyright Epic Games, Inc. All Rights Reserved.
+
+#pragma once
+
+#include "CoreMinimal.h"
+#include "Engine/EngineTypes.h"
+#include "Evaluation/MovieScenePlaybackCapabilities.h"
+
+class UCameraComponent;
+
+namespace UE::MovieScene
+{
+	/**
+	 *
+	 */
+	struct FOnCameraCutUpdatedParams
+	{
+		AActor* ViewTarget = nullptr;
+		UCameraComponent* ViewTargetCamera = nullptr;
+		bool bIsJumpCut = true;
+	};
+
+	/**
+	 * Playback capability for sequences that can run camera cuts.
+	 */
+	struct MOVIESCENETRACKS_API FCameraCutPlaybackCapability
+	{
+		/** Playback capability ID */
+		static TPlaybackCapabilityID<FCameraCutPlaybackCapability> ID;
+
+		/** Whether the associated sequence should execute camera cuts */
+		virtual bool ShouldUpdateCameraCut() { return true; }
+
+		/** The play rate at which to run camera blends */
+		virtual float GetCameraBlendPlayRate() { return 1.f; }
+
+		/** Whether the associated sequence requires a custom aspect ratio axis constraint */
+		virtual TOptional<EAspectRatioAxisConstraint> GetAspectRatioAxisConstraintOverride() { return TOptional<EAspectRatioAxisConstraint>(); }
+
+		/** Called when a camera cut has happened */
+		virtual void OnCameraCutUpdated(const FOnCameraCutUpdatedParams& Params) {}
+
+#if WITH_EDITOR
+		virtual bool ShouldPreviewCameraCutBlends() { return true; }
+#endif
+	};
+
+}  // namespace UE::MovieScene
+
diff --git a/Engine/Source/Runtime/UMG/Public/Animation/UMGSequencePlayer.h b/Engine/Source/Runtime/UMG/Public/Animation/UMGSequencePlayer.h
--- a/Engine/Source/Runtime/UMG/Public/Animation/UMGSequencePlayer.h
+++ b/Engine/Source/Runtime/UMG/Public/Animation/UMGSequencePlayer.h
@@ -72,7 +72,6 @@ public:
 	virtual FMovieSceneRootEvaluationTemplateInstance& GetEvaluationTemplate() override { return RootTemplateInstance; }
 	UMG_API virtual UMovieSceneEntitySystemLinker* ConstructEntitySystemLinker() override;
 	virtual UObject* AsUObject() override { return this; }
-	virtual void UpdateCameraCut(UObject* CameraObject, const EMovieSceneCameraCutParams& CameraCutParams) override {}
 	virtual void SetViewportSettings(const TMap<FViewportClient*, EMovieSceneViewportParams>& ViewportParamsMap) override {}
 	virtual void GetViewportSettings(TMap<FViewportClient*, EMovieSceneViewportParams>& ViewportParamsMap) const override {}
 	UMG_API virtual EMovieScenePlayerStatus::Type GetPlaybackStatus() const override;
