Sequencer: don't cache the aspect ratio axis constraint if there's no local player

diff --git a/Engine/Source/Runtime/MovieSceneTracks/Private/TrackInstances/MovieSceneCameraCutGameHandler.cpp b/Engine/Source/Runtime/MovieSceneTracks/Private/TrackInstances/MovieSceneCameraCutGameHandler.cpp
--- a/Engine/Source/Runtime/MovieSceneTracks/Private/TrackInstances/MovieSceneCameraCutGameHandler.cpp
+++ b/Engine/Source/Runtime/MovieSceneTracks/Private/TrackInstances/MovieSceneCameraCutGameHandler.cpp
@@ -34,7 +34,11 @@ FPreAnimatedCameraCutState FPreAnimatedCameraCutTraits::CachePreAnimatedValue(
 
 		// Save previous aspect ratio axis constraint.
 		ULocalPlayer* LocalPlayer = (PC != nullptr) ? PC->GetLocalPlayer() : nullptr;
-		EAspectRatioAxisConstraint AspectRatioAxisConstraint = LocalPlayer->AspectRatioAxisConstraint;
+		TOptional<EAspectRatioAxisConstraint> AspectRatioAxisConstraint;
+		if (LocalPlayer)
+		{
+			AspectRatioAxisConstraint = LocalPlayer->AspectRatioAxisConstraint;
+		}
 
 		return StorageType{ LocalPlayer, ViewTarget, AspectRatioAxisConstraint };
 	}
@@ -74,9 +78,9 @@ void FPreAnimatedCameraCutTraits::RestorePreAnimatedValue(
 	ULocalPlayer* LocalPlayer = (PC != nullptr) ? 
 		PC->GetLocalPlayer() : 
 		Cast<ULocalPlayer>(CachedValue.LastLocalPlayer.ResolveObjectPtr());
-	if (LocalPlayer)
+	if (LocalPlayer && CachedValue.LastAspectRatioAxisConstraint.IsSet())
 	{
-		LocalPlayer->AspectRatioAxisConstraint = CachedValue.LastAspectRatioAxisConstraint;
+		LocalPlayer->AspectRatioAxisConstraint = CachedValue.LastAspectRatioAxisConstraint.GetValue();
 	}
 }
 
diff --git a/Engine/Source/Runtime/MovieSceneTracks/Private/TrackInstances/MovieSceneCameraCutGameHandler.h b/Engine/Source/Runtime/MovieSceneTracks/Private/TrackInstances/MovieSceneCameraCutGameHandler.h
--- a/Engine/Source/Runtime/MovieSceneTracks/Private/TrackInstances/MovieSceneCameraCutGameHandler.h
+++ b/Engine/Source/Runtime/MovieSceneTracks/Private/TrackInstances/MovieSceneCameraCutGameHandler.h
@@ -22,7 +22,7 @@ struct FPreAnimatedCameraCutState
 {
 	FObjectKey LastLocalPlayer;
 	FObjectKey LastViewTarget;
-	EAspectRatioAxisConstraint LastAspectRatioAxisConstraint;
+	TOptional<EAspectRatioAxisConstraint> LastAspectRatioAxisConstraint;
 };
 
 /** Pre-animated traits for in-game camera cuts */
