Added missing includes for newly added test

diff --git a/Engine/Source/Runtime/MovieSceneTracks/Private/Tests/MovieScenePartialEvaluationTests.cpp b/Engine/Source/Runtime/MovieSceneTracks/Private/Tests/MovieScenePartialEvaluationTests.cpp
--- a/Engine/Source/Runtime/MovieSceneTracks/Private/Tests/MovieScenePartialEvaluationTests.cpp
+++ b/Engine/Source/Runtime/MovieSceneTracks/Private/Tests/MovieScenePartialEvaluationTests.cpp
@@ -6,6 +6,7 @@
 #include "EntitySystem/MovieSceneEntitySystemRunner.h"
 #include "IMovieScenePlayer.h"
 #include "Evaluation/MovieSceneEvaluationTemplateInstance.h"
+#include "Compilation/MovieSceneCompiledDataManager.h"
 #include "MovieSceneSequence.h"
 #include "Tracks/MovieSceneFloatTrack.h"
 #include "Tracks/MovieSceneSubTrack.h"
@@ -13,6 +14,8 @@
 #include "Sections/MovieSceneSubSection.h"
 #include "Misc/AutomationTest.h"
 #include "UObject/StrongObjectPtr.h"
+#include "UObject/Package.h"
+#include "Engine/Engine.h"
 
 #if WITH_DEV_AUTOMATION_TESTS
 
