Sequencer: Blending Hierarchical Bias now correctly supports 'gaps'

Previously when using hbias blending, all lower hbias levels would contribute to the resulting weight for all properties within that blend target, regardless of the presence of a track that animates that property.
For example, given the following structure:

Sequence A (bias 100, weight 1)
   -> Sequence B (bias 200, weight 1)
       -> Sequence C (bias 300, weight 0.25)

If a property was animated in just A and C, Sequence B's weight was still being used and totally overriding any contribution from A; since sequence B doesn't animate the property at all, the initial value was used to fill in the gap.
This is now handled correctly by storing the active hbias 'chain' that contributes to the property in the HierarchicalBlendTarget component.
This component supports up to a maximum of 8 unique HBias contributors and is 16 bytes big.

This change also ensures that the HierarchicalBlendTarget component is removed from entities where it is no longer relevant (ie, it was previously using hbias blends but is no more)

diff --git a/Engine/Source/Runtime/MovieScene/Private/EntitySystem/BuiltInComponentTypes.cpp b/Engine/Source/Runtime/MovieScene/Private/EntitySystem/BuiltInComponentTypes.cpp
--- a/Engine/Source/Runtime/MovieScene/Private/EntitySystem/BuiltInComponentTypes.cpp
+++ b/Engine/Source/Runtime/MovieScene/Private/EntitySystem/BuiltInComponentTypes.cpp
@@ -17,6 +17,71 @@ namespace UE
 namespace MovieScene
 {
 
+static constexpr int16 GInvalidBlendTarget = std::numeric_limits<int16>::lowest();
+
+FHierarchicalBlendTarget::FHierarchicalBlendTarget()
+	// Fill with invalid values
+	: HBiasChain{
+		GInvalidBlendTarget, GInvalidBlendTarget, GInvalidBlendTarget, GInvalidBlendTarget,
+		GInvalidBlendTarget, GInvalidBlendTarget, GInvalidBlendTarget, GInvalidBlendTarget
+	}
+{
+}
+
+void FHierarchicalBlendTarget::Add(int16 InValue)
+{
+	static constexpr int32 MaxNum = UE_ARRAY_COUNT(HBiasChain);
+
+	const int32 Index = Algo::LowerBound(HBiasChain, InValue, TGreater<>());
+	if (Index >= MaxNum)
+	{
+		checkf(false, TEXT("Unable to support more than %d levels of nested hierarchical blending"), MaxNum);
+		return;
+	}
+
+	if (HBiasChain[Index] == InValue)
+	{
+		// Already exists
+		return;
+	}
+
+	if (HBiasChain[MaxNum-1] != GInvalidBlendTarget)
+	{
+		checkf(false, TEXT("Unable to support more than %d levels of nested hierarchical blending"), MaxNum);
+		return;
+	}
+
+	// Right shift by 1 if necessary
+	if (Index < MaxNum-1 && HBiasChain[Index] != GInvalidBlendTarget)
+	{
+		void* Src = HBiasChain + Index;
+		void* Dst = HBiasChain + Index + 1;
+		int32 Slack = MaxNum-Index-2; // -2 to avoid buffer overrun. The last one will always get overwritten
+		FMemory::Memmove(Dst, Src, sizeof(int16)*Slack);
+	}
+
+	// Set the value
+	HBiasChain[Index] = InValue;
+}
+
+int32 FHierarchicalBlendTarget::Num() const
+{
+	return Algo::LowerBound(HBiasChain, GInvalidBlendTarget, TGreater<>());
+}
+
+int16 FHierarchicalBlendTarget::operator[](int32 Index) const
+{
+	static constexpr int32 MaxNum = UE_ARRAY_COUNT(HBiasChain);
+
+	check(Index < MaxNum);
+	return HBiasChain[Index];
+}
+
+TArrayView<const int16> FHierarchicalBlendTarget::AsArray() const
+{
+	return TArrayView<const int16>(HBiasChain, Num());
+}
+
 bool FObjectComponent::IsStrongReference() const
 {
 	return ObjectKey == FObjectKey();
diff --git a/Engine/Source/Runtime/MovieScene/Public/EntitySystem/BuiltInComponentTypes.h b/Engine/Source/Runtime/MovieScene/Public/EntitySystem/BuiltInComponentTypes.h
--- a/Engine/Source/Runtime/MovieScene/Public/EntitySystem/BuiltInComponentTypes.h
+++ b/Engine/Source/Runtime/MovieScene/Public/EntitySystem/BuiltInComponentTypes.h
@@ -326,6 +326,65 @@ private:
 	FObjectKey ObjectKey;
 };
 
+/**
+ * Specifies a unique, sorted path of hbiases that contribute to a blended output
+ * Supports up to 8 unique HBiases in its path
+ */
+struct FHierarchicalBlendTarget
+{
+	/** Default Constructor */
+	MOVIESCENE_API FHierarchicalBlendTarget();
+
+	/**
+	 * Add the specified HBias to this blend target.
+	 * Duplicates are not supported - if the HBias already exists in this chain, this function does nothing.
+	 */
+	MOVIESCENE_API void Add(int16 HBias);
+
+	/**
+	 * Return the number of hbiases that contribute to this blend target
+	 */
+	MOVIESCENE_API int32 Num() const;
+
+	/**
+	 * Return the HBias at the specified index. Out of bounds indices will fail an assertion
+	 */
+	MOVIESCENE_API int16 operator[](int32 Index) const;
+
+	/**
+	 * Convert this blend target into an array view
+	 */
+	MOVIESCENE_API TArrayView<const int16> AsArray() const;
+
+public:
+
+	friend uint32 GetTypeHash(const FHierarchicalBlendTarget& In)
+	{
+		const void* Data = In.HBiasChain;
+		const uint32* Data32 = static_cast<const uint32*>(Data);
+		static_assert(sizeof(HBiasChain) == sizeof(uint32)*4);
+
+		// Use 32 bit ints for hashing speed
+		return HashCombine(Data32[0], Data32[1]) ^ HashCombine(Data32[2], Data32[3]);
+	}
+	friend bool operator<(const FHierarchicalBlendTarget& A, const FHierarchicalBlendTarget& B)
+	{
+		return FMemory::Memcmp(A.HBiasChain, B.HBiasChain, sizeof(HBiasChain)) < 0;
+	}
+	friend bool operator==(const FHierarchicalBlendTarget& A, const FHierarchicalBlendTarget& B)
+	{
+		return FMemory::Memcmp(A.HBiasChain, B.HBiasChain, sizeof(HBiasChain)) == 0;
+	}
+	friend bool operator!=(const FHierarchicalBlendTarget& A, const FHierarchicalBlendTarget& B)
+	{
+		return FMemory::Memcmp(A.HBiasChain, B.HBiasChain, sizeof(HBiasChain)) != 0;
+	}
+
+private:
+	/** 16 bytes of hbias chain. Supports up to a maximum of 8 sub-sequences. */
+	int16 HBiasChain[8];
+};
+
 /**
  * The key to a grouping policy registered on the grouping system (see UMovieSceneEntityGroupingSystem)
  */
@@ -528,7 +587,7 @@ public:
 	TComponentTypeID<FMovieSceneSequenceID> HierarchicalEasingProvider;
 
 	// Defines an HBias level that is the highest blend target for a given set of components that need to blend together
-	TComponentTypeID<int16>       HierarchicalBlendTarget;
+	TComponentTypeID<FHierarchicalBlendTarget> HierarchicalBlendTarget;
 
 	// A float representing the evaluated easing weight
 	TComponentTypeID<double> WeightAndEasingResult;
diff --git a/Engine/Source/Runtime/MovieScene/Public/EntitySystem/IMovieSceneTaskScheduler.h b/Engine/Source/Runtime/MovieScene/Public/EntitySystem/IMovieSceneTaskScheduler.h
--- a/Engine/Source/Runtime/MovieScene/Public/EntitySystem/IMovieSceneTaskScheduler.h
+++ b/Engine/Source/Runtime/MovieScene/Public/EntitySystem/IMovieSceneTaskScheduler.h
@@ -68,24 +68,40 @@ struct FTaskParams
 		bForceConsumeUpstream = false;
 	}
 
+	/**
+	 * Set a custom stat ID for this task
+	 */
 	FTaskParams& Stat(const TStatId& InStatId)
 	{
 		StatId = InStatId;
 		return *this;
 	}
 
+	/**
+	 * Force this task to run on the game thread
+	 */
 	FTaskParams& ForceGameThread()
 	{
 		bForceGameThread = true;
 		return *this;
 	}
 
+	/**
+	 * Force this task to run Pre/Post callbacks even if there is no meaningful work to be done in the body
+	 */
+	FTaskParams& ForcePrePostTask()
+	{
+		bForcePrePostTask = true;
+		return *this;
+	}
+
 #if !(UE_BUILD_SHIPPING || UE_BUILD_TEST)
 	const TCHAR* DebugName;
 #endif
 	TStatId StatId;
 	uint8 bForceGameThread : 1;
 	uint8 bSerialTasks : 1;
+	uint8 bForcePrePostTask : 1;
 	uint8 bForcePropagateDownstream : 1;
 	uint8 bForceConsumeUpstream : 1;
 };
diff --git a/Engine/Source/Runtime/MovieScene/Public/EntitySystem/MovieSceneEntitySystemTask.h b/Engine/Source/Runtime/MovieScene/Public/EntitySystem/MovieSceneEntitySystemTask.h
--- a/Engine/Source/Runtime/MovieScene/Public/EntitySystem/MovieSceneEntitySystemTask.h
+++ b/Engine/Source/Runtime/MovieScene/Public/EntitySystem/MovieSceneEntitySystemTask.h
@@ -614,6 +614,11 @@ private:
 		};
 		FTaskID TaskID = InScheduler->CreateForkedAllocationTask(TaskParams, SharedTask, InFunction, PrelockComponentData, Filter, ReadDependencies, WriteDependencies);
 
+		if (!TaskID && TaskParams.bForcePrePostTask)
+		{
+			TaskID = InScheduler->AddNullTask();
+		}
+
 		if (TaskID)
 		{
 			FTaskID PreTask  = SchedulePreTask(InScheduler, TaskParams, SharedTask, (TaskImpl*)0);
@@ -1499,6 +1504,12 @@ private:
 		};
 
 		FTaskID TaskID = InScheduler->CreateForkedAllocationTask(TaskParams, SharedTask, InFunction, PrelockComponentData, Filter, ReadDependencies, WriteDependencies);
+
+		if (!TaskID && TaskParams.bForcePrePostTask)
+		{
+			TaskID = InScheduler->AddNullTask();
+		}
+
 		if (TaskID)
 		{
 			FTaskID PreTask  = SchedulePreTask(InScheduler, TaskParams, SharedTask, (TaskImpl*)0);
diff --git a/Engine/Source/Runtime/MovieSceneTracks/Private/Systems/MovieSceneMaterialParameterSystem.cpp b/Engine/Source/Runtime/MovieSceneTracks/Private/Systems/MovieSceneMaterialParameterSystem.cpp
--- a/Engine/Source/Runtime/MovieSceneTracks/Private/Systems/MovieSceneMaterialParameterSystem.cpp
+++ b/Engine/Source/Runtime/MovieSceneTracks/Private/Systems/MovieSceneMaterialParameterSystem.cpp
@@ -397,6 +397,9 @@ struct TOverlappingMaterialParameterHandler : Mixin
 			};
 			FBlendInfo IgnoredBlendInfo;
 			FBlendInfo BlendInfo;
+
+			FHierarchicalBlendTarget BlendTarget;
+
 			for (FMovieSceneEntityID Input : Inputs)
 			{
 				FBlendInfo& BlendInfoToUpdate = Linker->EntityManager.HasComponent(Input, BuiltInComponents->Tags.Ignored) ? IgnoredBlendInfo : BlendInfo;
@@ -404,6 +407,8 @@ struct TOverlappingMaterialParameterHandler : Mixin
 				TOptionalComponentReader<int16> HBiasComponent = Linker->EntityManager.ReadComponent(Input, BuiltInComponents->HierarchicalBias);
 				const int16 HBias = HBiasComponent ? *HBiasComponent : 0;
 
+				BlendTarget.Add(HBias);
+
 				if (HBias > BlendInfoToUpdate.HBias)
 				{
 					BlendInfoToUpdate.HBias = HBias;
@@ -440,13 +445,17 @@ struct TOverlappingMaterialParameterHandler : Mixin
 				{
 					if (!Linker->EntityManager.HasComponent(Input, BuiltInComponents->HierarchicalBlendTarget))
 					{
-						Linker->EntityManager.AddComponent(Input, BuiltInComponents->HierarchicalBlendTarget, BlendInfo.HBias);
+						Linker->EntityManager.AddComponent(Input, BuiltInComponents->HierarchicalBlendTarget, BlendTarget);
 					}
 					else
 					{
-						Linker->EntityManager.WriteComponentChecked(Input, BuiltInComponents->HierarchicalBlendTarget, BlendInfo.HBias);
+						Linker->EntityManager.WriteComponentChecked(Input, BuiltInComponents->HierarchicalBlendTarget, BlendTarget);
 					}
 				}
+				else
+				{
+					Linker->EntityManager.RemoveComponent(Input, BuiltInComponents->HierarchicalBlendTarget);
+				}
 
 				// Ensure we have the blender type tag on the inputs.
 				Linker->EntityManager.AddComponent(Input, BlenderTypeTag);
diff --git a/Engine/Source/Runtime/MovieSceneTracks/Private/Systems/MovieScenePropertyInstantiator.cpp b/Engine/Source/Runtime/MovieSceneTracks/Private/Systems/MovieScenePropertyInstantiator.cpp
--- a/Engine/Source/Runtime/MovieSceneTracks/Private/Systems/MovieScenePropertyInstantiator.cpp
+++ b/Engine/Source/Runtime/MovieSceneTracks/Private/Systems/MovieScenePropertyInstantiator.cpp
@@ -652,18 +652,23 @@ void UMovieScenePropertyInstantiatorSystem::UpdatePropertyInfo(const FPropertyPa
 		{
 			MetaDataToUpdate = &IgnoredHBiasEntry;
 		}
-		else if (ThisHBias >= ActiveBiasEntry.HBias)
+		else 
 		{
-			MetaDataToUpdate = &ActiveBiasEntry;
+			ActiveBiasEntry.BlendTarget.Add(ThisHBias);
 
-			if (ThisHBias > ActiveBiasEntry.HBias)
+			if (ThisHBias >= ActiveBiasEntry.HBias)
 			{
-				// We found a greater bias than any we've encountered this far
-				// Reset the empty channel list for the active hbias
-				ActiveBiasEmptyChannels = FChannelMask(true, Params.PropertyDefinition->CompositeSize);
-				ActiveBiasEntry.HBias = ThisHBias;
-
 				MetaDataToUpdate = &ActiveBiasEntry;
+
+				if (ThisHBias > ActiveBiasEntry.HBias)
+				{
+					// We found a greater bias than any we've encountered this far
+					// Reset the empty channel list for the active hbias
+					ActiveBiasEmptyChannels = FChannelMask(true, Params.PropertyDefinition->CompositeSize);
+					ActiveBiasEntry.HBias = ThisHBias;
+
+					MetaDataToUpdate = &ActiveBiasEntry;
+				}
 			}
 		}
 
@@ -989,13 +994,20 @@ void UMovieScenePropertyInstantiatorSystem::InitializeBlendPath(const FPropertyP
 			? Contributors.CreateConstKeyIterator(ContributorKey)
 			: NewContributors.CreateConstKeyIterator(ContributorKey);
 
+		FTypelessMutation Mutation;
+		if (!Params.PropertyInfo->HierarchicalMetaData.bBlendHierarchicalBias)
+		{
+			// Make sure that the hierarchical blend target component does not exist if it no longer has one
+			Mutation.RemoveMask.Set(BuiltInComponents->HierarchicalBlendTarget);
+		}
+
 		for (; ContributorIt; ++ContributorIt)
 		{
 			FEntityBuilder()
 			.Add(BuiltInComponents->BlendChannelInput, BlendChannel)
 			.AddTag(SetupResult.CurrentInfo.BlenderTypeTag)
-			.AddConditional(BuiltInComponents->HierarchicalBlendTarget, Params.PropertyInfo->HierarchicalMetaData.HBias, Params.PropertyInfo->HierarchicalMetaData.bBlendHierarchicalBias)
-			.MutateExisting(&Linker->EntityManager, ContributorIt.Value());
+			.AddConditional(BuiltInComponents->HierarchicalBlendTarget, Params.PropertyInfo->HierarchicalMetaData.BlendTarget, Params.PropertyInfo->HierarchicalMetaData.bBlendHierarchicalBias)
+			.MutateExisting(&Linker->EntityManager, ContributorIt.Value(), Mutation);
 		}
 
 		check(!Linker->EntityManager.HasComponent(Params.PropertyInfo->FinalBlendOutputID, BuiltInComponents->BlendChannelInput));
@@ -1016,6 +1028,11 @@ void UMovieScenePropertyInstantiatorSystem::InitializeBlendPath(const FPropertyP
 	InputMutation.RemoveMask = CleanFastPathMask;
 	InputMutation.RemoveMask.Set(Params.PropertyDefinition->InitialValueType);
 	InputMutation.RemoveMask.Set(BuiltInComponents->Tags.HasAssignedInitialValue);
+	if (!Params.PropertyInfo->HierarchicalMetaData.bBlendHierarchicalBias)
+	{
+		// Make sure that the hierarchical blend target component does not exist if it no longer has one
+		InputMutation.RemoveMask.Set(BuiltInComponents->HierarchicalBlendTarget);
+	}
 	for (FComponentTypeID Component : Params.PropertyDefinition->MetaDataTypes)
 	{
 		InputMutation.RemoveMask.Set(Component);
@@ -1147,7 +1164,7 @@ void UMovieScenePropertyInstantiatorSystem::InitializeBlendPath(const FPropertyP
 
 		FEntityBuilder()
 		.Add(BuiltInComponents->BlendChannelInput, NewBlendChannel)
-		.AddConditional(BuiltInComponents->HierarchicalBlendTarget, Params.PropertyInfo->HierarchicalMetaData.HBias, Params.PropertyInfo->HierarchicalMetaData.bBlendHierarchicalBias)
+		.AddConditional(BuiltInComponents->HierarchicalBlendTarget, Params.PropertyInfo->HierarchicalMetaData.BlendTarget, Params.PropertyInfo->HierarchicalMetaData.bBlendHierarchicalBias)
 		.AddTag(SetupResult.CurrentInfo.BlenderTypeTag)
 		.MutateExisting(&Linker->EntityManager, Contributor, InputMutation);
 	}
diff --git a/Engine/Source/Runtime/MovieSceneTracks/Private/Systems/WeightAndEasingEvaluatorSystem.cpp b/Engine/Source/Runtime/MovieSceneTracks/Private/Systems/WeightAndEasingEvaluatorSystem.cpp
--- a/Engine/Source/Runtime/MovieSceneTracks/Private/Systems/WeightAndEasingEvaluatorSystem.cpp
+++ b/Engine/Source/Runtime/MovieSceneTracks/Private/Systems/WeightAndEasingEvaluatorSystem.cpp
@@ -408,7 +408,7 @@ private:
  */
 struct FConsumerEasingChannelMutation : FEasingChannelMutationBase
 {
-	explicit FConsumerEasingChannelMutation(UMovieSceneEntitySystemLinker* Linker, TSparseArray<uint16>* InInstanceIDToChannel, TMap<TTuple<int16, int16>, uint16>* InCachedHierarchicalBlendTargetChannels)
+	explicit FConsumerEasingChannelMutation(UMovieSceneEntitySystemLinker* Linker, TSparseArray<uint16>* InInstanceIDToChannel, TMap<TTuple<int16, FHierarchicalBlendTarget>, uint16>* InCachedHierarchicalBlendTargetChannels)
 		: FEasingChannelMutationBase(Linker)
 		, InstanceIDToChannel(InInstanceIDToChannel)
 		, CachedHierarchicalBlendTargetChannels(InCachedHierarchicalBlendTargetChannels)
@@ -422,9 +422,9 @@ struct FConsumerEasingChannelMutation : FEasingChannelMutationBase
 	{
 		const FMovieSceneEntityID* EntityIDs = Allocation->GetRawEntityIDs();
 
-		TComponentReader<FInstanceHandle> InstanceHandles        = Allocation->ReadComponents(BuiltInComponents->InstanceHandle);
-		TOptionalComponentWriter<uint16>  ExistingEasingChannels = Allocation->TryWriteComponents(BuiltInComponents->HierarchicalEasingChannel, FEntityAllocationWriteContext::NewAllocation());
-		TOptionalComponentReader<int16>   BlendTargets           = Allocation->TryReadComponents(BuiltInComponents->HierarchicalBlendTarget);
+		TComponentReader<FInstanceHandle>                  InstanceHandles        = Allocation->ReadComponents(BuiltInComponents->InstanceHandle);
+		TOptionalComponentWriter<uint16>                   ExistingEasingChannels = Allocation->TryWriteComponents(BuiltInComponents->HierarchicalEasingChannel, FEntityAllocationWriteContext::NewAllocation());
+		TOptionalComponentReader<FHierarchicalBlendTarget> BlendTargets           = Allocation->TryReadComponents(BuiltInComponents->HierarchicalBlendTarget);
 
 		const bool bHasExistingEasing = ExistingEasingChannels.IsValid();
 
@@ -438,7 +438,7 @@ struct FConsumerEasingChannelMutation : FEasingChannelMutationBase
 
 			for (int32 Index = 0; Index < Allocation->Num(); ++Index)
 			{
-				TTuple<int16, int16> BlendTargetKey = MakeTuple(
+				TTuple<int16, FHierarchicalBlendTarget> BlendTargetKey = MakeTuple(
 					HierarchicalBias ? HierarchicalBias[Index] : 0, BlendTargets[Index]
 					);
 
@@ -493,9 +493,9 @@ struct FConsumerEasingChannelMutation : FEasingChannelMutationBase
 	{
 		FEntityAllocationWriteContext NewAllocation = FEntityAllocationWriteContext::NewAllocation();
 
-		TComponentWriter<uint16>          EasingChannels      = EntityRange.Allocation->WriteComponents(BuiltInComponents->HierarchicalEasingChannel, NewAllocation);
-		TComponentReader<FInstanceHandle> InstanceHandles     = EntityRange.Allocation->ReadComponents(BuiltInComponents->InstanceHandle);
-		TOptionalComponentReader<int16>   BlendTargets        = EntityRange.Allocation->TryReadComponents(BuiltInComponents->HierarchicalBlendTarget);
+		TComponentWriter<uint16>                           EasingChannels  = EntityRange.Allocation->WriteComponents(BuiltInComponents->HierarchicalEasingChannel, NewAllocation);
+		TComponentReader<FInstanceHandle>                  InstanceHandles = EntityRange.Allocation->ReadComponents(BuiltInComponents->InstanceHandle);
+		TOptionalComponentReader<FHierarchicalBlendTarget> BlendTargets    = EntityRange.Allocation->TryReadComponents(BuiltInComponents->HierarchicalBlendTarget);
 
 		if (BlendTargets)
 		{
@@ -507,7 +507,7 @@ struct FConsumerEasingChannelMutation : FEasingChannelMutationBase
 			{
 				const int32 Offset = EntityRange.ComponentStartOffset + Index;
 
-				TTuple<int16, int16> BlendTargetKey
+				TTuple<int16, FHierarchicalBlendTarget> BlendTargetKey
 					= MakeTuple(HierarchicalBias ? HierarchicalBias[Index] : 0, BlendTargets[Index]);
 
 				EasingChannels[Offset] = CachedHierarchicalBlendTargetChannels->FindChecked(BlendTargetKey);
@@ -526,7 +526,7 @@ struct FConsumerEasingChannelMutation : FEasingChannelMutationBase
 private:
 
 	TSparseArray<uint16>* InstanceIDToChannel;
-	TMap<TTuple<int16, int16>, uint16>* CachedHierarchicalBlendTargetChannels;
+	TMap<TTuple<int16, FHierarchicalBlendTarget>, uint16>* CachedHierarchicalBlendTargetChannels;
 };
 
 /**
@@ -726,7 +726,9 @@ bool UMovieSceneHierarchicalEasingInstantiatorSystem::IsRelevantImpl(UMovieScene
 {
 	using namespace UE::MovieScene;
 
-	return CachedInstanceIDToChannel.Num() || InLinker->EntityManager.ContainsComponent(FBuiltInComponentTypes::Get()->HierarchicalEasingProvider);
+	return CachedInstanceIDToChannel.Num()
+		|| InLinker->EntityManager.ContainsComponent(FBuiltInComponentTypes::Get()->HierarchicalEasingProvider)
+		|| InLinker->EntityManager.ContainsComponent(FBuiltInComponentTypes::Get()->Tags.BlendHierarchicalBias);
 }
 
 void UMovieSceneHierarchicalEasingInstantiatorSystem::OnLink()
@@ -891,44 +893,35 @@ void UMovieSceneHierarchicalEasingInstantiatorSystem::FinalizeBlendTargets()
 		struct FHBiasChannelData
 		{
 			uint16 AccumulatorChannelID = UE::MovieScene::INVALID_EASING_CHANNEL;
-			bool bIsBlendTarget = false;
 		};
 		TSortedMap<int16, FHBiasChannelData, TInlineAllocator<16>, TGreater<>> HBiasToChannelData;
+		TSet<FHierarchicalBlendTarget, DefaultKeyFuncs<FHierarchicalBlendTarget>, TInlineSetAllocator<16>> BlendTargetChannelData;
 
 		// Step 1: Gather all the hierarchical blend targets
 		//
 		{
-			auto GatherHierarchicalBlendTargets = [&HBiasToChannelData](const FEntityAllocation* Allocation, const int16* BlendTargets, const int16* HBiases)
-			{
-				const int32 Num = Allocation->Num();
-				if (HBiases)
-				{
-					for (int32 Index = 0; Index < Num; ++Index)
-					{
-						HBiasToChannelData.FindOrAdd(BlendTargets[Index]).bIsBlendTarget = true;
-						HBiasToChannelData.FindOrAdd(HBiases[Index]);
-					}
-				}
-				else
-				{
-					// No HBias components means 0, so add that to the map once
-					HBiasToChannelData.FindOrAdd(0);
-					for (int32 Index = 0; Index < Num; ++Index)
-					{
-						HBiasToChannelData.FindOrAdd(BlendTargets[Index]).bIsBlendTarget = true;
-					}
-				}
-			};
-
 			auto Task = FEntityTaskBuilder()
 			.Read(BuiltInComponents->HierarchicalBlendTarget)
-			.ReadOptional(BuiltInComponents->HierarchicalBias)
 			.FilterNone({ BuiltInComponents->Tags.NeedsUnlink });
 			if (!bChannelsHaveBeenInvalidated)
 			{
 				Task.FilterAll({ BuiltInComponents->Tags.NeedsLink });
 			}
-			Task.Iterate_PerAllocation(&Linker->EntityManager, GatherHierarchicalBlendTargets);
+			Task.Iterate_PerEntity(&Linker->EntityManager,
+				[&BlendTargetChannelData](const FHierarchicalBlendTarget& BlendTarget)
+				{
+					BlendTargetChannelData.Add(BlendTarget);
+				}
+			);
+
+			for (const FHierarchicalBlendTarget& BlendTarget : BlendTargetChannelData)
+			{
+				// Ensure that there are hbias channels for each hbias in the chain
+				for (int16 HBias : BlendTarget.AsArray())
+				{
+					HBiasToChannelData.FindOrAdd(HBias);
+				}
+			}
 		}
 
 		// Step 2: Create channels for each blend target sorted by bias
@@ -955,53 +948,42 @@ void UMovieSceneHierarchicalEasingInstantiatorSystem::FinalizeBlendTargets()
 		{
 			CachedHierarchicalBlendTargetChannels.Empty();
 
-			// HBiasToChannelData is a sorted map which is a sorted array internally.
-			// Unfortunately we don't have access to the underlying array, but we can use iterators to achieve the same results.
-			using FSortedMapIterator = TSortedMap<int16, FHBiasChannelData, TInlineAllocator<16>, TGreater<>>::TConstIterator;
-
-			const int32 NumChannels = HBiasToChannelData.Num();
-			int32 Index = 0;
-			for (FSortedMapIterator It = HBiasToChannelData.CreateConstIterator(); It; ++It, ++Index)
+			for (const FHierarchicalBlendTarget& BlendTarget : BlendTargetChannelData)
 			{
-				if (It.Value().bIsBlendTarget)
-				{
-					const int16 ThisHBias = It.Key();
-					const int32 N = NumChannels - Index - 1;
-					int32 I = N;
+				// This is a blend target - create channels for its entire parent chain
+				uint16 ParentChannel = INVALID_EASING_CHANNEL;
 
-					// This is a blend target - create channels for its entire parent chain
-					uint16 ParentChannel = INVALID_EASING_CHANNEL;
-					for (FSortedMapIterator NextIt = It; NextIt; ++NextIt)
-					{
-						const int16 SourceHBias = NextIt.Key();
+				const int32 N = BlendTarget.Num();
+				int32 I = N;
 
-						FHierarchicalEasingChannelData BlendTargetChannel;
-						// Get our value from the accumulator channel for the source bias which includes all the accumulated weights for that HBias level.
-						// Normally this should just be a single weight
-						BlendTargetChannel.ParentChannel = HBiasToChannelData.FindChecked(SourceHBias).AccumulatorChannelID;
-						BlendTargetChannel.BlendMode = EHierarchicalBlendMode::ChildFirstBlendTarget;
+				for (int16 SourceHBias : BlendTarget.AsArray())
+				{
+					FHierarchicalEasingChannelData BlendTargetChannel;
+					// Get our value from the accumulator channel for the source bias which includes all the accumulated weights for that HBias level.
+					// Normally this should just be a single weight
+					BlendTargetChannel.ParentChannel = HBiasToChannelData.FindChecked(SourceHBias).AccumulatorChannelID;
+					BlendTargetChannel.BlendMode = EHierarchicalBlendMode::ChildFirstBlendTarget;
 
 #if UE_ENABLE_MOVIESCENE_BEZIER_BLENDING
-						BlendTargetChannel.CoeffA = Factorial(N) / (Factorial(I) * Factorial(N-I));
-						BlendTargetChannel.ExpB   = static_cast<float>(N-I);
-						BlendTargetChannel.ExpC   = static_cast<float>(I);
-						BlendTargetChannel.BlendMode = EHierarchicalBlendMode::BezierSeries;
+					BlendTargetChannel.CoeffA = Factorial(N) / (Factorial(I) * Factorial(N-I));
+					BlendTargetChannel.ExpB   = static_cast<float>(N-I);
+					BlendTargetChannel.ExpC   = static_cast<float>(I);
+					BlendTargetChannel.BlendMode = EHierarchicalBlendMode::BezierSeries;
 #else
-						BlendTargetChannel.BlendMode = EHierarchicalBlendMode::ChildFirstBlendTarget;
+					BlendTargetChannel.BlendMode = EHierarchicalBlendMode::ChildFirstBlendTarget;
 #endif
-						const uint16 BlendTargetChannelID = Buffer.AddBlendTargetChannel(BlendTargetChannel);
+					const uint16 BlendTargetChannelID = Buffer.AddBlendTargetChannel(BlendTargetChannel);
 
-						if (NextIt != It)
-						{
-							// If this isn't the first one, connect the last one to this one so we can
-							// combine weights hierarchically up the chain
-							Buffer.Channels[Buffer.Channels.Num()-2].ResultChannel = BlendTargetChannelID;
-						}
+					if (N != I)
+					{
+						// If this isn't the first one, connect the last one to this one so we can
+						// combine weights hierarchically up the chain
+						Buffer.Channels[Buffer.Channels.Num()-2].ResultChannel = BlendTargetChannelID;
+					}
 
-						CachedHierarchicalBlendTargetChannels.Add(MakeTuple(SourceHBias, ThisHBias), BlendTargetChannelID);
+					CachedHierarchicalBlendTargetChannels.Add(MakeTuple(SourceHBias, BlendTarget), BlendTargetChannelID);
 
-						--I;
-					}
+					--I;
 				}
 			}
 		}
@@ -1141,7 +1123,7 @@ void UWeightAndEasingEvaluatorSystem::OnSchedulePersistentTasks(UE::MovieScene::
 		.Read(Components->WeightAndEasingResult)
 		.Read(Components->HierarchicalEasingChannel)
 		.FilterAll({ Components->HierarchicalEasingProvider })  // Only harvest results from entities that are providing results
-		.SetStat(GET_STATID(MovieSceneEval_HarvestEasingTask))
+		.SetParams(FTaskParams(GET_STATID(MovieSceneEval_HarvestEasingTask)).ForcePrePostTask())
 		.Schedule_PerEntity<FHarvestHierarchicalEasings>(&Linker->EntityManager, TaskScheduler, &PreAllocatedComputationData.Channels);
 
 		// Step 3: Apply hierarchical easing results to all entities inside affected sub-sequences.
diff --git a/Engine/Source/Runtime/MovieSceneTracks/Public/Systems/MovieScenePropertyInstantiator.h b/Engine/Source/Runtime/MovieSceneTracks/Public/Systems/MovieScenePropertyInstantiator.h
--- a/Engine/Source/Runtime/MovieSceneTracks/Public/Systems/MovieScenePropertyInstantiator.h
+++ b/Engine/Source/Runtime/MovieSceneTracks/Public/Systems/MovieScenePropertyInstantiator.h
@@ -152,6 +152,7 @@ private:
 			bInUse = false;
 		}
 
+		UE::MovieScene::FHierarchicalBlendTarget BlendTarget;
 		int32 NumContributors = 0;
 		int16 HBias = 0;
 		uint8 bWantsRestoreState : 1;
diff --git a/Engine/Source/Runtime/MovieSceneTracks/Public/Systems/WeightAndEasingEvaluatorSystem.h b/Engine/Source/Runtime/MovieSceneTracks/Public/Systems/WeightAndEasingEvaluatorSystem.h
--- a/Engine/Source/Runtime/MovieSceneTracks/Public/Systems/WeightAndEasingEvaluatorSystem.h
+++ b/Engine/Source/Runtime/MovieSceneTracks/Public/Systems/WeightAndEasingEvaluatorSystem.h
@@ -135,7 +135,7 @@ private:
 	TSparseArray<uint16> CachedInstanceIDToChannel;
 
 	/** Cached output channels for a specific pair (source, target) of hierarchical blends */
-	TMap<TTuple<int16,int16>, uint16> CachedHierarchicalBlendTargetChannels;
+	TMap<TTuple<int16, UE::MovieScene::FHierarchicalBlendTarget>, uint16> CachedHierarchicalBlendTargetChannels;
 
 	UPROPERTY()
 	TObjectPtr<UWeightAndEasingEvaluatorSystem> EvaluatorSystem;
