Fixed up part of CL that does not exist on the branch that CL 28152800 was submitted on

diff --git a/Engine/Source/Runtime/MovieSceneTracks/Private/Systems/MovieSceneCustomPrimitiveDataSystem.cpp b/Engine/Source/Runtime/MovieSceneTracks/Private/Systems/MovieSceneCustomPrimitiveDataSystem.cpp
--- a/Engine/Source/Runtime/MovieSceneTracks/Private/Systems/MovieSceneCustomPrimitiveDataSystem.cpp
+++ b/Engine/Source/Runtime/MovieSceneTracks/Private/Systems/MovieSceneCustomPrimitiveDataSystem.cpp
@@ -215,6 +215,9 @@ namespace UE::MovieScene
 				};
 				FBlendInfo IgnoredBlendInfo;
 				FBlendInfo BlendInfo;
+
+				FHierarchicalBlendTarget BlendTarget;
+
 				for (FMovieSceneEntityID Input : Inputs)
 				{
 					FBlendInfo& BlendInfoToUpdate = Linker->EntityManager.HasComponent(Input, BuiltInComponents->Tags.Ignored) ? IgnoredBlendInfo : BlendInfo;
@@ -222,6 +225,8 @@ namespace UE::MovieScene
 					TOptionalComponentReader<int16> HBiasComponent = Linker->EntityManager.ReadComponent(Input, BuiltInComponents->HierarchicalBias);
 					const int16 HBias = HBiasComponent ? *HBiasComponent : 0;
 
+					BlendTarget.Add(HBias);
+
 					if (HBias > BlendInfoToUpdate.HBias)
 					{
 						BlendInfoToUpdate.HBias = HBias;
@@ -258,13 +263,17 @@ namespace UE::MovieScene
 					{
 						if (!Linker->EntityManager.HasComponent(Input, BuiltInComponents->HierarchicalBlendTarget))
 						{
-							Linker->EntityManager.AddComponent(Input, BuiltInComponents->HierarchicalBlendTarget, BlendInfo.HBias);
+							Linker->EntityManager.AddComponent(Input, BuiltInComponents->HierarchicalBlendTarget, BlendTarget);
 						}
 						else
 						{
-							Linker->EntityManager.WriteComponentChecked(Input, BuiltInComponents->HierarchicalBlendTarget, BlendInfo.HBias);
+							Linker->EntityManager.WriteComponentChecked(Input, BuiltInComponents->HierarchicalBlendTarget, BlendTarget);
 						}
 					}
+					else
+					{
+						Linker->EntityManager.RemoveComponent(Input, BuiltInComponents->HierarchicalBlendTarget);
+					}
 
 					// Ensure we have the blender type tag on the inputs.
 					Linker->EntityManager.AddComponent(Input, BlenderTypeTag);
