Sequencer: Check for valid GEditor

diff --git a/Engine/Source/Runtime/MovieSceneTracks/Private/TrackInstances/MovieSceneCameraCutEditorHandler.cpp b/Engine/Source/Runtime/MovieSceneTracks/Private/TrackInstances/MovieSceneCameraCutEditorHandler.cpp
--- a/Engine/Source/Runtime/MovieSceneTracks/Private/TrackInstances/MovieSceneCameraCutEditorHandler.cpp
+++ b/Engine/Source/Runtime/MovieSceneTracks/Private/TrackInstances/MovieSceneCameraCutEditorHandler.cpp
@@ -29,6 +29,11 @@ void FPreAnimatedCameraCutEditorTraits::RestorePreAnimatedValue(
 		const FPreAnimatedCameraCutEditorState& CachedValue, 
 		const FRestoreStateParams& Params)
 {
+	if (!GEditor)
+	{
+		return;
+	}
+
 	// Check that our pointer is still valid by searching it in active viewports.
 	if (GEditor->GetLevelViewportClients().Find(InKey) == INDEX_NONE)
 	{
@@ -87,6 +92,11 @@ void FCameraCutEditorHandler::SetCameraCut(
 		UObject* CameraObject, 
 		const FMovieSceneCameraCutParams& CameraCutParams)
 {
+	if (!GEditor)
+	{
+		return;
+	}
+
 	FCameraCutPlaybackCapabilityCompatibilityWrapper Wrapper(SequenceInstance);
 
 	// If we don't want to update camera cuts, let's remember it and release the viewports
@@ -236,7 +246,7 @@ void FCameraCutEditorHandler::SetCameraCutForViewport(
 		}
 
 		// If there are selected actors, invalidate the viewports hit proxies, otherwise they won't be selectable afterwards
-		if (ViewportClient.Viewport && GEditor->GetSelectedActorCount() > 0)
+		if (ViewportClient.Viewport && GEditor && GEditor->GetSelectedActorCount() > 0)
 		{
 			ViewportClient.Viewport->InvalidateHitProxy();
 		}
