Sequencer: fix editor preview blending back to gameplay

diff --git a/Engine/Source/Runtime/MovieSceneTracks/Private/TrackInstances/MovieSceneCameraCutEditorHandler.cpp b/Engine/Source/Runtime/MovieSceneTracks/Private/TrackInstances/MovieSceneCameraCutEditorHandler.cpp
--- a/Engine/Source/Runtime/MovieSceneTracks/Private/TrackInstances/MovieSceneCameraCutEditorHandler.cpp
+++ b/Engine/Source/Runtime/MovieSceneTracks/Private/TrackInstances/MovieSceneCameraCutEditorHandler.cpp
@@ -283,7 +283,7 @@ void FCameraCutEditorHandler::SetCameraCutForViewport(
 		else
 		{
 			// We have no "to" camera, so let's blend back to the original viewport position.
-			FromPreviewTarget.PreAnimatedStorage = PreAnimatedStorage;
+			ToPreviewTarget.PreAnimatedStorage = PreAnimatedStorage;
 		}
 
 		ViewportPreviewer.SetupBlend(FromPreviewTarget, ToPreviewTarget, CameraCutParams.PreviewBlendFactor);
