Sequencer: add option to disable auto-arranging of camera cut sections

diff --git a/Engine/Source/Editor/MovieSceneTools/Private/TrackEditors/CameraCutTrackEditor.cpp b/Engine/Source/Editor/MovieSceneTools/Private/TrackEditors/CameraCutTrackEditor.cpp
--- a/Engine/Source/Editor/MovieSceneTools/Private/TrackEditors/CameraCutTrackEditor.cpp
+++ b/Engine/Source/Editor/MovieSceneTools/Private/TrackEditors/CameraCutTrackEditor.cpp
@@ -139,54 +139,46 @@ void FCameraCutTrackEditor::BuildTrackContextMenu(FMenuBuilder& MenuBuilder, UMo
 		"Edit",
 		EUserInterfaceActionType::ToggleButton
 	);
+
+	MenuBuilder.AddMenuEntry(
+		LOCTEXT("AutoArrangeShots", "Auto Arrange"),
+		LOCTEXT("AutoArrangeShotsTooltip", "Auto-arrange and resize sections to fill gaps."),
+		FSlateIcon(),
+		FUIAction(
+			FExecuteAction::CreateSP(this, &FCameraCutTrackEditor::HandleToggleAutoArrangeSectionsExecute, CameraCutTrack),
+			FCanExecuteAction::CreateLambda([=]() { return CameraCutTrack != nullptr; }),
+			FIsActionChecked::CreateLambda([=]() { return CameraCutTrack->IsAutoManagingSections(); })
+			),
+		"Edit",
+		EUserInterfaceActionType::ToggleButton
+	);
 }
 
 void FCameraCutTrackEditor::HandleToggleCanBlendExecute(UMovieSceneCameraCutTrack* CameraCutTrack)
 {
+	const FScopedTransaction Transaction(LOCTEXT("CameraCutTrackSetCanBlend", "Set Camera Cut Track Can Blend"));
+
+	CameraCutTrack->Modify();
+
 	CameraCutTrack->bCanBlend = !CameraCutTrack->bCanBlend;
 
 	if (!CameraCutTrack->bCanBlend)
 	{
-		// Reset all easing and remove overlaps.
-		const UMovieScene* FocusedMovieScene = GetFocusedMovieScene();
-		const FFrameRate TickResolution = FocusedMovieScene->GetTickResolution();
-		const FFrameRate DisplayRate = FocusedMovieScene->GetDisplayRate();
-
-		const TArray<UMovieSceneSection*> Sections = CameraCutTrack->GetAllSections();
-		for (int32 Idx = 1; Idx < Sections.Num(); ++Idx)
-		{
-			UMovieSceneSection* CurSection = Sections[Idx];
-			UMovieSceneSection* PrevSection = Sections[Idx - 1];
-
-			CurSection->Modify();
-
-			TRange<FFrameNumber> CurSectionRange = CurSection->GetRange();
-			TRange<FFrameNumber> PrevSectionRange = PrevSection->GetRange();
-			const FFrameNumber OverlapOrGap = (PrevSectionRange.GetUpperBoundValue() - CurSectionRange.GetLowerBoundValue());
-			if (OverlapOrGap > 0)
-			{
-				const FFrameTime TimeAtHalfBlend = CurSectionRange.GetLowerBoundValue() + FMath::FloorToInt(OverlapOrGap.Value / 2.f);
-				const FFrameNumber FrameAtHalfBlend = FFrameRate::Snap(TimeAtHalfBlend, TickResolution, DisplayRate).CeilToFrame();
+		CameraCutTrack->RearrangeAllSections();
+	}
+}
 
-				PrevSectionRange.SetUpperBoundValue(FrameAtHalfBlend);
-				PrevSection->SetRange(PrevSectionRange);
+void FCameraCutTrackEditor::HandleToggleAutoArrangeSectionsExecute(UMovieSceneCameraCutTrack* CameraCutTrack)
+{
+	const FScopedTransaction Transaction(LOCTEXT("CameraCutTrackSetAutoArrangeSections", "Set Camera Cut Track Auto Arrange"));
 
-				CurSectionRange.SetLowerBoundValue(FrameAtHalfBlend);
-				CurSection->SetRange(CurSectionRange);
-			}
+	CameraCutTrack->Modify();
 
-			CurSection->Easing.AutoEaseInDuration = 0;
-			PrevSection->Easing.AutoEaseOutDuration = 0;
-		}
-		if (Sections.Num() > 0)
-		{
-			Sections[0]->Modify();
+	CameraCutTrack->SetIsAutoManagingSections(!CameraCutTrack->IsAutoManagingSections());
 
-			Sections[0]->Easing.AutoEaseInDuration = 0;
-			Sections[0]->Easing.ManualEaseInDuration = 0;
-			Sections.Last()->Easing.AutoEaseOutDuration = 0;
-			Sections.Last()->Easing.ManualEaseOutDuration = 0;
-		}
+	if (CameraCutTrack->IsAutoManagingSections())
+	{
+		CameraCutTrack->RearrangeAllSections();
 	}
 }
 
diff --git a/Engine/Source/Editor/MovieSceneTools/Private/TrackEditors/CameraCutTrackEditor.h b/Engine/Source/Editor/MovieSceneTools/Private/TrackEditors/CameraCutTrackEditor.h
--- a/Engine/Source/Editor/MovieSceneTools/Private/TrackEditors/CameraCutTrackEditor.h
+++ b/Engine/Source/Editor/MovieSceneTools/Private/TrackEditors/CameraCutTrackEditor.h
@@ -99,6 +99,9 @@ private:
 	/** Called when toggling whether the track can blend camera cuts. */
 	void HandleToggleCanBlendExecute(UMovieSceneCameraCutTrack* CameraCutTrack);
 
+	/** Called when toggling whether the track auto-manages camera cuts. */
+	void HandleToggleAutoArrangeSectionsExecute(UMovieSceneCameraCutTrack* CameraCutTrack);
+
 	/** Delegate for camera button lock state */
 	ECheckBoxState IsCameraLocked() const; 
 
diff --git a/Engine/Source/Runtime/MovieScene/Public/MovieSceneTimeHelpers.h b/Engine/Source/Runtime/MovieScene/Public/MovieSceneTimeHelpers.h
--- a/Engine/Source/Runtime/MovieScene/Public/MovieSceneTimeHelpers.h
+++ b/Engine/Source/Runtime/MovieScene/Public/MovieSceneTimeHelpers.h
@@ -93,6 +93,30 @@ inline TRange<FFrameNumber> MakeDiscreteRange(FFrameNumber MinInclusive, FFrameN
 }
 
 
+/**
+ * Make a new range that includes both the lower and upper bounds of the given range.
+ */
+inline TRange<FFrameNumber> MakeHullRange(const TRange<FFrameNumber>& InRange)
+{
+	return TRange<FFrameNumber>(
+			InRange.HasLowerBound() ? TRangeBound<FFrameNumber>::Inclusive(InRange.GetLowerBoundValue()) : TRangeBound<FFrameNumber>::Open(),
+			InRange.HasUpperBound() ? TRangeBound<FFrameNumber>::Inclusive(InRange.GetUpperBoundValue()) : TRangeBound<FFrameNumber>::Open());
+}
+
+
+/**
+ * Make a new range that includes both the given minimum and maximum.
+ *
+ * @param MinInclusive The minimum value for the inclusive lower bound
+ * @param MaxInclusive The maximum value for the inclusive lower bound
+ * @return A new range.
+ */
+inline TRange<FFrameNumber> MakeHullRange(FFrameNumber MinInclusive, FFrameNumber MaxInclusive)
+{
+	return TRange<FFrameNumber>(TRangeBound<FFrameNumber>::Inclusive(MinInclusive), TRangeBound<FFrameNumber>::Inclusive(MaxInclusive));
+}
+
+
 /**
  * Return whether the given range is empty or zero frames wide
  */
diff --git a/Engine/Source/Runtime/MovieSceneTracks/Private/Tracks/MovieSceneCameraCutTrack.cpp b/Engine/Source/Runtime/MovieSceneTracks/Private/Tracks/MovieSceneCameraCutTrack.cpp
--- a/Engine/Source/Runtime/MovieSceneTracks/Private/Tracks/MovieSceneCameraCutTrack.cpp
+++ b/Engine/Source/Runtime/MovieSceneTracks/Private/Tracks/MovieSceneCameraCutTrack.cpp
@@ -19,6 +19,7 @@
 UMovieSceneCameraCutTrack::UMovieSceneCameraCutTrack( const FObjectInitializer& ObjectInitializer )
 	: Super( ObjectInitializer )
 	, bCanBlend(false)
+	, bAutoArrangeSections(true)
 {
 #if WITH_EDITORONLY_DATA
 	TrackTint = FColor(120, 120, 120, 65);
@@ -65,14 +66,7 @@ UMovieSceneCameraCutSection* UMovieSceneCameraCutTrack::AddNewCameraCut(const FM
 	MovieSceneHelpers::SortConsecutiveSections(MutableView(Sections));
 
 	// Once CameraCuts are sorted fixup the surrounding CameraCuts to fix any gaps
-	if (bCanBlend)
-	{
-		MovieSceneHelpers::FixupConsecutiveBlendingSections(MutableView(Sections), *NewSection, false);
-	}
-	else
-	{
-		MovieSceneHelpers::FixupConsecutiveSections(MutableView(Sections), *NewSection, false);
-	}
+	AutoArrangeSectionsIfNeeded(*NewSection, false);
 
 	return NewSection;
 }
@@ -117,7 +111,7 @@ EMovieSceneTrackEasingSupportFlags UMovieSceneCameraCutTrack::SupportsEasing(FMo
 		{
 			return EMovieSceneTrackEasingSupportFlags::AutomaticEasing | EMovieSceneTrackEasingSupportFlags::ManualEasing;
 		}
-		else if (NumSections > 1)
+		else if (NumSections > 1 && bAutoArrangeSections)
 		{
 			// Find the section with the earliest start time, and the section with the latest end time.
 			int32 EdgeSections[2] = { 0, NumSections - 1 };
@@ -162,6 +156,32 @@ EMovieSceneTrackEasingSupportFlags UMovieSceneCameraCutTrack::SupportsEasing(FMo
 			}
 			return Flags;
 		}
+		else if (NumSections > 1 && !bAutoArrangeSections)
+		{
+			// The given section supports manual easing on one side only if it doesn't overlap with any other section
+			// on that side.
+			const TRange<FFrameNumber> ForSectionRange = Params.ForSection->GetTrueRange();
+			EMovieSceneTrackEasingSupportFlags Flags = EMovieSceneTrackEasingSupportFlags::All;
+			for (int32 Index = 0; Index < NumSections; ++Index)
+			{
+				const UMovieSceneSection* Section(Sections[Index]);
+				const TRange<FFrameNumber> SectionRange(UE::MovieScene::MakeHullRange(Section->GetRange()));
+				if (Section != Params.ForSection)
+				{
+					if (ForSectionRange.HasLowerBound() && SectionRange.Contains(ForSectionRange.GetLowerBoundValue()))
+					{
+						// Lower bound is contained inside another section... disable manual ease-in.
+						Flags &= ~EMovieSceneTrackEasingSupportFlags::ManualEaseIn;
+					}
+					if (ForSectionRange.HasUpperBound() && SectionRange.Contains(ForSectionRange.GetUpperBoundValue()))
+					{
+						// Upper bound is contained inside another section... disable manual ease-out.
+						Flags &= ~EMovieSceneTrackEasingSupportFlags::ManualEaseOut;
+					}
+				}
+			}
+			return Flags;
+		}
 	}
 	return EMovieSceneTrackEasingSupportFlags::AutomaticEasing;
 }
@@ -185,15 +205,8 @@ bool UMovieSceneCameraCutTrack::IsEmpty() const
 void UMovieSceneCameraCutTrack::RemoveSection(UMovieSceneSection& Section)
 {
 	Sections.Remove(&Section);
-
-	if (bCanBlend)
-	{
-		MovieSceneHelpers::FixupConsecutiveBlendingSections(MutableView(Sections), Section, true);
-	}
-	else
-	{
-		MovieSceneHelpers::FixupConsecutiveSections(MutableView(Sections), Section, true);
-	}
+	AutoArrangeSectionsIfNeeded(Section, true);
+	MovieSceneHelpers::SortConsecutiveSections(MutableView(Sections));
 
 	// @todo Sequencer: The movie scene owned by the section is now abandoned.  Should we offer to delete it?  
 }
@@ -201,16 +214,8 @@ void UMovieSceneCameraCutTrack::RemoveSection(UMovieSceneSection& Section)
 void UMovieSceneCameraCutTrack::RemoveSectionAt(int32 SectionIndex)
 {
 	UMovieSceneSection* SectionToDelete = Sections[SectionIndex];
-	if (bCanBlend)
-	{
-		MovieSceneHelpers::FixupConsecutiveBlendingSections(MutableView(Sections), *SectionToDelete, true);
-	}
-	else
-	{
-		MovieSceneHelpers::FixupConsecutiveSections(MutableView(Sections), *SectionToDelete, true);
-	}
-
 	Sections.RemoveAt(SectionIndex);
+	AutoArrangeSectionsIfNeeded(*SectionToDelete, true);
 	MovieSceneHelpers::SortConsecutiveSections(MutableView(Sections));
 }
 
@@ -231,16 +236,7 @@ FText UMovieSceneCameraCutTrack::GetDefaultDisplayName() const
 EMovieSceneSectionMovedResult UMovieSceneCameraCutTrack::OnSectionMoved(UMovieSceneSection& Section, const FMovieSceneSectionMovedParams& Params)
 {
 	const bool bCleanUp = (Params.MoveType == EPropertyChangeType::ValueSet);
-	bool bCleanUpDone = false;
-	if (bCanBlend)
-	{
-		bCleanUpDone = MovieSceneHelpers::FixupConsecutiveBlendingSections(MutableView(Sections), Section, false, bCleanUp);
-	}
-	else
-	{
-		bCleanUpDone = MovieSceneHelpers::FixupConsecutiveSections(MutableView(Sections), Section, false, bCleanUp);
-	}
-
+	const bool bCleanUpDone = AutoArrangeSectionsIfNeeded(Section, false, bCleanUp);
 	return bCleanUpDone ? EMovieSceneSectionMovedResult::SectionsChanged : EMovieSceneSectionMovedResult::None;
 }
 #endif
@@ -282,5 +278,111 @@ void UMovieSceneCameraCutTrack::PreCompileImpl(FMovieSceneTrackPreCompileResult&
 	}
 }
 
+bool UMovieSceneCameraCutTrack::AutoArrangeSectionsIfNeeded(UMovieSceneSection& ChangedSection, bool bWasDeletion, bool bCleanUp)
+{
+	if (bAutoArrangeSections)
+	{
+		if (bCanBlend)
+		{
+			return MovieSceneHelpers::FixupConsecutiveBlendingSections(MutableView(Sections), ChangedSection, bWasDeletion, bCleanUp);
+		}
+		else
+		{
+			return MovieSceneHelpers::FixupConsecutiveSections(MutableView(Sections), ChangedSection, bWasDeletion, bCleanUp);
+		}
+	}
+	return false;
+}
+
+void UMovieSceneCameraCutTrack::RearrangeAllSections()
+{
+	const UMovieScene* MovieScene = GetTypedOuter<UMovieScene>();
+	if (!ensureMsgf(MovieScene, TEXT("Can't auto-arrange sections on a track that doesn't belong to a MovieScene")))
+	{
+		return;
+	}
+
+	// Sort sections by time.
+	MovieSceneHelpers::SortConsecutiveSections(MutableView(Sections));
+
+	// Go over the sections and change anything we don't want to support.
+	const bool bRemoveGaps = bAutoArrangeSections;
+	const bool bRemoveOverlaps = !bCanBlend;
+	const bool bRemoveEasings = !bCanBlend;
+
+	const FFrameRate TickResolution = MovieScene->GetTickResolution();
+	const FFrameRate DisplayRate = MovieScene->GetDisplayRate();
+
+	for (int32 Idx = 1; Idx < Sections.Num(); ++Idx)
+	{
+		UMovieSceneSection* CurSection = Sections[Idx];
+		UMovieSceneSection* PrevSection = Sections[Idx - 1];
+
+		TRange<FFrameNumber> CurSectionRange = CurSection->GetRange();
+		TRange<FFrameNumber> PrevSectionRange = PrevSection->GetRange();
+
+		// See if these two sections overlap or have a gap between them.
+		FFrameNumber OverlapOrGap(0);
+		if (PrevSectionRange.HasUpperBound() && CurSectionRange.HasLowerBound())
+		{
+			OverlapOrGap = (PrevSectionRange.GetUpperBoundValue() - CurSectionRange.GetLowerBoundValue());
+		}
+
+		// If there's an overlap and we don't want it, resize the sections so that they start/stop
+		// around the middle of the overlap zone (not the exact middle... we snap to the display rate)
+		// Similarly, if there's a gap and we don't want it, make the sections meet in the middle of it.
+		TOptional<FFrameNumber> MeetupFrame;
+		if ((bRemoveOverlaps && OverlapOrGap > 0) || (bRemoveGaps && OverlapOrGap < 0))
+		{
+			const FFrameTime TimeAtHalf = CurSectionRange.GetLowerBoundValue() + FMath::FloorToInt(OverlapOrGap.Value / 2.f);
+			MeetupFrame = FFrameRate::Snap(TimeAtHalf, TickResolution, DisplayRate).CeilToFrame();
+		}
+
+		if (MeetupFrame.IsSet())
+		{
+			PrevSectionRange.SetUpperBoundValue(MeetupFrame.GetValue());
+			PrevSection->SetRange(PrevSectionRange);
+
+			CurSectionRange.SetLowerBoundValue(MeetupFrame.GetValue());
+			CurSection->SetRange(CurSectionRange);
+		}
+
+		// Remove any easings from blending overlaps or manual ease-in/out.
+		if (bRemoveEasings)
+		{
+			CurSection->Easing.AutoEaseInDuration = 0;
+			CurSection->Easing.ManualEaseInDuration = 0;
+
+			PrevSection->Easing.AutoEaseOutDuration = 0;
+			PrevSection->Easing.ManualEaseOutDuration = 0;
+		}
+
+		if (MeetupFrame.IsSet() || bRemoveEasings)
+		{
+			CurSection->Modify();
+		}
+	}
+
+	// Remove begin/end easings.
+	if (bRemoveEasings && Sections.Num() > 0)
+	{
+		UMovieSceneSection* FirstSection = Sections[0];
+
+		FirstSection->Easing.AutoEaseInDuration = 0;
+		FirstSection->Easing.ManualEaseInDuration = 0;
+
+		FirstSection->Modify();
+	}
+	if (bRemoveEasings && Sections.Num() > 1)
+	{
+		UMovieSceneSection* LastSection = Sections.Last();
+
+		LastSection->Easing.AutoEaseOutDuration = 0;
+		LastSection->Easing.ManualEaseOutDuration = 0;
+
+		LastSection->Modify();
+	}
+}
+
 #undef LOCTEXT_NAMESPACE
 
diff --git a/Engine/Source/Runtime/MovieSceneTracks/Public/Tracks/MovieSceneCameraCutTrack.h b/Engine/Source/Runtime/MovieSceneTracks/Public/Tracks/MovieSceneCameraCutTrack.h
--- a/Engine/Source/Runtime/MovieSceneTracks/Public/Tracks/MovieSceneCameraCutTrack.h
+++ b/Engine/Source/Runtime/MovieSceneTracks/Public/Tracks/MovieSceneCameraCutTrack.h
@@ -62,16 +62,31 @@ public:
 	virtual FText GetDefaultDisplayName() const override;
 #endif
 
+	/** @return Whether camera cut sections should automatically resize to fill gaps */
+	bool IsAutoManagingSections() const
+	{
+		return bAutoArrangeSections;
+	}
+
+	/** Sets whether camera cut sections should automatically resize to fill gaps */
+	void SetIsAutoManagingSections(bool bInAutoArrangeSections)
+	{
+		bAutoArrangeSections = bInAutoArrangeSections;
+	}
+
 #if WITH_EDITOR
 	virtual EMovieSceneSectionMovedResult OnSectionMoved(UMovieSceneSection& Section, const FMovieSceneSectionMovedParams& Params) override;
 #endif
 
+	MOVIESCENETRACKS_API void RearrangeAllSections();
 	MOVIESCENETRACKS_API FFrameNumber FindEndTimeForCameraCut(FFrameNumber StartTime);
 
 protected:
 
 	virtual void PreCompileImpl(FMovieSceneTrackPreCompileResult& OutPreCompileResult) override;
 
+	bool AutoArrangeSectionsIfNeeded(UMovieSceneSection& ChangedSection, bool bWasDeletion, bool bCleanUp = false);
+
 public:
 	UPROPERTY()
 	bool bCanBlend;
@@ -81,4 +96,9 @@ private:
 	/** All movie scene sections. */
 	UPROPERTY()
 	TArray<TObjectPtr<UMovieSceneSection>> Sections;
+
+	/** Whether camera cut sections should automatically resize to fill gaps */
+	UPROPERTY()
+	bool bAutoArrangeSections;
 };
+
