Sequencer: Defer removal of hierarchical blend targets to UMovieSceneHierarchicalEasingFinalizationSystem so it can retain easing channels if necessary

CL 28223790 fixed a problem where invalid hierarchical blend channels could remain on a contributor even after the blend target component was removed. Unfortunately this had the effect of always removing hierarchical easing components that didn't even relate to Blend Targets that may have beenpreviously  added by UMovieSceneHierarchicalEasingInstantiatorSystem.

In order to resolve these order dependencies we now add a tag to contributors to indicate to UMovieSceneHierarchicalEasingInstantiatorSystem that it needs to remove the blend target. In this way, when the blend target is flagged for removal, it will revert back to the regular hbias weight path, and remove the channel altogether if necessary.

diff --git a/Engine/Source/Runtime/MovieScene/Private/EntitySystem/BuiltInComponentTypes.cpp b/Engine/Source/Runtime/MovieScene/Private/EntitySystem/BuiltInComponentTypes.cpp
--- a/Engine/Source/Runtime/MovieScene/Private/EntitySystem/BuiltInComponentTypes.cpp
+++ b/Engine/Source/Runtime/MovieScene/Private/EntitySystem/BuiltInComponentTypes.cpp
@@ -261,6 +261,7 @@ FBuiltInComponentTypes::FBuiltInComponentTypes()
 	Tags.Finished                = ComponentRegistry->NewTag(TEXT("Finished Evaluating"));
 	Tags.Ignored                 = ComponentRegistry->NewTag(TEXT("Ignored"));
 	Tags.DontOptimizeConstants   = ComponentRegistry->NewTag(TEXT("Don't Optimize Constants"));
+	Tags.RemoveHierarchicalBlendTarget = ComponentRegistry->NewTag(TEXT("Remove Hierarchical Blend Target"));
 	Tags.FixedTime               = ComponentRegistry->NewTag(TEXT("Fixed Time"));
 	Tags.PreRoll                 = ComponentRegistry->NewTag(TEXT("Pre Roll"));
 	Tags.SectionPreRoll          = ComponentRegistry->NewTag(TEXT("Section Pre Roll"));
diff --git a/Engine/Source/Runtime/MovieScene/Public/EntitySystem/BuiltInComponentTypes.h b/Engine/Source/Runtime/MovieScene/Public/EntitySystem/BuiltInComponentTypes.h
--- a/Engine/Source/Runtime/MovieScene/Public/EntitySystem/BuiltInComponentTypes.h
+++ b/Engine/Source/Runtime/MovieScene/Public/EntitySystem/BuiltInComponentTypes.h
@@ -658,6 +658,8 @@ public:
 
 		FComponentTypeID DontOptimizeConstants;
 
+		FComponentTypeID RemoveHierarchicalBlendTarget;
+
 	} Tags;
 
 	struct
diff --git a/Engine/Source/Runtime/MovieSceneTracks/Private/Systems/MovieScenePropertyInstantiator.cpp b/Engine/Source/Runtime/MovieSceneTracks/Private/Systems/MovieScenePropertyInstantiator.cpp
--- a/Engine/Source/Runtime/MovieSceneTracks/Private/Systems/MovieScenePropertyInstantiator.cpp
+++ b/Engine/Source/Runtime/MovieSceneTracks/Private/Systems/MovieScenePropertyInstantiator.cpp
@@ -998,8 +998,7 @@ void UMovieScenePropertyInstantiatorSystem::InitializeBlendPath(const FPropertyP
 		if (!Params.PropertyInfo->HierarchicalMetaData.bBlendHierarchicalBias)
 		{
 			// Make sure that the hierarchical blend target component does not exist if it no longer has one
-			Mutation.RemoveMask.Set(BuiltInComponents->HierarchicalBlendTarget);
-			Mutation.RemoveMask.Set(BuiltInComponents->HierarchicalEasingChannel);
+			Mutation.AddMask.Set(BuiltInComponents->Tags.RemoveHierarchicalBlendTarget);
 		}
 
 		for (; ContributorIt; ++ContributorIt)
@@ -1032,8 +1031,7 @@ void UMovieScenePropertyInstantiatorSystem::InitializeBlendPath(const FPropertyP
 	if (!Params.PropertyInfo->HierarchicalMetaData.bBlendHierarchicalBias)
 	{
 		// Make sure that the hierarchical blend target component does not exist if it no longer has one
-		InputMutation.RemoveMask.Set(BuiltInComponents->HierarchicalBlendTarget);
-		InputMutation.RemoveMask.Set(BuiltInComponents->HierarchicalEasingChannel);
+		InputMutation.AddMask.Set(BuiltInComponents->Tags.RemoveHierarchicalBlendTarget);
 	}
 	for (FComponentTypeID Component : Params.PropertyDefinition->MetaDataTypes)
 	{
diff --git a/Engine/Source/Runtime/MovieSceneTracks/Private/Systems/WeightAndEasingEvaluatorSystem.cpp b/Engine/Source/Runtime/MovieSceneTracks/Private/Systems/WeightAndEasingEvaluatorSystem.cpp
--- a/Engine/Source/Runtime/MovieSceneTracks/Private/Systems/WeightAndEasingEvaluatorSystem.cpp
+++ b/Engine/Source/Runtime/MovieSceneTracks/Private/Systems/WeightAndEasingEvaluatorSystem.cpp
@@ -311,6 +311,7 @@ struct FEasingChannelMutationBase : IMovieSceneConditionalEntityMutation
 		if (StaleEntities.Num())
 		{
 			FComponentMask ComponentsToRemove;
+			ComponentsToRemove.Set(BuiltInComponents->HierarchicalBlendTarget);
 			ComponentsToRemove.Set(BuiltInComponents->HierarchicalEasingChannel);
 			ComponentsToRemove.Set(BuiltInComponents->WeightAndEasingResult);
 
@@ -426,12 +427,13 @@ struct FConsumerEasingChannelMutation : FEasingChannelMutationBase
 		TOptionalComponentWriter<uint16>                   ExistingEasingChannels = Allocation->TryWriteComponents(BuiltInComponents->HierarchicalEasingChannel, FEntityAllocationWriteContext::NewAllocation());
 		TOptionalComponentReader<FHierarchicalBlendTarget> BlendTargets           = Allocation->TryReadComponents(BuiltInComponents->HierarchicalBlendTarget);
 
+		const bool bRemoveBlendTarget = Allocation->HasComponent(BuiltInComponents->Tags.RemoveHierarchicalBlendTarget);
 		const bool bHasExistingEasing = ExistingEasingChannels.IsValid();
 
 		// Loop through each entity and check whether it has easing.
 		// If so, either assign the new easing channel or mark it for mutation.
 		// If not, mark it for removal
-		if (BlendTargets)
+		if (BlendTargets && !bRemoveBlendTarget)
 		{
 			check(CachedHierarchicalBlendTargetChannels);
 			TOptionalComponentReader<int16> HierarchicalBias = Allocation->TryReadComponents(BuiltInComponents->HierarchicalBias);
@@ -1005,6 +1007,12 @@ void UMovieSceneHierarchicalEasingInstantiatorSystem::FinalizeBlendTargets()
 		ensure(bChannelsHaveBeenInvalidated || !ConsumerMutation.HasStaleComponents());
 		ConsumerMutation.RemoveStaleComponents(Linker);
 	}
+
+	if (Linker->EntityManager.ContainsComponent(BuiltInComponents->Tags.RemoveHierarchicalBlendTarget))
+	{
+		FRemoveSingleMutation RemoveTag(BuiltInComponents->Tags.RemoveHierarchicalBlendTarget);
+		Linker->EntityManager.MutateAll(FEntityComponentFilter().All({ BuiltInComponents->Tags.RemoveHierarchicalBlendTarget }), RemoveTag);
+	}
 }
 
 UMovieSceneHierarchicalEasingFinalizationSystem::UMovieSceneHierarchicalEasingFinalizationSystem(const FObjectInitializer& ObjInit)
