make sure that metasound is registered when retrieving default parameters.

diff --git a/Engine/Source/Runtime/MovieSceneTracks/Private/Sections/MovieSceneAudioSection.cpp b/Engine/Source/Runtime/MovieSceneTracks/Private/Sections/MovieSceneAudioSection.cpp
--- a/Engine/Source/Runtime/MovieSceneTracks/Private/Sections/MovieSceneAudioSection.cpp
+++ b/Engine/Source/Runtime/MovieSceneTracks/Private/Sections/MovieSceneAudioSection.cpp
@@ -154,7 +154,7 @@ EMovieSceneChannelProxyType  UMovieSceneAudioSection::CacheChannelProxy()
 #endif
 
 	using namespace MovieSceneAudioSectionPrivate;
-	SetupSoundInputParameters(GetSound());
+	SetupSoundInputParameters(Sound);
 	AddInputChannels<FMovieSceneFloatChannel, float>(this, Channels);
 	AddInputChannels<FMovieSceneBoolChannel, bool>(this, Channels);
 	AddInputChannels<FMovieSceneIntegerChannel, int32>(this, Channels);
@@ -166,11 +166,13 @@ EMovieSceneChannelProxyType  UMovieSceneAudioSection::CacheChannelProxy()
 	return EMovieSceneChannelProxyType::Dynamic;
 }
 
-void UMovieSceneAudioSection::SetupSoundInputParameters(const USoundBase* InSoundBase)
+void UMovieSceneAudioSection::SetupSoundInputParameters(USoundBase* InSoundBase)
 {
 	// Populate with defaults.
 	if (InSoundBase)
 	{
+		InSoundBase->InitResources();
+
 		TArray<FAudioParameter> DefaultParams;
 		InSoundBase->GetAllDefaultParameters(DefaultParams);
 
diff --git a/Engine/Source/Runtime/MovieSceneTracks/Public/Sections/MovieSceneAudioSection.h b/Engine/Source/Runtime/MovieSceneTracks/Public/Sections/MovieSceneAudioSection.h
--- a/Engine/Source/Runtime/MovieSceneTracks/Public/Sections/MovieSceneAudioSection.h
+++ b/Engine/Source/Runtime/MovieSceneTracks/Public/Sections/MovieSceneAudioSection.h
@@ -212,7 +212,7 @@ private:
 		}	
 	}
 
-	MOVIESCENETRACKS_API void SetupSoundInputParameters(const USoundBase* InSoundBase);
+	MOVIESCENETRACKS_API void SetupSoundInputParameters(USoundBase* InSoundBase);
 
 	/** The sound cue or wave that this section plays */
 	UPROPERTY(EditAnywhere, Category="Audio", BlueprintGetter=GetSound, BlueprintSetter=SetSound)
