Sequencer: upgrade visibility track to ECS evaluation

diff --git a/Engine/Plugins/Importers/USDImporter/Source/USDUtilities/Private/USDPrimConversion.cpp b/Engine/Plugins/Importers/USDImporter/Source/USDUtilities/Private/USDPrimConversion.cpp
--- a/Engine/Plugins/Importers/USDImporter/Source/USDUtilities/Private/USDPrimConversion.cpp
+++ b/Engine/Plugins/Importers/USDImporter/Source/USDUtilities/Private/USDPrimConversion.cpp
@@ -48,11 +48,13 @@
 #include "Sections/MovieSceneBoolSection.h"
 #include "Sections/MovieSceneColorSection.h"
 #include "Sections/MovieSceneFloatSection.h"
+#include "Sections/MovieSceneVisibilitySection.h"
 #include "Tracks/MovieScene3DTransformTrack.h"
 #include "Tracks/MovieSceneBoolTrack.h"
 #include "Tracks/MovieSceneColorTrack.h"
 #include "Tracks/MovieSceneFloatTrack.h"
 #include "Tracks/MovieScenePropertyTrack.h"
+#include "Tracks/MovieSceneVisibilityTrack.h"
 
 #if USE_USD_SDK
 
@@ -579,6 +581,75 @@ bool UsdToUnreal::ConvertBoolTimeSamples( const UE::FUsdStage& Stage, const TArr
 	return true;
 }
 
+bool UsdToUnreal::ConvertBoolTimeSamples(const UE::FUsdStage& Stage, const TArray<double>& UsdTimeSamples, const TFunction<bool(double)>& ReaderFunc, UMovieSceneVisibilityTrack& MovieSceneTrack, const FMovieSceneSequenceTransform& SequenceTransform)
+{
+	if (!ReaderFunc)
+	{
+		return false;
+	}
+
+	const UMovieScene* MovieScene = MovieSceneTrack.GetTypedOuter< UMovieScene >();
+	if (!MovieScene)
+	{
+		return false;
+	}
+
+	const FFrameRate Resolution = MovieScene->GetTickResolution();
+	const FFrameRate DisplayRate = MovieScene->GetDisplayRate();
+
+	FScopedUsdAllocs Allocs;
+
+	pxr::UsdStageRefPtr UsdStage{ Stage };
+	FUsdStageInfo StageInfo{ Stage };
+
+	TArray< FFrameNumber > FrameNumbers;
+	FrameNumbers.Reserve(UsdTimeSamples.Num());
+
+	TArray< bool > SectionValues;
+	SectionValues.Reserve(UsdTimeSamples.Num());
+
+	const double StageTimeCodesPerSecond = UsdStage->GetTimeCodesPerSecond();
+	const FFrameRate StageFrameRate(StageTimeCodesPerSecond, 1);
+
+	double LastTimeSample = TNumericLimits<double>::Lowest();
+	for (const double UsdTimeSample : UsdTimeSamples)
+	{
+		// We never want to evaluate the same time twice
+		if (FMath::IsNearlyEqual(UsdTimeSample, LastTimeSample))
+		{
+			continue;
+		}
+		LastTimeSample = UsdTimeSample;
+
+		int32 FrameNumber = FMath::FloorToInt(UsdTimeSample);
+		float SubFrameNumber = UsdTimeSample - FrameNumber;
+
+		FFrameTime FrameTime(FrameNumber, SubFrameNumber);
+
+		FFrameTime KeyFrameTime = FFrameRate::TransformTime(FrameTime, StageFrameRate, Resolution);
+		KeyFrameTime *= SequenceTransform;
+		FrameNumbers.Add(KeyFrameTime.GetFrame());
+
+		bool UEValue = ReaderFunc(UsdTimeSample);
+		SectionValues.Emplace_GetRef(UEValue);
+	}
+
+	bool bSectionAdded = false;
+	UMovieSceneVisibilitySection* Section = Cast< UMovieSceneVisibilitySection >(MovieSceneTrack.FindOrAddSection(0, bSectionAdded));
+	Section->EvalOptions.CompletionMode = EMovieSceneCompletionMode::KeepState;
+
+	TMovieSceneChannelData<bool> Data = Section->GetChannel().GetData();
+	Data.Reset();
+	for (int32 KeyIndex = 0; KeyIndex < FrameNumbers.Num(); ++KeyIndex)
+	{
+		Data.AddKey(FrameNumbers[KeyIndex], SectionValues[KeyIndex]);
+	}
+
+	Section->SetRange(Section->GetAutoSizeRange().Get(TRange<FFrameNumber>::Empty()));
+
+	return true;
+}
+
 bool UsdToUnreal::ConvertFloatTimeSamples( const UE::FUsdStage& Stage, const TArray<double>& UsdTimeSamples, const TFunction<float( double )>& ReaderFunc, UMovieSceneFloatTrack& MovieSceneTrack, const FMovieSceneSequenceTransform& SequenceTransform )
 {
 	if ( !ReaderFunc )
diff --git a/Engine/Plugins/Importers/USDImporter/Source/USDUtilities/Public/USDPrimConversion.h b/Engine/Plugins/Importers/USDImporter/Source/USDUtilities/Public/USDPrimConversion.h
--- a/Engine/Plugins/Importers/USDImporter/Source/USDUtilities/Public/USDPrimConversion.h
+++ b/Engine/Plugins/Importers/USDImporter/Source/USDUtilities/Public/USDPrimConversion.h
@@ -39,6 +39,7 @@ class UMovieSceneFloatTrack;
 class UMovieScenePropertyTrack;
 class UMovieSceneSkeletalAnimationTrack;
 class UMovieSceneTrack;
+class UMovieSceneVisibilityTrack;
 class USceneComponent;
 class USkeletalMeshComponent;
 struct FFrameRate;
@@ -90,6 +91,7 @@ namespace UsdToUnreal
 	 * This is mostly used when reading attributes from USD into tracks for the automatically generated ULevelSequence provided with AUsdStageActors.
 	 */
 	USDUTILITIES_API bool ConvertBoolTimeSamples( const UE::FUsdStage& UsdStage, const TArray<double>& UsdTimeSamples, const TFunction<bool( double )>& ReaderFunc, UMovieSceneBoolTrack& MovieSceneTrack, const FMovieSceneSequenceTransform& SequenceTransform );
+	USDUTILITIES_API bool ConvertBoolTimeSamples( const UE::FUsdStage& UsdStage, const TArray<double>& UsdTimeSamples, const TFunction<bool(double)>& ReaderFunc, UMovieSceneVisibilityTrack& MovieSceneTrack, const FMovieSceneSequenceTransform& SequenceTransform );
 	USDUTILITIES_API bool ConvertFloatTimeSamples( const UE::FUsdStage& UsdStage, const TArray<double>& UsdTimeSamples, const TFunction<float( double )>& ReaderFunc, UMovieSceneFloatTrack& MovieSceneTrack, const FMovieSceneSequenceTransform& SequenceTransform );
 	USDUTILITIES_API bool ConvertColorTimeSamples( const UE::FUsdStage& UsdStage, const TArray<double>& UsdTimeSamples, const TFunction<FLinearColor( double )>& ReaderFunc, UMovieSceneColorTrack& MovieSceneTrack, const FMovieSceneSequenceTransform& SequenceTransform );
 	USDUTILITIES_API bool ConvertTransformTimeSamples( const UE::FUsdStage& UsdStage, const TArray<double>& UsdTimeSamples, const TFunction<FTransform( double )>& ReaderFunc, UMovieScene3DTransformTrack& MovieSceneTrack, const FMovieSceneSequenceTransform& SequenceTransform );
diff --git a/Engine/Source/Runtime/MovieScene/Public/Sections/MovieSceneBoolSection.h b/Engine/Source/Runtime/MovieScene/Public/Sections/MovieSceneBoolSection.h
--- a/Engine/Source/Runtime/MovieScene/Public/Sections/MovieSceneBoolSection.h
+++ b/Engine/Source/Runtime/MovieScene/Public/Sections/MovieSceneBoolSection.h
@@ -17,8 +17,8 @@ class UObject;
 /**
  * A single bool section.
  */
-UCLASS(MinimalAPI)
-class UMovieSceneBoolSection 
+UCLASS()
+class MOVIESCENE_API UMovieSceneBoolSection 
 	: public UMovieSceneSection
 {
 	GENERATED_UCLASS_BODY()
@@ -58,7 +58,7 @@ public:
 	/**
 	 * Set a flag indicating that the actual property that this bool represents is the opposite of the values stored in this section
 	 */
-	MOVIESCENE_API void SetIsExternallyInverted(bool bInIsExternallyInverted);
+	void SetIsExternallyInverted(bool bInIsExternallyInverted);
 
 protected:
 
diff --git a/Engine/Source/Runtime/MovieSceneTracks/Private/Evaluation/MovieSceneVisibilityTemplate.cpp b/Engine/Source/Runtime/MovieSceneTracks/Private/Evaluation/MovieSceneVisibilityTemplate.cpp
deleted file mode 100644
diff --git a/Engine/Source/Runtime/MovieSceneTracks/Private/Evaluation/MovieSceneVisibilityTemplate.h b/Engine/Source/Runtime/MovieSceneTracks/Private/Evaluation/MovieSceneVisibilityTemplate.h
deleted file mode 100644
diff --git a/Engine/Source/Runtime/MovieSceneTracks/Private/MovieSceneTracksComponentTypes.cpp b/Engine/Source/Runtime/MovieSceneTracks/Private/MovieSceneTracksComponentTypes.cpp
--- a/Engine/Source/Runtime/MovieSceneTracks/Private/MovieSceneTracksComponentTypes.cpp
+++ b/Engine/Source/Runtime/MovieSceneTracks/Private/MovieSceneTracksComponentTypes.cpp
@@ -569,6 +569,9 @@ FMovieSceneTracksComponentTypes::FMovieSceneTracksComponentTypes()
 	Tags.Slomo = ComponentRegistry->NewTag(TEXT("Slomo"));
 	ComponentRegistry->Factories.DefineChildComponent(Tags.Slomo, Tags.Slomo);
 
+	Tags.Visibility = ComponentRegistry->NewTag(TEXT("Visibility"));
+	ComponentRegistry->Factories.DefineChildComponent(Tags.Visibility, Tags.Visibility);
+
 	// Used to indicate the ParameterName component for certain parameter types (scalar, vector2d, vector, color)
 	// should be interpreted as an index for custom primitive data.
 	Tags.CustomPrimitiveData = ComponentRegistry->NewTag(TEXT("Custom Primitive Data"));
diff --git a/Engine/Source/Runtime/MovieSceneTracks/Private/Sections/MovieSceneVisibilitySection.cpp b/Engine/Source/Runtime/MovieSceneTracks/Private/Sections/MovieSceneVisibilitySection.cpp
new file mode 100644
--- /dev/null
+++ b/Engine/Source/Runtime/MovieSceneTracks/Private/Sections/MovieSceneVisibilitySection.cpp
@@ -0,0 +1,97 @@
+// Copyright Epic Games, Inc. All Rights Reserved.
+
+#include "Sections/MovieSceneVisibilitySection.h"
+
+#include "EntitySystem/BuiltInComponentTypes.h"
+#include "EntitySystem/MovieSceneEntityBuilder.h"
+#include "EntitySystem/MovieSceneEntityManager.h"
+#include "EntitySystem/MovieSceneEntitySystemLinker.h"
+#include "Evaluation/MovieSceneEvaluationField.h"
+#include "MovieScene.h"
+#include "MovieSceneTimeHelpers.h"
+#include "MovieSceneTracksComponentTypes.h"
+
+#include UE_INLINE_GENERATED_CPP_BY_NAME(MovieSceneVisibilitySection)
+
+namespace UE::MovieScene
+{
+
+int32 VisibilityToEntityID(bool bVisible)
+{
+	return bVisible ? 1 : 0;
+}
+
+bool EntityIDToVisibility(int32 EntityID)
+{
+	return EntityID > 0;
+}
+
+}  // namespace UE::MovieScene
+
+UMovieSceneVisibilitySection::UMovieSceneVisibilitySection(const FObjectInitializer& Init)
+	: Super(Init)
+{
+	BoolCurve.SetDefault(true);
+}
+
+void UMovieSceneVisibilitySection::ImportEntityImpl(UMovieSceneEntitySystemLinker* EntityLinker, const FEntityImportParams& Params, FImportedEntity* OutImportedEntity)
+{
+	using namespace UE::MovieScene;
+
+	FBuiltInComponentTypes* BuiltInComponents = FBuiltInComponentTypes::Get();
+	FMovieSceneTracksComponentTypes* TrackComponents = FMovieSceneTracksComponentTypes::Get();
+
+	OutImportedEntity->AddBuilder(
+		FEntityBuilder()
+		.AddConditional(BuiltInComponents->GenericObjectBinding, Params.GetObjectBindingID(), Params.GetObjectBindingID().IsValid())
+		.AddTag(TrackComponents->Tags.Visibility)
+		.Add(BuiltInComponents->BoolResult, EntityIDToVisibility(Params.EntityID))
+	);
+}
+
+bool UMovieSceneVisibilitySection::PopulateEvaluationFieldImpl(const TRange<FFrameNumber>& EffectiveRange, const FMovieSceneEvaluationFieldEntityMetaData& InMetaData, FMovieSceneEntityComponentFieldBuilder* OutFieldBuilder)
+{
+	using namespace UE::MovieScene;
+
+	const int32 MetaDataIndex = OutFieldBuilder->AddMetaData(InMetaData);
+
+	// Add ranges when visibility is ON or OFF to the evaluation tree, skipping keys that don't actually
+	// change the value. We will encode the ON/OFF value inside the entity ID.
+	TArrayView<const FFrameNumber> Times  = BoolCurve.GetTimes();
+	TArrayView<const bool>         Values = BoolCurve.GetValues();
+
+	const TRange<FFrameNumber> TrueSectionRange = GetTrueRange();
+	TRangeBound<FFrameNumber> StartBound = TrueSectionRange.GetLowerBound();
+
+	if (Times.Num() == 0)
+	{
+		const bool bDefaultVisible = BoolCurve.GetDefault().Get(true);
+		OutFieldBuilder->AddPersistentEntity(EffectiveRange, this, VisibilityToEntityID(bDefaultVisible), MetaDataIndex);
+		return true;
+	}
+
+	// The first keyframe determines both the value before it and after it.
+	bool bIsVisible = Values[0];
+
+	for (int32 Index = 0; Index < Times.Num(); ++Index)
+	{
+		FFrameNumber CurrentTime = Times[Index];
+		TRange<FFrameNumber> Range(StartBound, TRangeBound<FFrameNumber>::Exclusive(CurrentTime));
+		if (!Range.IsEmpty())
+		{
+			OutFieldBuilder->AddPersistentEntity(Range, this, VisibilityToEntityID(bIsVisible), MetaDataIndex);
+		}
+
+		StartBound = TRangeBound<FFrameNumber>::Inclusive(CurrentTime);
+		bIsVisible = Values[Index];
+	}
+
+	TRange<FFrameNumber> TailRange(StartBound, TrueSectionRange.GetUpperBound());
+	if (!TailRange.IsEmpty())
+	{
+		OutFieldBuilder->AddPersistentEntity(TailRange, this, VisibilityToEntityID(bIsVisible), MetaDataIndex);
+	}
+
+	return true;
+}
+
diff --git a/Engine/Source/Runtime/MovieSceneTracks/Private/Systems/MovieSceneVisibilitySystem.cpp b/Engine/Source/Runtime/MovieSceneTracks/Private/Systems/MovieSceneVisibilitySystem.cpp
new file mode 100644
--- /dev/null
+++ b/Engine/Source/Runtime/MovieSceneTracks/Private/Systems/MovieSceneVisibilitySystem.cpp
@@ -0,0 +1,182 @@
+// Copyright Epic Games, Inc. All Rights Reserved.
+
+#include "Systems/MovieSceneVisibilitySystem.h"
+
+#include "EntitySystem/BuiltInComponentTypes.h"
+#include "EntitySystem/MovieSceneEntitySystemLinker.h"
+#include "EntitySystem/MovieSceneEntitySystemTask.h"
+#include "EntitySystem/MovieScenePreAnimatedStateSystem.h"
+#include "Evaluation/PreAnimatedState/IMovieScenePreAnimatedStorage.h"
+#include "Evaluation/PreAnimatedState/MovieScenePreAnimatedObjectStorage.h"
+#include "Evaluation/PreAnimatedState/MovieScenePreAnimatedStorageID.h"
+#include "Evaluation/PreAnimatedState/MovieScenePreAnimatedStorageID.inl"
+#include "MovieSceneTracksComponentTypes.h"
+
+#include UE_INLINE_GENERATED_CPP_BY_NAME(MovieSceneVisibilitySystem)
+
+namespace UE::MovieScene
+{
+
+struct FPreAnimatedVisibilityState
+{
+	bool bHidden = false;
+	bool bTemporarilyHiddenInGame = false;
+
+	FPreAnimatedVisibilityState()
+	{}
+	FPreAnimatedVisibilityState(bool bInHidden, bool bInTemporarilyHiddenInGame)
+		: bHidden(bInHidden)
+		, bTemporarilyHiddenInGame(bInTemporarilyHiddenInGame)
+	{}
+};
+
+struct FPreAnimatedVisibilityTraits : FBoundObjectPreAnimatedStateTraits
+{
+	using KeyType     = FObjectKey;
+	using StorageType = FPreAnimatedVisibilityState;
+
+	StorageType CachePreAnimatedValue(const KeyType& Object)
+	{
+		UObject* ObjectPtr = Object.ResolveObjectPtr();
+		if (AActor* Actor = Cast<AActor>(ObjectPtr))
+		{
+			const bool bHidden = Actor->IsHidden();
+			const bool bTemporarilyHiddenInGame = 
+#if WITH_EDITOR
+				Actor->IsTemporarilyHiddenInEditor();
+#else
+				false;
+#endif  // WITH_EDITOR
+			return StorageType(bHidden, bTemporarilyHiddenInGame);
+		}
+		else if (USceneComponent* SceneComponent = Cast<USceneComponent>(ObjectPtr))
+		{
+			const bool bHidden = SceneComponent->bHiddenInGame;
+			// Second argument doesn't matter, we don't use it when the object is a scene component.
+			return StorageType(bHidden, false);
+		}
+		else
+		{
+			// Dummy value, we won't use it if the bound object isn't an actor or scene component.
+			return StorageType(false, false);
+		}
+	}
+
+	void RestorePreAnimatedValue(const KeyType& Object, StorageType& InOutCachedValue, const FRestoreStateParams& Params)
+	{
+		UObject* ObjectPtr = Object.ResolveObjectPtr();
+		if (AActor* Actor = Cast<AActor>(ObjectPtr))
+		{
+			Actor->SetActorHiddenInGame(InOutCachedValue.bHidden);
+#if WITH_EDITOR
+			Actor->SetIsTemporarilyHiddenInEditor(InOutCachedValue.bTemporarilyHiddenInGame);
+#endif  // WITH_EDITOR
+		}
+		else if (USceneComponent* SceneComponent = Cast<USceneComponent>(ObjectPtr))
+		{
+			SceneComponent->SetHiddenInGame(InOutCachedValue.bHidden);
+		}
+	}
+};
+
+struct FPreAnimatedVisibilityStorage : TPreAnimatedStateStorage_ObjectTraits<FPreAnimatedVisibilityTraits>
+{
+	static TAutoRegisterPreAnimatedStorageID<FPreAnimatedVisibilityStorage> StorageID;
+};
+
+TAutoRegisterPreAnimatedStorageID<FPreAnimatedVisibilityStorage> FPreAnimatedVisibilityStorage::StorageID;
+
+struct FVisibilityTask
+{
+	TSharedPtr<FPreAnimatedVisibilityStorage> PreAnimatedStorage;
+
+	void ForEachAllocation(
+			FEntityAllocationProxy Item,
+			TRead<FMovieSceneEntityID> EntityIDs,
+			TRead<FRootInstanceHandle> RootInstanceHandles,
+			TRead<UObject*> BoundObjects,
+			TRead<bool> BoolResults) const
+	{
+		const int32 Num = Item.GetAllocation()->Num();
+
+		// Cache pre-animiated state for all bound objects in this allocation.
+		FPreAnimatedTrackerParams TrackerParams(Item);
+		PreAnimatedStorage->BeginTrackingEntities(TrackerParams, EntityIDs, RootInstanceHandles, BoundObjects);
+
+		FCachePreAnimatedValueParams CacheParams;
+		PreAnimatedStorage->CachePreAnimatedValues(CacheParams, BoundObjects.AsArray(Num));
+
+		// Now make bound objects in this allocation visible/hidden according to the evaluated bool channel.
+		for (int32 Index = 0; Index < Num; ++Index)
+		{
+			if (UObject* const BoundObject = BoundObjects[Index])
+			{
+				const bool bShouldBeVisible = BoolResults[Index];
+
+				if (AActor* Actor = Cast<AActor>(BoundObject))
+				{
+					Actor->SetActorHiddenInGame(!bShouldBeVisible);
+#if WITH_EDITOR
+					if (GIsEditor && Actor->GetWorld() != nullptr && !Actor->GetWorld()->IsPlayInEditor())
+					{
+						Actor->SetIsTemporarilyHiddenInEditor(!bShouldBeVisible);
+					}
+#endif  // WITH_EDITOR
+				}
+				else if (USceneComponent* SceneComponent = Cast<USceneComponent>(BoundObject))
+				{
+					SceneComponent->SetHiddenInGame(!bShouldBeVisible);
+				}
+			}
+		}
+	}
+};
+
+} // namespace UE::MovieScene
+
+
+UMovieSceneVisibilitySystem::UMovieSceneVisibilitySystem(const FObjectInitializer& ObjInit)
+	: Super(ObjInit)
+{
+	using namespace UE::MovieScene;
+
+	SystemCategories = EEntitySystemCategory::Core;
+
+	RelevantComponent = FMovieSceneTracksComponentTypes::Get()->Tags.Visibility;
+
+	if (HasAnyFlags(RF_ClassDefaultObject))
+	{
+		DefineComponentConsumer(GetClass(), FBuiltInComponentTypes::Get()->BoundObject);
+	}
+}
+
+void UMovieSceneVisibilitySystem::OnLink()
+{
+	using namespace UE::MovieScene;
+
+	PreAnimatedStorage = Linker->PreAnimatedState.GetOrCreateStorage<FPreAnimatedVisibilityStorage>();
+}
+
+void UMovieSceneVisibilitySystem::OnUnlink()
+{
+	PreAnimatedStorage = nullptr;
+}
+
+void UMovieSceneVisibilitySystem::OnRun(FSystemTaskPrerequisites& InPrerequisites, FSystemSubsequentTasks& Subsequents)
+{
+	using namespace UE::MovieScene;
+
+	FBuiltInComponentTypes* BuiltInComponents = FBuiltInComponentTypes::Get();
+	FMovieSceneTracksComponentTypes* TrackComponents = FMovieSceneTracksComponentTypes::Get();
+
+	FVisibilityTask Task { PreAnimatedStorage };
+
+	FEntityTaskBuilder()
+	.ReadEntityIDs()
+	.Read(BuiltInComponents->RootInstanceHandle)
+	.Read(BuiltInComponents->BoundObject)
+	.Read(BuiltInComponents->BoolResult)
+	.FilterAll({ BuiltInComponents->Tags.NeedsLink, TrackComponents->Tags.Visibility })
+	.RunInline_PerAllocation(&Linker->EntityManager, Task);
+}
+
diff --git a/Engine/Source/Runtime/MovieSceneTracks/Private/Tracks/MovieSceneVisibilityTrack.cpp b/Engine/Source/Runtime/MovieSceneTracks/Private/Tracks/MovieSceneVisibilityTrack.cpp
--- a/Engine/Source/Runtime/MovieSceneTracks/Private/Tracks/MovieSceneVisibilityTrack.cpp
+++ b/Engine/Source/Runtime/MovieSceneTracks/Private/Tracks/MovieSceneVisibilityTrack.cpp
@@ -1,55 +1,60 @@
 // Copyright Epic Games, Inc. All Rights Reserved.
 
 #include "Tracks/MovieSceneVisibilityTrack.h"
+
 #include "Sections/MovieSceneBoolSection.h"
-#include "Evaluation/MovieSceneVisibilityTemplate.h"
+#include "Sections/MovieSceneVisibilitySection.h"
+#include "Serialization/ObjectReader.h"
+#include "Serialization/ObjectWriter.h"
 
 #include UE_INLINE_GENERATED_CPP_BY_NAME(MovieSceneVisibilityTrack)
 
 #define LOCTEXT_NAMESPACE "MovieSceneVisibilityTrack"
 
-
 UMovieSceneVisibilityTrack::UMovieSceneVisibilityTrack(const FObjectInitializer& ObjectInitializer)
 	: Super(ObjectInitializer)
-{}
-
-bool UMovieSceneVisibilityTrack::SupportsType(TSubclassOf<UMovieSceneSection> SectionClass) const
 {
-	return SectionClass == UMovieSceneBoolSection::StaticClass();
 }
 
-UMovieSceneSection* UMovieSceneVisibilityTrack::CreateNewSection()
+void UMovieSceneVisibilityTrack::PostLoad()
 {
-	UMovieSceneBoolSection* NewBoolSection = Cast<UMovieSceneBoolSection>(Super::CreateNewSection());
+	// Upgrade bool sections to visibility sections.
+	TArray<uint8> Bytes;
+	bool bUpgraded = false;
 
-#if WITH_EDITORONLY_DATA
-	if (NewBoolSection)
+	for (int32 Index = 0; Index < Sections.Num(); ++Index)
 	{
-		NewBoolSection->SetIsExternallyInverted(true);
+		UMovieSceneBoolSection* BoolSection = ExactCast<UMovieSceneBoolSection>(Sections[Index]);
+		if (BoolSection)
+		{
+			BoolSection->ConditionalPostLoad();
+			Bytes.Reset();
+
+			FObjectWriter(BoolSection, Bytes);
+			UMovieSceneVisibilitySection* NewSection = NewObject<UMovieSceneVisibilitySection>(this, NAME_None, RF_Transactional);
+			FObjectReader(NewSection, Bytes);
+
+			Sections[Index] = NewSection;
+			bUpgraded = true;
+		}
+	}
+
+	if (bUpgraded)
+	{
+		Modify();
 	}
-#endif
 
-	return NewBoolSection;
+	Super::PostLoad();
 }
 
-FMovieSceneEvalTemplatePtr UMovieSceneVisibilityTrack::CreateTemplateForSection(const UMovieSceneSection& InSection) const
+bool UMovieSceneVisibilityTrack::SupportsType(TSubclassOf<UMovieSceneSection> SectionClass) const
 {
-	return FMovieSceneVisibilitySectionTemplate(*CastChecked<const UMovieSceneBoolSection>(&InSection), *this);
+	return SectionClass == UMovieSceneVisibilitySection::StaticClass();
 }
 
-void UMovieSceneVisibilityTrack::PostLoad()
+UMovieSceneSection* UMovieSceneVisibilityTrack::CreateNewSection()
 {
-#if WITH_EDITORONLY_DATA
-	for (UMovieSceneSection* Section : GetAllSections())
-	{
-		if (Section)
-		{
-			CastChecked<UMovieSceneBoolSection>(Section)->SetIsExternallyInverted(true);
-		}
-	}
-#endif
-
-	Super::PostLoad();
+	return NewObject<UMovieSceneVisibilitySection>(this, NAME_None, RF_Transactional);
 }
 
 #if WITH_EDITORONLY_DATA
@@ -61,6 +66,5 @@ FText UMovieSceneVisibilityTrack::GetDisplayName() const
 
 #endif
 
-
 #undef LOCTEXT_NAMESPACE
 
diff --git a/Engine/Source/Runtime/MovieSceneTracks/Public/MovieSceneTracksComponentTypes.h b/Engine/Source/Runtime/MovieSceneTracks/Public/MovieSceneTracksComponentTypes.h
--- a/Engine/Source/Runtime/MovieSceneTracks/Public/MovieSceneTracksComponentTypes.h
+++ b/Engine/Source/Runtime/MovieSceneTracks/Public/MovieSceneTracksComponentTypes.h
@@ -637,6 +637,7 @@ struct FMovieSceneTracksComponentTypes
 		FComponentTypeID BoundMaterialChanged;
 		FComponentTypeID CustomPrimitiveData;
 		FComponentTypeID Slomo;
+		FComponentTypeID Visibility;
 	} Tags;
 
 	TComponentTypeID<FLevelVisibilityComponentData> LevelVisibility;
diff --git a/Engine/Source/Runtime/MovieSceneTracks/Public/Sections/MovieSceneVisibilitySection.h b/Engine/Source/Runtime/MovieSceneTracks/Public/Sections/MovieSceneVisibilitySection.h
new file mode 100644
--- /dev/null
+++ b/Engine/Source/Runtime/MovieSceneTracks/Public/Sections/MovieSceneVisibilitySection.h
@@ -0,0 +1,33 @@
+// Copyright Epic Games, Inc. All Rights Reserved.
+
+#pragma once
+
+#include "CoreTypes.h"
+#include "EntitySystem/IMovieSceneEntityProvider.h"
+#include "Sections/MovieSceneBoolSection.h"
+#include "UObject/ObjectMacros.h"
+#include "UObject/UObjectGlobals.h"
+
+#include "MovieSceneVisibilitySection.generated.h"
+
+/**
+ * A spawn section.
+ */
+UCLASS(MinimalAPI)
+class UMovieSceneVisibilitySection 
+	: public UMovieSceneBoolSection
+	, public IMovieSceneEntityProvider
+{
+	GENERATED_BODY()
+
+public:
+
+	UMovieSceneVisibilitySection(const FObjectInitializer& Init);
+
+public:
+
+	// IMovieSceneEntityProvider interface
+	virtual void ImportEntityImpl(UMovieSceneEntitySystemLinker* EntityLinker, const FEntityImportParams& Params, FImportedEntity* OutImportedEntity) override;
+	virtual bool PopulateEvaluationFieldImpl(const TRange<FFrameNumber>& EffectiveRange, const FMovieSceneEvaluationFieldEntityMetaData& InMetaData, FMovieSceneEntityComponentFieldBuilder* OutFieldBuilder) override;
+};
+
diff --git a/Engine/Source/Runtime/MovieSceneTracks/Public/Systems/MovieSceneVisibilitySystem.h b/Engine/Source/Runtime/MovieSceneTracks/Public/Systems/MovieSceneVisibilitySystem.h
new file mode 100644
--- /dev/null
+++ b/Engine/Source/Runtime/MovieSceneTracks/Public/Systems/MovieSceneVisibilitySystem.h
@@ -0,0 +1,39 @@
+// Copyright Epic Games, Inc. All Rights Reserved.
+
+#pragma once
+
+#include "EntitySystem/MovieSceneEntityInstantiatorSystem.h"
+#include "EntitySystem/MovieSceneEntitySystem.h"
+#include "UObject/ObjectMacros.h"
+#include "UObject/UObjectGlobals.h"
+
+#include "MovieSceneVisibilitySystem.generated.h"
+
+class UMovieSceneEntitySystemLinker;
+class UObject;
+
+namespace UE::MovieScene
+{
+	struct FPreAnimatedVisibilityStorage;
+}  // namespace UE::MovieScene
+
+UCLASS()
+class UMovieSceneVisibilitySystem : public UMovieSceneEntityInstantiatorSystem
+{
+public:
+
+	GENERATED_BODY()
+
+	UMovieSceneVisibilitySystem(const FObjectInitializer& ObjInit);
+
+private:
+
+	virtual void OnLink() override;
+	virtual void OnUnlink() override;
+	virtual void OnRun(FSystemTaskPrerequisites& InPrerequisites, FSystemSubsequentTasks& Subsequents) override;
+
+private:
+
+	TSharedPtr<UE::MovieScene::FPreAnimatedVisibilityStorage> PreAnimatedStorage;
+};
+
diff --git a/Engine/Source/Runtime/MovieSceneTracks/Public/Tracks/MovieSceneVisibilityTrack.h b/Engine/Source/Runtime/MovieSceneTracks/Public/Tracks/MovieSceneVisibilityTrack.h
--- a/Engine/Source/Runtime/MovieSceneTracks/Public/Tracks/MovieSceneVisibilityTrack.h
+++ b/Engine/Source/Runtime/MovieSceneTracks/Public/Tracks/MovieSceneVisibilityTrack.h
@@ -4,7 +4,8 @@
 
 #include "CoreMinimal.h"
 #include "UObject/ObjectMacros.h"
-#include "Tracks/MovieSceneBoolTrack.h"
+#include "Tracks/MovieScenePropertyTrack.h"
+
 #include "MovieSceneVisibilityTrack.generated.h"
 
 /**
@@ -12,14 +13,13 @@
  */
 UCLASS(MinimalAPI)
 class UMovieSceneVisibilityTrack
-	: public UMovieSceneBoolTrack
+	: public UMovieScenePropertyTrack
 {
 	GENERATED_UCLASS_BODY()
 
 public:
 
 	// UMovieSceneTrack interface
-	virtual FMovieSceneEvalTemplatePtr CreateTemplateForSection(const UMovieSceneSection& InSection) const override;
 	virtual void PostLoad() override;
 	virtual bool SupportsType(TSubclassOf<UMovieSceneSection> SectionClass) const override;
 	virtual UMovieSceneSection* CreateNewSection() override;
