Sequencer: Check object channel property track vs property tracks property class

diff --git a/Engine/Source/Runtime/MovieSceneTracks/Private/Sections/MovieSceneObjectPropertySection.cpp b/Engine/Source/Runtime/MovieSceneTracks/Private/Sections/MovieSceneObjectPropertySection.cpp
--- a/Engine/Source/Runtime/MovieSceneTracks/Private/Sections/MovieSceneObjectPropertySection.cpp
+++ b/Engine/Source/Runtime/MovieSceneTracks/Private/Sections/MovieSceneObjectPropertySection.cpp
@@ -1,6 +1,7 @@
 // Copyright Epic Games, Inc. All Rights Reserved.
 
 #include "Sections/MovieSceneObjectPropertySection.h"
+#include "Tracks/MovieSceneObjectPropertyTrack.h"
 #include "Channels/MovieSceneChannelProxy.h"
 
 #include UE_INLINE_GENERATED_CPP_BY_NAME(MovieSceneObjectPropertySection)
@@ -19,3 +20,28 @@ UMovieSceneObjectPropertySection::UMovieSceneObjectPropertySection(const FObject
 	ChannelProxy = MakeShared<FMovieSceneChannelProxy>(ObjectChannel);
 #endif
 }
+
+void UMovieSceneObjectPropertySection::PostLoad()
+{
+	if (UMovieSceneObjectPropertyTrack* PropertyTrack = Cast<UMovieSceneObjectPropertyTrack>(GetOuter()))
+	{
+		UClass* PropertyTrackPropertyClass = PropertyTrack->PropertyClass;
+		UClass* ObjectChannelPropertyClass = ObjectChannel.GetPropertyClass();
+		if (PropertyTrackPropertyClass != nullptr && ObjectChannelPropertyClass == nullptr)
+		{
+			UE_LOG(LogMovieScene, Warning, TEXT("Mismatch in property track's property class %s and object channel's property class (null). Setting object channel's property class to %s."), *PropertyTrackPropertyClass->GetName(), *PropertyTrackPropertyClass->GetName());
+			ObjectChannel.SetPropertyClass(PropertyTrackPropertyClass);
+		}
+		else if (ObjectChannelPropertyClass != nullptr && PropertyTrackPropertyClass == nullptr)
+		{
+			UE_LOG(LogMovieScene, Warning, TEXT("Mismatch in object channel's property class %s and property track's property class (null). Setting property track's property class to %s."), *ObjectChannelPropertyClass->GetName(), *ObjectChannelPropertyClass->GetName());
+			PropertyTrack->PropertyClass = ObjectChannelPropertyClass;
+		}
+		else if (PropertyTrackPropertyClass != nullptr && ObjectChannelPropertyClass != nullptr)
+		{
+			ensureMsgf(PropertyTrackPropertyClass == ObjectChannelPropertyClass, TEXT("Mismatch in property track's property class %s and object channel's property class %s"), *PropertyTrackPropertyClass->GetName(), *ObjectChannelPropertyClass->GetName());
+		}
+	}
+
+	Super::PostLoad();
+}
\ No newline at end of file
diff --git a/Engine/Source/Runtime/MovieSceneTracks/Public/Sections/MovieSceneObjectPropertySection.h b/Engine/Source/Runtime/MovieSceneTracks/Public/Sections/MovieSceneObjectPropertySection.h
--- a/Engine/Source/Runtime/MovieSceneTracks/Public/Sections/MovieSceneObjectPropertySection.h
+++ b/Engine/Source/Runtime/MovieSceneTracks/Public/Sections/MovieSceneObjectPropertySection.h
@@ -20,6 +20,11 @@ public:
 
 	UMovieSceneObjectPropertySection(const FObjectInitializer& ObjInit);
 
+	// ~UObject interface
+	virtual void PostLoad() override;
+
+public:
+
 	UPROPERTY()
 	FMovieSceneObjectPathChannel ObjectChannel;
 };
