CIS fix for MetaSound InitResources called on cook from MovieSceneAudioSection

diff --git a/Engine/Source/Runtime/MovieSceneTracks/Private/Sections/MovieSceneAudioSection.cpp b/Engine/Source/Runtime/MovieSceneTracks/Private/Sections/MovieSceneAudioSection.cpp
--- a/Engine/Source/Runtime/MovieSceneTracks/Private/Sections/MovieSceneAudioSection.cpp
+++ b/Engine/Source/Runtime/MovieSceneTracks/Private/Sections/MovieSceneAudioSection.cpp
@@ -171,7 +171,14 @@ void UMovieSceneAudioSection::SetupSoundInputParameters(USoundBase* InSoundBase)
 	// Populate with defaults.
 	if (InSoundBase)
 	{
-		InSoundBase->InitResources();
+		// Don't init resources when running cook, as this can trigger 
+		// registration of a MetaSound and its dependent graphs.
+		// Those will instead be registered when the MetaSound itself is cooked (FMetasoundAssetBase::CookMetaSound)
+		// in a way that does not deal with runtime data like this function does
+		if (!IsRunningCookCommandlet())
+		{
+			InSoundBase->InitResources();
+		}
 
 		TArray<FAudioParameter> DefaultParams;
 		InSoundBase->GetAllDefaultParameters(DefaultParams);
