Sequencer: add option for choosing whether the camera cut track restores the original (pre-animated) viewport position on unlock

This change adds the new option on the sequencer settings. We need to know the value of that option in the evaluation code, so that we don't cache/restore pre-animated viewport positions when playing or scrubbing, so this is exposed via the camera cut playback capability.

In addition to this, extra work must be done when locking/unlocking cinematic viewports with the little camera button on the camera cut track.

When locking:
- If the option is ON (restore viewports), we force-update the pre-animated state so that we will restore the new viewport position we had just before clicking the button.
- If the option is OFF, we don't do anything. We should have any pre-animated state anyway since saving it would have early-returned.

When unlocking:
- If the option is ON (restore viewports), we force-restore the pre-animated state, to, well, restore the viewports as requested.
- If the option is OFF, we force-discard any pre-animated state we might have had.

Last, the world is now cached on the game pre-animated state so that we don't need to rely on a terminal player, which could not be provided in some cases.

diff --git a/Engine/Source/Editor/MovieSceneTools/Private/TrackEditors/CameraCutTrackEditor.cpp b/Engine/Source/Editor/MovieSceneTools/Private/TrackEditors/CameraCutTrackEditor.cpp
--- a/Engine/Source/Editor/MovieSceneTools/Private/TrackEditors/CameraCutTrackEditor.cpp
+++ b/Engine/Source/Editor/MovieSceneTools/Private/TrackEditors/CameraCutTrackEditor.cpp
@@ -1,29 +1,33 @@
 // Copyright Epic Games, Inc. All Rights Reserved.
 
 #include "TrackEditors/CameraCutTrackEditor.h"
-#include "Widgets/SBoxPanel.h"
-#include "Framework/MultiBox/MultiBoxBuilder.h"
-#include "Tracks/MovieSceneCameraCutTrack.h"
-#include "Modules/ModuleManager.h"
-#include "Application/ThrottleManager.h"
-#include "Widgets/Layout/SBox.h"
-#include "Widgets/Input/SCheckBox.h"
+
+#include "DragAndDrop/ActorDragDropGraphEdOp.h"
 #include "MovieSceneCommonHelpers.h"
-#include "Styling/AppStyle.h"
-#include "GameFramework/WorldSettings.h"
-#include "LevelEditorViewport.h"
+#include "MovieSceneObjectBindingIDPicker.h"
+#include "MovieSceneToolHelpers.h"
 #include "Sections/CameraCutSection.h"
 #include "Sections/MovieSceneCameraCutSection.h"
+#include "SequencerSettings.h"
 #include "SequencerUtilities.h"
-#include "Editor.h"
+#include "TrackEditorThumbnail/TrackEditorThumbnailPool.h"
+#include "TrackInstances/MovieSceneCameraCutTrackInstance.h"
+#include "Tracks/MovieSceneCameraCutTrack.h"
+
 #include "ActorEditorUtils.h"
-#include "SceneOutlinerPublicTypes.h"
-#include "SceneOutlinerModule.h"
 #include "ActorTreeItem.h"
-#include "TrackEditorThumbnail/TrackEditorThumbnailPool.h"
-#include "MovieSceneObjectBindingIDPicker.h"
-#include "MovieSceneToolHelpers.h"
-#include "DragAndDrop/ActorDragDropGraphEdOp.h"
+#include "Application/ThrottleManager.h"
+#include "Editor.h"
+#include "Framework/MultiBox/MultiBoxBuilder.h"
+#include "GameFramework/WorldSettings.h"
+#include "LevelEditorViewport.h"
+#include "Modules/ModuleManager.h"
+#include "SceneOutlinerModule.h"
+#include "SceneOutlinerPublicTypes.h"
+#include "Styling/AppStyle.h"
+#include "Widgets/Input/SCheckBox.h"
+#include "Widgets/Layout/SBox.h"
+#include "Widgets/SBoxPanel.h"
 
 #define LOCTEXT_NAMESPACE "FCameraCutTrackEditor"
 
@@ -562,10 +566,21 @@ ECheckBoxState FCameraCutTrackEditor::IsCameraLocked() const
 
 void FCameraCutTrackEditor::OnLockCameraClicked(ECheckBoxState CheckBoxState)
 {
+	TSharedPtr<ISequencer> SequencerPtr = GetSequencer();
 
 	const bool bEnableCameraCuts = (CheckBoxState == ECheckBoxState::Checked);
-	GetSequencer()->SetPerspectiveViewportCameraCutEnabled(bEnableCameraCuts);
-	GetSequencer()->ForceEvaluate();
+	SequencerPtr->SetPerspectiveViewportCameraCutEnabled(bEnableCameraCuts);
+
+	bool bNeedsRestoreViewport = true;
+	if (const USequencerSettings* SequencerSettings = SequencerPtr->GetSequencerSettings())
+	{
+		bNeedsRestoreViewport = SequencerSettings->GetRestoreOriginalViewportOnCameraCutUnlock();
+	}
+
+	UMovieSceneEntitySystemLinker* Linker = SequencerPtr->GetEvaluationTemplate().GetEntitySystemLinker();
+	UMovieSceneCameraCutTrackInstance::ToggleCameraCutLock(Linker, bEnableCameraCuts, bNeedsRestoreViewport);
+
+	SequencerPtr->ForceEvaluate();
 }
 
 void FCameraCutTrackEditor::ToggleLockCamera()
diff --git a/Engine/Source/Editor/MovieSceneTools/Private/TrackEditors/CinematicShotTrackEditor.cpp b/Engine/Source/Editor/MovieSceneTools/Private/TrackEditors/CinematicShotTrackEditor.cpp
--- a/Engine/Source/Editor/MovieSceneTools/Private/TrackEditors/CinematicShotTrackEditor.cpp
+++ b/Engine/Source/Editor/MovieSceneTools/Private/TrackEditors/CinematicShotTrackEditor.cpp
@@ -1,36 +1,40 @@
 // Copyright Epic Games, Inc. All Rights Reserved.
 
 #include "TrackEditors/CinematicShotTrackEditor.h"
-#include "Misc/Paths.h"
-#include "Widgets/SBoxPanel.h"
-#include "Framework/MultiBox/MultiBoxBuilder.h"
-#include "GameFramework/Actor.h"
-#include "Factories/Factory.h"
-#include "Tracks/MovieSceneSubTrack.h"
-#include "Tracks/MovieSceneCinematicShotTrack.h"
-#include "Sections/MovieSceneCinematicShotSection.h"
-#include "Modules/ModuleManager.h"
-#include "Application/ThrottleManager.h"
-#include "Framework/Application/SlateApplication.h"
-#include "Widgets/Layout/SBox.h"
-#include "Widgets/Input/SCheckBox.h"
-#include "Styling/AppStyle.h"
-#include "LevelEditorViewport.h"
+
+#include "AutomatedLevelSequenceCapture.h"
+#include "LevelSequence.h"
+#include "MovieSceneCaptureModule.h"
+#include "MovieSceneTimeHelpers.h"
 #include "MovieSceneToolHelpers.h"
-#include "FCPXML/FCPXMLMovieSceneTranslator.h"
+#include "MovieSceneToolsProjectSettings.h"
 #include "Sections/CinematicShotSection.h"
+#include "Sections/MovieSceneCinematicShotSection.h"
+#include "SequencerSettings.h"
 #include "SequencerUtilities.h"
-#include "UObject/UObjectHash.h"
-#include "UObject/UObjectIterator.h"
-#include "LevelSequence.h"
-#include "AutomatedLevelSequenceCapture.h"
-#include "MovieSceneCaptureModule.h"
 #include "TrackEditorThumbnail/TrackEditorThumbnailPool.h"
-#include "MovieSceneToolsProjectSettings.h"
+#include "TrackInstances/MovieSceneCameraCutTrackInstance.h"
+#include "Tracks/MovieSceneCinematicShotTrack.h"
+#include "Tracks/MovieSceneSubTrack.h"
+
+#include "Application/ThrottleManager.h"
 #include "Editor.h"
-#include "MovieSceneTimeHelpers.h"
+#include "FCPXML/FCPXMLMovieSceneTranslator.h"
+#include "Factories/Factory.h"
+#include "Framework/Application/SlateApplication.h"
+#include "Framework/MultiBox/MultiBoxBuilder.h"
 #include "Framework/Notifications/NotificationManager.h"
+#include "GameFramework/Actor.h"
+#include "LevelEditorViewport.h"
+#include "Misc/Paths.h"
+#include "Modules/ModuleManager.h"
+#include "Styling/AppStyle.h"
+#include "UObject/UObjectHash.h"
+#include "UObject/UObjectIterator.h"
+#include "Widgets/Input/SCheckBox.h"
+#include "Widgets/Layout/SBox.h"
 #include "Widgets/Notifications/SNotificationList.h"
+#include "Widgets/SBoxPanel.h"
 
 #define LOCTEXT_NAMESPACE "FCinematicShotTrackEditor"
 
@@ -282,26 +286,21 @@
 
 void FCinematicShotTrackEditor::OnLockShotsClicked(ECheckBoxState CheckBoxState)
 {
-	if (CheckBoxState == ECheckBoxState::Checked)
+	TSharedPtr<ISequencer> SequencerPtr = GetSequencer();
+
+	const bool bEnableCameraCuts = (CheckBoxState == ECheckBoxState::Checked);
+	SequencerPtr->SetPerspectiveViewportCameraCutEnabled(bEnableCameraCuts);
+
+	bool bNeedsRestoreViewport = true;
+	if (const USequencerSettings* SequencerSettings = SequencerPtr->GetSequencerSettings())
 	{
-		for( FLevelEditorViewportClient* LevelVC : GEditor->GetLevelViewportClients() )
-		{
-			if (LevelVC && LevelVC->AllowsCinematicControl() && LevelVC->GetViewMode() != VMI_Unknown)
-			{
-				LevelVC->SetActorLock(nullptr);
-				LevelVC->bLockedCameraView = false;
-				LevelVC->UpdateViewForLockedActor();
-				LevelVC->Invalidate();
-			}
-		}
-		GetSequencer()->SetPerspectiveViewportCameraCutEnabled(true);
-	}
-	else
-	{
-		GetSequencer()->SetPerspectiveViewportCameraCutEnabled(false);
+		bNeedsRestoreViewport = SequencerSettings->GetRestoreOriginalViewportOnCameraCutUnlock();
 	}
 
-	GetSequencer()->ForceEvaluate();
+	UMovieSceneEntitySystemLinker* Linker = SequencerPtr->GetEvaluationTemplate().GetEntitySystemLinker();
+	UMovieSceneCameraCutTrackInstance::ToggleCameraCutLock(Linker, bEnableCameraCuts, bNeedsRestoreViewport);
+
+	SequencerPtr->ForceEvaluate();
 }
 
 FText FCinematicShotTrackEditor::GetLockShotsToolTip() const
diff --git a/Engine/Source/Editor/Sequencer/Private/SSequencer.cpp b/Engine/Source/Editor/Sequencer/Private/SSequencer.cpp
--- a/Engine/Source/Editor/Sequencer/Private/SSequencer.cpp
+++ b/Engine/Source/Editor/Sequencer/Private/SSequencer.cpp
@@ -2158,6 +2158,7 @@ TSharedRef<SWidget> SSequencer::MakeViewMenu()
 	if (Sequencer->IsLevelEditorSequencer())
 	{
 		MenuBuilder.AddMenuEntry( FSequencerCommands::Get().TogglePilotCamera );
+		MenuBuilder.AddMenuEntry( FSequencerCommands::Get().ToggleRestoreOriginalViewportOnCameraCutUnlock );
 	}
 
 	MenuBuilder.AddMenuEntry( FSequencerCommands::Get().ToggleAutoScroll );
diff --git a/Engine/Source/Editor/Sequencer/Private/Sequencer.cpp b/Engine/Source/Editor/Sequencer/Private/Sequencer.cpp
--- a/Engine/Source/Editor/Sequencer/Private/Sequencer.cpp
+++ b/Engine/Source/Editor/Sequencer/Private/Sequencer.cpp
@@ -10231,6 +10231,12 @@ void FSequencer::BindCommands()
 		FCanExecuteAction::CreateLambda([] { return true; }),
 		FIsActionChecked::CreateLambda([this] { return Settings->GetAutoExpandNodesOnSelection(); }) );
 
+	SequencerCommandBindings->MapAction(
+		Commands.ToggleRestoreOriginalViewportOnCameraCutUnlock,
+		FExecuteAction::CreateLambda([this] { Settings->SetRestoreOriginalViewportOnCameraCutUnlock(!Settings->GetRestoreOriginalViewportOnCameraCutUnlock()); }),
+		FCanExecuteAction::CreateLambda([] { return true; }),
+		FIsActionChecked::CreateLambda([this] { return Settings->GetRestoreOriginalViewportOnCameraCutUnlock(); }) );
+
 	SequencerCommandBindings->MapAction(
 		Commands.ToggleExpandCollapseNodes,
 		FExecuteAction::CreateSP(this, &FSequencer::ToggleExpandCollapseNodes));
@@ -11533,6 +11539,11 @@ UE::Sequencer::FSequencerSelection& FSequencer::GetSelection()
 	return *ViewModel->GetSelection();
 }
 
+bool FSequencer::ShouldCacheEditorPreAnimatedState()
+{
+	return Settings->GetRestoreOriginalViewportOnCameraCutUnlock();
+}
+
 float FSequencer::GetCameraBlendPlayRate()
 {
 	return PlaybackSpeed;
diff --git a/Engine/Source/Editor/Sequencer/Private/Sequencer.h b/Engine/Source/Editor/Sequencer/Private/Sequencer.h
--- a/Engine/Source/Editor/Sequencer/Private/Sequencer.h
+++ b/Engine/Source/Editor/Sequencer/Private/Sequencer.h
@@ -844,11 +844,9 @@ protected:
 	// FCameraCutPlaybackCapability interface
 
 	virtual bool ShouldUpdateCameraCut() override { return IsPerspectiveViewportCameraCutEnabled(); }
+	virtual bool ShouldCacheEditorPreAnimatedState() override;
 	virtual float GetCameraBlendPlayRate() override;
 	virtual void OnCameraCutUpdated(const UE::MovieScene::FOnCameraCutUpdatedParams& Params) override;
-#if WITH_EDITOR
-	virtual bool ShouldPreviewCameraCutBlends() override { return !IsInSilentMode(); }
-#endif
 
 protected:
 
diff --git a/Engine/Source/Editor/Sequencer/Private/SequencerCommands.cpp b/Engine/Source/Editor/Sequencer/Private/SequencerCommands.cpp
--- a/Engine/Source/Editor/Sequencer/Private/SequencerCommands.cpp
+++ b/Engine/Source/Editor/Sequencer/Private/SequencerCommands.cpp
@@ -59,6 +59,7 @@ void FSequencerCommands::RegisterCommands()
 	UI_COMMAND( ToggleKeepPlaybackRangeInSectionBounds, "Keep Playback Range in Section Bounds", "When checked, the playback range will be synchronized to the section bounds", EUserInterfaceActionType::ToggleButton, FInputChord() );
 
 	UI_COMMAND( ToggleAutoExpandNodesOnSelection, "Auto Expand Nodes on Selection", "Toggle auto expanding the outliner tree on child selection", EUserInterfaceActionType::ToggleButton, FInputChord());
+	UI_COMMAND( ToggleRestoreOriginalViewportOnCameraCutUnlock, "Restore Cinematic Viewports", "Toggle whether unlocking a camera cut track should return the viewport to its original location, or keep it where the camera cut was", EUserInterfaceActionType::ToggleButton, FInputChord() );
 	UI_COMMAND( ToggleExpandCollapseNodes, "Expand/Collapse Nodes", "Toggle expand or collapse selected nodes", EUserInterfaceActionType::Button, FInputChord(EKeys::V) );
 	UI_COMMAND( ToggleExpandCollapseNodesAndDescendants, "Expand/Collapse Nodes and Descendants", "Toggle expand or collapse selected nodes and descendants", EUserInterfaceActionType::Button, FInputChord(EModifierKey::Shift, EKeys::V) );
 	UI_COMMAND( ExpandAllNodes, "Expand All Nodes", "Expand all nodes and descendants", EUserInterfaceActionType::Button, FInputChord() );
diff --git a/Engine/Source/Editor/Sequencer/Private/SequencerSettings.cpp b/Engine/Source/Editor/Sequencer/Private/SequencerSettings.cpp
--- a/Engine/Source/Editor/Sequencer/Private/SequencerSettings.cpp
+++ b/Engine/Source/Editor/Sequencer/Private/SequencerSettings.cpp
@@ -62,6 +62,7 @@ USequencerSettings::USequencerSettings( const FObjectInitializer& ObjectInitiali
 	TrajectoryPathCap = 250;
 	FrameNumberDisplayFormat = EFrameNumberDisplayFormats::Seconds;
 	bAutoExpandNodesOnSelection = true;
+	bRestoreOriginalViewportOnCameraCutUnlock = true;
 	TreeViewWidth = 0.3f;
 	bShowTickLines = true;
 	bShowSequencerToolbar = true;
@@ -960,6 +961,15 @@ void USequencerSettings::SetAutoExpandNodesOnSelection(bool bInAutoExpandNodesOn
 	}
 }
 
+void USequencerSettings::SetRestoreOriginalViewportOnCameraCutUnlock(bool bInRestoreOriginalViewportOnCameraCutUnlock)
+{
+	if (bInRestoreOriginalViewportOnCameraCutUnlock != bRestoreOriginalViewportOnCameraCutUnlock)
+	{
+		bRestoreOriginalViewportOnCameraCutUnlock = bInRestoreOriginalViewportOnCameraCutUnlock;
+		SaveConfig();
+	}
+}
+
 void USequencerSettings::SetTreeViewWidth(float InTreeViewWidth)
 {
 	if (InTreeViewWidth != TreeViewWidth)
diff --git a/Engine/Source/Editor/Sequencer/Public/SequencerCommands.h b/Engine/Source/Editor/Sequencer/Public/SequencerCommands.h
--- a/Engine/Source/Editor/Sequencer/Public/SequencerCommands.h
+++ b/Engine/Source/Editor/Sequencer/Public/SequencerCommands.h
@@ -157,6 +157,12 @@ public:
 	/** Toggle auto expand outliner tree on child selection */
 	TSharedPtr< FUICommandInfo > ToggleAutoExpandNodesOnSelection;
 
+	/**
+	 * Toggle whether unlocking a camera cut track should return the viewport to its original location, or keep it where
+	 * the camera cut was.
+	 */
+	TSharedPtr< FUICommandInfo > ToggleRestoreOriginalViewportOnCameraCutUnlock;
+
 	/** Expand/collapse nodes */
 	TSharedPtr< FUICommandInfo > ToggleExpandCollapseNodes;
 
diff --git a/Engine/Source/Editor/Sequencer/Public/SequencerSettings.h b/Engine/Source/Editor/Sequencer/Public/SequencerSettings.h
--- a/Engine/Source/Editor/Sequencer/Public/SequencerSettings.h
+++ b/Engine/Source/Editor/Sequencer/Public/SequencerSettings.h
@@ -437,6 +437,18 @@ public:
 	/** Sets whether or not to expand the outliner tree view when a child element is selected (from outside of the tree view). */
 	void SetAutoExpandNodesOnSelection(bool bInAutoExpandNodesOnSelection);
 
+
+	/**
+	 * Gets whether unlocking a camera cut track should return the viewport to its original location, or keep it where
+	 * the camera cut was.
+	 */
+	bool GetRestoreOriginalViewportOnCameraCutUnlock() const { return bRestoreOriginalViewportOnCameraCutUnlock; }
+	/**
+	 * Sets whether unlocking a camera cut track should return the viewport to its original location, or keep it where
+	 * the camera cut was.
+	 */
+	void SetRestoreOriginalViewportOnCameraCutUnlock(bool bInRestoreOriginalViewportOnCameraCutUnlock);
+
 	/** Gets the tree view width percentage */
 	float GetTreeViewWidth() const { return TreeViewWidth; }
 	/** Sets the tree view width percentage */
@@ -696,6 +708,14 @@ protected:
 	UPROPERTY(config, EditAnywhere, Category = General)
 	bool bAutoExpandNodesOnSelection;
 
+	/**
+	 * Whether unlocking a camera cut track should return the viewport to its original location, or keep it where the
+	 * camera cut was.
+	 * WARNING: Disabling this will make previewing camera cut blends useless, since it will blend to the same position.
+	 */
+	UPROPERTY(config, EditAnywhere, Category=General)
+	bool bRestoreOriginalViewportOnCameraCutUnlock;
+
 	/** The tree view width percentage */
 	UPROPERTY(config, EditAnywhere, Category = General)
 	float TreeViewWidth;
diff --git a/Engine/Source/Runtime/MovieSceneTracks/Private/TrackInstances/MovieSceneCameraCutEditorHandler.cpp b/Engine/Source/Runtime/MovieSceneTracks/Private/TrackInstances/MovieSceneCameraCutEditorHandler.cpp
--- a/Engine/Source/Runtime/MovieSceneTracks/Private/TrackInstances/MovieSceneCameraCutEditorHandler.cpp
+++ b/Engine/Source/Runtime/MovieSceneTracks/Private/TrackInstances/MovieSceneCameraCutEditorHandler.cpp
@@ -1,6 +1,7 @@
 // Copyright Epic Games, Inc. All Rights Reserved.
 
 #include "TrackInstances/MovieSceneCameraCutEditorHandler.h"
+#include "Engine/EngineTypes.h"
 
 #if WITH_EDITOR
 
@@ -14,6 +15,13 @@
 namespace UE::MovieScene
 {
 
+bool FPreAnimatedCameraCutEditorTraits::ShouldHandleViewportCameraCuts(FLevelEditorViewportClient* ViewportClient)
+{
+	UWorld* EditorWorld = ViewportClient->GetWorld();
+	return EditorWorld && 
+		(EditorWorld->WorldType == EWorldType::Editor || EditorWorld->WorldType == EWorldType::EditorPreview);
+}
+
 FPreAnimatedCameraCutEditorState FPreAnimatedCameraCutEditorTraits::CachePreAnimatedValue(
 		FLevelEditorViewportClient* InKey)
 {
@@ -40,8 +48,17 @@ void FPreAnimatedCameraCutEditorTraits::RestorePreAnimatedValue(
 		return;
 	}
 
+	// Check that we have an editor viewport.
+	if (!ShouldHandleViewportCameraCuts(InKey))
+	{
+		return;
+	}
+
 	// If the viewport wasn't locked to cinematics anyway, don't mess it up.
-	if (!InKey->IsLockedToCinematic())
+	// However, we don't call `IsLockedToCinematics` because we want to also consider the case
+	// of a camera actor that is PendingKill after having been unspawned by sequencer in the
+	// current update.
+	if (InKey->GetCinematicActorLock().LockedActor.IsExplicitlyNull())
 	{
 		return;
 	}
@@ -67,10 +84,21 @@ void FCameraCutEditorHandler::CachePreAnimatedValue(
 		return;
 	}
 
+	FCameraCutPlaybackCapabilityCompatibilityWrapper Wrapper(SequenceInstance);
+	if (!Wrapper.ShouldCacheEditorPreAnimatedState())
+	{
+		return;
+	}
+	
 	TSharedPtr<FPreAnimatedCameraCutEditorStorage> PreAnimatedStorage = Linker->PreAnimatedState.GetOrCreateStorage<FPreAnimatedCameraCutEditorStorage>();
 
 	for (FLevelEditorViewportClient* LevelVC : GEditor->GetLevelViewportClients())
 	{
+		if (!FPreAnimatedCameraCutEditorTraits::ShouldHandleViewportCameraCuts(LevelVC))
+		{
+			continue;
+		}
+
 		PreAnimatedStorage->CachePreAnimatedValue(
 				LevelVC,
 				[](FLevelEditorViewportClient* InKey) { return FPreAnimatedCameraCutEditorTraits::CachePreAnimatedValue(InKey); },
@@ -78,6 +106,56 @@ void FCameraCutEditorHandler::CachePreAnimatedValue(
 	}
 }
 
+void FCameraCutEditorHandler::ForcePreAnimatedValueOperation(
+		UMovieSceneEntitySystemLinker* Linker,
+		EForcedCameraCutPreAnimatedStorageOperation Operation)
+{
+	if (!GEditor)
+	{
+		return;
+	}
+
+	TSharedPtr<FPreAnimatedCameraCutEditorStorage> PreAnimatedStorage = Linker->PreAnimatedState.GetOrCreateStorage<FPreAnimatedCameraCutEditorStorage>();
+
+	for (FLevelEditorViewportClient* LevelVC : GEditor->GetLevelViewportClients())
+	{
+		if (!FPreAnimatedCameraCutEditorTraits::ShouldHandleViewportCameraCuts(LevelVC))
+		{
+			continue;
+		}
+
+		FPreAnimatedStorageIndex StorageIndex = PreAnimatedStorage->FindStorageIndex(LevelVC);
+		if (StorageIndex.IsValid())
+		{
+			switch (Operation)
+			{
+				case EForcedCameraCutPreAnimatedStorageOperation::Cache:
+					{
+						PreAnimatedStorage->DiscardPreAnimatedStateStorage(StorageIndex, EPreAnimatedStorageRequirement::Transient);
+						PreAnimatedStorage->CachePreAnimatedValue(
+								LevelVC,
+								[](FLevelEditorViewportClient* InKey) { return FPreAnimatedCameraCutEditorTraits::CachePreAnimatedValue(InKey); },
+								EPreAnimatedCaptureSourceTracking::AlwaysCache);
+					}
+					break;
+				case EForcedCameraCutPreAnimatedStorageOperation::Restore:
+					{
+						FRestoreStateParams Params;
+						Params.Linker = Linker;
+						// No need to set the terminal instance handle, we don't use it (see above).
+						PreAnimatedStorage->RestorePreAnimatedStateStorage(StorageIndex, EPreAnimatedStorageRequirement::Transient, EPreAnimatedStorageRequirement::Persistent, Params);
+					}
+					break;
+				case EForcedCameraCutPreAnimatedStorageOperation::Discard:
+					{
+						PreAnimatedStorage->DiscardPreAnimatedStateStorage(StorageIndex, EPreAnimatedStorageRequirement::Transient);
+					}
+					break;
+			}
+		}
+	}
+}
+
 FCameraCutEditorHandler::FCameraCutEditorHandler(
 		UMovieSceneEntitySystemLinker* InLinker,
 		const FSequenceInstance& InSequenceInstance,
@@ -115,7 +193,7 @@ void FCameraCutEditorHandler::SetCameraCut(
 
 	for (FLevelEditorViewportClient* LevelVC : GEditor->GetLevelViewportClients())
 	{
-		if (LevelVC == nullptr)
+		if (LevelVC == nullptr || !FPreAnimatedCameraCutEditorTraits::ShouldHandleViewportCameraCuts(LevelVC))
 		{
 			continue;
 		}
@@ -161,7 +239,10 @@ void FCameraCutEditorHandler::SetCameraCutForViewport(
 	// If that viewport wasn't locked to cinematics, see if we need to re-cache pre-animated state.
 	// This is necessary if the user released control with the camera button on the camera cut track.
 	// (see ReleaseCameraCutForViewport)
-	if (!ViewportClient.IsLockedToCinematic())
+	//
+	// Note that PreAnimatedStorage can be null here if we never cached any pre-animated state, which
+	// is possible if the option to restore viewports on unlock is off.
+	if (PreAnimatedStorage.IsValid() && !ViewportClient.IsLockedToCinematic())
 	{
 		FPreAnimatedStorageIndex StorageIndex = PreAnimatedStorage->FindStorageIndex(&ViewportClient);
 		if (!StorageIndex.IsValid())
@@ -308,8 +389,15 @@ void FCameraCutEditorHandler::ReleaseCameraCutForViewport(
 	}
 	
 	// Restore the viewport to its pre-animated state.
+	//
+	// Note that PreAnimatedStorage can be null here if we never cached any pre-animated state, which
+	// is possible if the option to restore viewports on unlock is off.
+	FPreAnimatedStorageIndex StorageIndex;
 	TSharedPtr<FPreAnimatedCameraCutEditorStorage> PreAnimatedStorage = Linker->PreAnimatedState.FindStorage(FPreAnimatedCameraCutEditorStorage::StorageID);
-	FPreAnimatedStorageIndex StorageIndex = PreAnimatedStorage->FindStorageIndex(&ViewportClient);
+	if (PreAnimatedStorage.IsValid())
+	{
+		StorageIndex = PreAnimatedStorage->FindStorageIndex(&ViewportClient);
+	}
 	if (StorageIndex.IsValid())
 	{
 		FPreAnimatedCameraCutEditorState CachedValue = PreAnimatedStorage->GetCachedValue(StorageIndex);
diff --git a/Engine/Source/Runtime/MovieSceneTracks/Private/TrackInstances/MovieSceneCameraCutEditorHandler.h b/Engine/Source/Runtime/MovieSceneTracks/Private/TrackInstances/MovieSceneCameraCutEditorHandler.h
--- a/Engine/Source/Runtime/MovieSceneTracks/Private/TrackInstances/MovieSceneCameraCutEditorHandler.h
+++ b/Engine/Source/Runtime/MovieSceneTracks/Private/TrackInstances/MovieSceneCameraCutEditorHandler.h
@@ -38,6 +38,7 @@ struct FPreAnimatedCameraCutEditorTraits : FPreAnimatedStateTraits
 	using KeyType = FLevelEditorViewportClient*;
 	using StorageType = FPreAnimatedCameraCutEditorState;
 
+	static bool ShouldHandleViewportCameraCuts(FLevelEditorViewportClient* ViewportClient);
 	static StorageType CachePreAnimatedValue(KeyType InKey);
 
 	void RestorePreAnimatedValue(KeyType InKey, const StorageType& CachedValue, const FRestoreStateParams& Params);
@@ -56,6 +57,14 @@ struct FPreAnimatedCameraCutEditorStorage : TPreAnimatedStateStorage<FPreAnimate
 	static TAutoRegisterPreAnimatedStorageID<FPreAnimatedCameraCutEditorStorage> StorageID;
 };
 
+/** Type of forced operation on camera cut editor pre-animated state storage */
+enum class EForcedCameraCutPreAnimatedStorageOperation
+{
+	Cache,
+	Restore,
+	Discard
+};
+
 /**
  * Utility class for executing camera cuts in the editor viewports.
  */
@@ -75,6 +84,10 @@ struct FCameraCutEditorHandler
 	static void CachePreAnimatedValue(
 			UMovieSceneEntitySystemLinker* Linker,
 			const FSequenceInstance& SequenceInstance);
+	/** Force cache/discard/restore pre-animated values */
+	static void ForcePreAnimatedValueOperation(
+			UMovieSceneEntitySystemLinker* Linker,
+			EForcedCameraCutPreAnimatedStorageOperation Operation);
 
 private:
 
diff --git a/Engine/Source/Runtime/MovieSceneTracks/Private/TrackInstances/MovieSceneCameraCutGameHandler.cpp b/Engine/Source/Runtime/MovieSceneTracks/Private/TrackInstances/MovieSceneCameraCutGameHandler.cpp
--- a/Engine/Source/Runtime/MovieSceneTracks/Private/TrackInstances/MovieSceneCameraCutGameHandler.cpp
+++ b/Engine/Source/Runtime/MovieSceneTracks/Private/TrackInstances/MovieSceneCameraCutGameHandler.cpp
@@ -5,6 +5,7 @@
 #include "Camera/CameraComponent.h"
 #include "Camera/PlayerCameraManager.h"
 #include "Engine/Engine.h"
+#include "Engine/EngineTypes.h"
 #include "Engine/GameInstance.h"
 #include "Engine/LocalPlayer.h"
 #include "Engine/World.h"
@@ -18,13 +19,21 @@
 namespace UE::MovieScene
 {
 
+bool FPreAnimatedCameraCutTraits::ShouldHandleWorldCameraCuts(UWorld* World)
+{
+	return World && 
+		World->GetGameInstance() != nullptr &&
+		World->WorldType != EWorldType::Editor &&
+		World->WorldType != EWorldType::EditorPreview;
+}
+
 FPreAnimatedCameraCutState FPreAnimatedCameraCutTraits::CachePreAnimatedValue(
 		IMovieScenePlayer* Player, 
 		uint8 InKey)
 {
 	UObject* PlaybackContext = Player->GetPlaybackContext();
 	UWorld* World = PlaybackContext ? PlaybackContext->GetWorld() : nullptr;
-	if (World && World->GetGameInstance())
+	if (ShouldHandleWorldCameraCuts(World))
 	{
 		APlayerController* PC = World->GetGameInstance()->GetFirstLocalPlayerController();
 
@@ -40,7 +49,7 @@ FPreAnimatedCameraCutState FPreAnimatedCameraCutTraits::CachePreAnimatedValue(
 			AspectRatioAxisConstraint = LocalPlayer->AspectRatioAxisConstraint;
 		}
 
-		return StorageType{ LocalPlayer, ViewTarget, AspectRatioAxisConstraint };
+		return StorageType{ World, LocalPlayer, ViewTarget, AspectRatioAxisConstraint };
 	}
 	return StorageType();
 }
@@ -50,15 +59,8 @@ void FPreAnimatedCameraCutTraits::RestorePreAnimatedValue(
 		const FPreAnimatedCameraCutState& CachedValue, 
 		const FRestoreStateParams& Params)
 {
-	IMovieScenePlayer* Player = Params.GetTerminalPlayer();
-	if (!Player)
-	{
-		return;
-	}
-
-	UObject* PlaybackContext = Player->GetPlaybackContext();
-	UWorld* World = GEngine->GetWorldFromContextObject(PlaybackContext, EGetWorldErrorMode::LogAndReturnNull);
-	if (!World || !World->GetGameInstance())
+	UWorld* World = Cast<UWorld>(CachedValue.LastWorld.ResolveObjectPtr());
+	if (!ShouldHandleWorldCameraCuts(World))
 	{
 		return;
 	}
@@ -181,7 +183,7 @@ void FCameraCutGameHandler::SetCameraCut(
 	UWorld* World = PlaybackContext ? PlaybackContext->GetWorld() : nullptr;
 
 	// Also bail out if we don't have a world running any sort of game.
-	if (World == nullptr || World->GetGameInstance() == nullptr)
+	if (!FPreAnimatedCameraCutTraits::ShouldHandleWorldCameraCuts(World))
 	{
 		return;
 	}
diff --git a/Engine/Source/Runtime/MovieSceneTracks/Private/TrackInstances/MovieSceneCameraCutGameHandler.h b/Engine/Source/Runtime/MovieSceneTracks/Private/TrackInstances/MovieSceneCameraCutGameHandler.h
--- a/Engine/Source/Runtime/MovieSceneTracks/Private/TrackInstances/MovieSceneCameraCutGameHandler.h
+++ b/Engine/Source/Runtime/MovieSceneTracks/Private/TrackInstances/MovieSceneCameraCutGameHandler.h
@@ -20,6 +20,7 @@ struct FSequenceInstance;
 /** Pre-animated view target info */
 struct FPreAnimatedCameraCutState
 {
+	FObjectKey LastWorld;
 	FObjectKey LastLocalPlayer;
 	FObjectKey LastViewTarget;
 	TOptional<EAspectRatioAxisConstraint> LastAspectRatioAxisConstraint;
@@ -32,6 +33,7 @@ struct FPreAnimatedCameraCutTraits : FPreAnimatedStateTraits
 	using KeyType = uint8;
 	using StorageType = FPreAnimatedCameraCutState;
 
+	static bool ShouldHandleWorldCameraCuts(UWorld* World);
 	static StorageType CachePreAnimatedValue(IMovieScenePlayer* Player, uint8 InKey);
 
 	void RestorePreAnimatedValue(uint8 InKey, const StorageType& CachedValue, const FRestoreStateParams& Params);
diff --git a/Engine/Source/Runtime/MovieSceneTracks/Private/TrackInstances/MovieSceneCameraCutTrackInstance.cpp b/Engine/Source/Runtime/MovieSceneTracks/Private/TrackInstances/MovieSceneCameraCutTrackInstance.cpp
--- a/Engine/Source/Runtime/MovieSceneTracks/Private/TrackInstances/MovieSceneCameraCutTrackInstance.cpp
+++ b/Engine/Source/Runtime/MovieSceneTracks/Private/TrackInstances/MovieSceneCameraCutTrackInstance.cpp
@@ -10,6 +10,7 @@
 #include "EntitySystem/MovieSceneEntitySystemLinker.h"
 #include "EntitySystem/MovieSceneInstanceRegistry.h"
 #include "EntitySystem/MovieSceneSharedPlaybackState.h"
+#include "EntitySystem/TrackInstance/MovieSceneTrackInstanceSystem.h"
 #include "Evaluation/CameraCutPlaybackCapability.h"
 #include "Evaluation/MovieSceneEvaluation.h"
 #include "Evaluation/PreAnimatedState/MovieScenePreAnimatedCaptureSource.h"
@@ -22,6 +23,7 @@
 #include "Sections/MovieSceneCameraCutSection.h"
 #include "TrackInstances/MovieSceneCameraCutEditorHandler.h"
 #include "TrackInstances/MovieSceneCameraCutGameHandler.h"
+#include "TrackInstances/MovieSceneCameraCutViewportPreviewer.h"
 #include "Tracks/MovieSceneCameraCutTrack.h"
 
 #include UE_INLINE_GENERATED_CPP_BY_NAME(MovieSceneCameraCutTrackInstance)
@@ -55,6 +57,17 @@ bool FCameraCutPlaybackCapabilityCompatibilityWrapper::ShouldUpdateCameraCut()
 	return true;
 }
 
+#if WITH_EDITOR
+bool FCameraCutPlaybackCapabilityCompatibilityWrapper::ShouldCacheEditorPreAnimatedState()
+{
+	if (CameraCutCapability)
+	{
+		return CameraCutCapability->ShouldCacheEditorPreAnimatedState();
+	}
+	return true;
+}
+#endif
+
 void FCameraCutPlaybackCapabilityCompatibilityWrapper::OnCameraCutUpdated(const FOnCameraCutUpdatedParams& Params)
 {
 	if (CameraCutCapability)
@@ -265,6 +278,80 @@ public:
 }  // namespace MovieScene
 }  // namespace UE
 
+#if WITH_EDITOR
+
+void UMovieSceneCameraCutTrackInstance::ToggleCameraCutLock(UMovieSceneEntitySystemLinker* Linker, bool bEnableCameraCuts, bool bRestoreViewports)
+{
+	using namespace UE::MovieScene;
+
+	auto ForceEditorPreAnimatedStorageOperation = [](UMovieSceneCameraCutTrackInstance* This, UE::MovieScene::EForcedCameraCutPreAnimatedStorageOperation Operation)
+	{
+		using namespace UE::MovieScene;
+
+		UMovieSceneEntitySystemLinker* Linker = This->GetLinker();
+		for (const FCameraCutInputInfo& InputInfo : This->SortedInputInfos)
+		{
+			FScopedPreAnimatedCaptureSource CaptureSource(Linker, InputInfo.Input);
+			FCameraCutEditorHandler::ForcePreAnimatedValueOperation(Linker, Operation);
+		}
+	};
+
+	// Find the camera cut track instance and forcibly manage its pre-animated state for the editor 
+	// viewports depending on what we want to do with them.
+	UMovieSceneTrackInstanceInstantiator* TrackInstanceSystem = Linker->FindSystem<UMovieSceneTrackInstanceInstantiator>();
+	if (TrackInstanceSystem)
+	{
+		UMovieSceneCameraCutTrackInstance* CameraCutTrackInstance = nullptr;
+		for (auto It = TrackInstanceSystem->GetTrackInstances().CreateConstIterator(); It; ++It)
+		{
+			CameraCutTrackInstance = Cast<UMovieSceneCameraCutTrackInstance>(It->TrackInstance.Get());
+			if (CameraCutTrackInstance != nullptr)
+			{
+				break;
+			}
+		}
+
+		if (CameraCutTrackInstance)
+		{
+			if (bEnableCameraCuts)
+			{
+				// We locked the viewport to the cinematic... re-cache the viewport's position as the
+				// new pre-animated state if we intend to restore it later. Otherwise, don't cache it,
+				// so that we don't have any state to restore if we scrub/play out of the camera cut
+				// section.
+				if (bRestoreViewports)
+				{
+					ForceEditorPreAnimatedStorageOperation(CameraCutTrackInstance, EForcedCameraCutPreAnimatedStorageOperation::Cache);
+				}
+			}
+			else
+			{
+				// We unlocked the viewport... if we want to restore the original viewport position, 
+				// let's restore this position, which we saved as pre-animated value. If we want to stay
+				// in the same position as the cinematic camera, we can simply not do anything. But we in
+				// fact discard the pre-animated state altogether, because we don't want it to also come
+				// back when we scrub/play out of the camera cut section.
+				if (bRestoreViewports)
+				{
+					ForceEditorPreAnimatedStorageOperation(CameraCutTrackInstance, EForcedCameraCutPreAnimatedStorageOperation::Restore);
+				}
+				else
+				{
+					ForceEditorPreAnimatedStorageOperation(CameraCutTrackInstance, EForcedCameraCutPreAnimatedStorageOperation::Discard);
+				}
+			}
+		}
+	}
+}
+
+#endif // WITH_EDITOR
+
+void UMovieSceneCameraCutTrackInstance::OnInitialize()
+{
+#if WITH_EDITOR
+	ViewportPreviewer = MakeUnique<UE::MovieScene::FCameraCutViewportPreviewer>();
+#endif
+}
 
 void UMovieSceneCameraCutTrackInstance::OnAnimate()
 {
@@ -346,7 +433,7 @@ void UMovieSceneCameraCutTrackInstance::OnAnimate()
 
 	FCameraCutAnimator Animator(CameraCutCache);
 #if WITH_EDITOR
-	Animator.SetViewportPreviewer(&ViewportPreviewer);
+	Animator.SetViewportPreviewer(ViewportPreviewer.Get());
 #endif
 
 	// For now we only support one pre-roll.
@@ -527,8 +614,15 @@ void UMovieSceneCameraCutTrackInstance::OnEndUpdateInputs()
 void UMovieSceneCameraCutTrackInstance::OnDestroyed()
 {
 #if WITH_EDITOR
+	using namespace UE::MovieScene;
+
+	// Discard persistent storage on exit so that we don't restore whatever viewports we had
+	// way back when we opened the sequencer.
+	//FCameraCutEditorHandler::ForcePreAnimatedValueOperation(GetLinker(), EForcedCameraCutPreAnimatedStorageOperation::DiscardPersistent);
+
 	// Make sure we don't have any viewport modifiers registered anymore.
-	ViewportPreviewer.ToggleViewportPreviewModifiers(false);
+	ViewportPreviewer->ToggleViewportPreviewModifiers(false);
+	ViewportPreviewer.Reset();
 #endif
 }
 
diff --git a/Engine/Source/Runtime/MovieSceneTracks/Private/TrackInstances/MovieSceneCameraCutTrackInstance.h b/Engine/Source/Runtime/MovieSceneTracks/Private/TrackInstances/MovieSceneCameraCutTrackInstance.h
deleted file mode 100644
diff --git a/Engine/Source/Runtime/MovieSceneTracks/Public/Evaluation/CameraCutPlaybackCapability.h b/Engine/Source/Runtime/MovieSceneTracks/Public/Evaluation/CameraCutPlaybackCapability.h
--- a/Engine/Source/Runtime/MovieSceneTracks/Public/Evaluation/CameraCutPlaybackCapability.h
+++ b/Engine/Source/Runtime/MovieSceneTracks/Public/Evaluation/CameraCutPlaybackCapability.h
@@ -11,7 +11,7 @@ class UCameraComponent;
 namespace UE::MovieScene
 {
 	/**
-	 *
+	 * Parameter struct for notifying a playback client that a camera cut has happened.
 	 */
 	struct FOnCameraCutUpdatedParams
 	{
@@ -41,7 +41,8 @@ namespace UE::MovieScene
 		virtual void OnCameraCutUpdated(const FOnCameraCutUpdatedParams& Params) {}
 
 #if WITH_EDITOR
-		virtual bool ShouldPreviewCameraCutBlends() { return true; }
+		/** Whether the editor should cache the pre-animated viewport position */
+		virtual bool ShouldCacheEditorPreAnimatedState() { return true; }
 #endif
 	};
 
diff --git a/Engine/Source/Runtime/MovieSceneTracks/Public/TrackInstances/MovieSceneCameraCutTrackInstance.h b/Engine/Source/Runtime/MovieSceneTracks/Public/TrackInstances/MovieSceneCameraCutTrackInstance.h
new file mode 100644
--- /dev/null
+++ b/Engine/Source/Runtime/MovieSceneTracks/Public/TrackInstances/MovieSceneCameraCutTrackInstance.h
@@ -0,0 +1,98 @@
+// Copyright Epic Games, Inc. All Rights Reserved.
+
+#pragma once
+
+#include "CoreMinimal.h"
+#include "EntitySystem/BuiltInComponentTypes.h"
+#include "EntitySystem/TrackInstance/MovieSceneTrackInstance.h"
+#include "UObject/ObjectMacros.h"
+
+#include "MovieSceneCameraCutTrackInstance.generated.h"
+
+class IMovieScenePlayer;
+class UMovieSceneCameraCutSection;
+
+namespace UE::MovieScene
+{ 
+	class FCameraCutViewportPreviewer;
+	struct FCameraCutAnimator; 
+	struct FCameraCutPlaybackCapability;
+	struct FOnCameraCutUpdatedParams;
+	struct FSequenceInstance;
+
+	// Backwards compatibilty wrapper for camera cut playback capability.
+	struct FCameraCutPlaybackCapabilityCompatibilityWrapper
+	{
+		FCameraCutPlaybackCapabilityCompatibilityWrapper(const FSequenceInstance& SequenceInstance);
+
+		bool ShouldUpdateCameraCut();
+		void OnCameraCutUpdated(const FOnCameraCutUpdatedParams& Params);
+#if WITH_EDITOR
+		bool ShouldCacheEditorPreAnimatedState();
+#endif
+
+		FCameraCutPlaybackCapability* CameraCutCapability;
+		IMovieScenePlayer* Player;
+	};
+}
+
+/**
+ * Track instance used to animate camera cuts.
+ */
+UCLASS()
+class MOVIESCENETRACKS_API UMovieSceneCameraCutTrackInstance : public UMovieSceneTrackInstance
+{
+	GENERATED_BODY()
+
+public:
+#if WITH_EDITOR
+	/**
+	 * Toggle camera cut lock on cinematic editor viewports while also correctly managing
+	 * remember/restoring/discarding pre-animated viewport positions based on sequencer
+	 * settings.
+	 */
+	static void ToggleCameraCutLock(UMovieSceneEntitySystemLinker* Linker, bool bEnableCameraCuts, bool bRestoreViewports);
+#endif
+
+private:
+	virtual void OnInitialize() override;
+	virtual void OnAnimate() override;
+	virtual void OnInputAdded(const FMovieSceneTrackInstanceInput& InInput) override;
+	virtual void OnEndUpdateInputs() override;
+	virtual void OnDestroyed() override;
+
+private:
+	/**
+	 * Stores information about the last set camera in order to differentiate
+	 * between new and pre-existing cuts.
+	 */
+	struct FCameraCutCache
+	{
+		TWeakObjectPtr<> LastLockedCamera;
+		UE::MovieScene::FInstanceHandle LastInstanceHandle;
+		TObjectPtr<UMovieSceneSection> LastSection;
+	};
+
+	/**
+	 * Track instance input qualified with the global start time of its corresponding
+	 * section, used for sorting inputs and prioritizing more "recent" camera cuts
+	 * over "older" ones.
+	 */
+	struct FCameraCutInputInfo
+	{
+		FMovieSceneTrackInstanceInput Input;
+		float GlobalStartTime = 0.f;
+	};
+
+	FCameraCutCache CameraCutCache;
+	TArray<FCameraCutInputInfo> SortedInputInfos;
+
+#if WITH_EDITOR
+	TUniquePtr<UE::MovieScene::FCameraCutViewportPreviewer> ViewportPreviewer;
+#endif
+
+private:
+
+	friend struct UE::MovieScene::FCameraCutAnimator;
+};
+
