CIS fix for MetaSound GetDefaultParameters being called during cook from MoveSceneAudioSection with invalid runtime data

diff --git a/Engine/Source/Runtime/MovieSceneTracks/Private/Sections/MovieSceneAudioSection.cpp b/Engine/Source/Runtime/MovieSceneTracks/Private/Sections/MovieSceneAudioSection.cpp
--- a/Engine/Source/Runtime/MovieSceneTracks/Private/Sections/MovieSceneAudioSection.cpp
+++ b/Engine/Source/Runtime/MovieSceneTracks/Private/Sections/MovieSceneAudioSection.cpp
@@ -169,16 +169,16 @@ EMovieSceneChannelProxyType  UMovieSceneAudioSection::CacheChannelProxy()
 void UMovieSceneAudioSection::SetupSoundInputParameters(USoundBase* InSoundBase)
 {
 	// Populate with defaults.
-	if (InSoundBase)
+
+	// Don't init resources when running cook, as this can trigger 
+	// registration of a MetaSound and its dependent graphs.
+	// Those will instead be registered when the MetaSound itself is cooked (FMetasoundAssetBase::CookMetaSound)
+	// in a way that does not deal with runtime data like this function does
+	// Getting the default parameters and the rest of the function are 
+	// dependent on that runtime data and don't need to be cooked
+	if (InSoundBase && !IsRunningCookCommandlet())
 	{
-		// Don't init resources when running cook, as this can trigger 
-		// registration of a MetaSound and its dependent graphs.
-		// Those will instead be registered when the MetaSound itself is cooked (FMetasoundAssetBase::CookMetaSound)
-		// in a way that does not deal with runtime data like this function does
-		if (!IsRunningCookCommandlet())
-		{
-			InSoundBase->InitResources();
-		}
+		InSoundBase->InitResources();
 
 		TArray<FAudioParameter> DefaultParams;
 		InSoundBase->GetAllDefaultParameters(DefaultParams);
