Sequencer: Make sure the camera cut track only controls editor viewports that are related to the current playback context

diff --git a/Engine/Source/Runtime/MovieSceneTracks/Private/TrackInstances/MovieSceneCameraCutEditorHandler.cpp b/Engine/Source/Runtime/MovieSceneTracks/Private/TrackInstances/MovieSceneCameraCutEditorHandler.cpp
--- a/Engine/Source/Runtime/MovieSceneTracks/Private/TrackInstances/MovieSceneCameraCutEditorHandler.cpp
+++ b/Engine/Source/Runtime/MovieSceneTracks/Private/TrackInstances/MovieSceneCameraCutEditorHandler.cpp
@@ -15,11 +15,10 @@
 namespace UE::MovieScene
 {
 
-bool FPreAnimatedCameraCutEditorTraits::ShouldHandleViewportCameraCuts(FLevelEditorViewportClient* ViewportClient)
+bool FPreAnimatedCameraCutEditorTraits::ShouldHandleViewportCameraCuts(UWorld* ViewportWorld)
 {
-	UWorld* EditorWorld = ViewportClient->GetWorld();
-	return EditorWorld && 
-		(EditorWorld->WorldType == EWorldType::Editor || EditorWorld->WorldType == EWorldType::EditorPreview);
+	return ViewportWorld && 
+		(ViewportWorld->WorldType == EWorldType::Editor || ViewportWorld->WorldType == EWorldType::EditorPreview);
 }
 
 FPreAnimatedCameraCutEditorState FPreAnimatedCameraCutEditorTraits::CachePreAnimatedValue(
@@ -37,7 +36,7 @@ void FPreAnimatedCameraCutEditorTraits::RestorePreAnimatedValue(
 		const FPreAnimatedCameraCutEditorState& CachedValue, 
 		const FRestoreStateParams& Params)
 {
-	if (!GEditor)
+	if (!GEditor || !InKey)
 	{
 		return;
 	}
@@ -49,7 +48,7 @@ void FPreAnimatedCameraCutEditorTraits::RestorePreAnimatedValue(
 	}
 
 	// Check that we have an editor viewport.
-	if (!ShouldHandleViewportCameraCuts(InKey))
+	if (!ShouldHandleViewportCameraCuts(InKey->GetWorld()))
 	{
 		return;
 	}
@@ -89,12 +88,23 @@ void FCameraCutEditorHandler::CachePreAnimatedValue(
 	{
 		return;
 	}
+
+	IMovieScenePlayer* Player = SequenceInstance.GetPlayer();
+	UObject* PlaybackContext = Player->GetPlaybackContext();
+	UWorld* ContextWorld = PlaybackContext ? PlaybackContext->GetWorld() : nullptr;
+
+	// Only handle editor world/viewports.
+	if (!FPreAnimatedCameraCutEditorTraits::ShouldHandleViewportCameraCuts(ContextWorld))
+	{
+		return;
+	}
 	
 	TSharedPtr<FPreAnimatedCameraCutEditorStorage> PreAnimatedStorage = Linker->PreAnimatedState.GetOrCreateStorage<FPreAnimatedCameraCutEditorStorage>();
 
 	for (FLevelEditorViewportClient* LevelVC : GEditor->GetLevelViewportClients())
 	{
-		if (!FPreAnimatedCameraCutEditorTraits::ShouldHandleViewportCameraCuts(LevelVC))
+		// Only handle the viewports that are tied to our playback context.
+		if (!LevelVC || LevelVC->GetWorld() != ContextWorld)
 		{
 			continue;
 		}
@@ -108,6 +118,7 @@ void FCameraCutEditorHandler::CachePreAnimatedValue(
 
 void FCameraCutEditorHandler::ForcePreAnimatedValueOperation(
 		UMovieSceneEntitySystemLinker* Linker,
+		const FSequenceInstance& SequenceInstance,
 		EForcedCameraCutPreAnimatedStorageOperation Operation)
 {
 	if (!GEditor)
@@ -115,11 +126,22 @@ void FCameraCutEditorHandler::ForcePreAnimatedValueOperation(
 		return;
 	}
 
+	IMovieScenePlayer* Player = SequenceInstance.GetPlayer();
+	UObject* PlaybackContext = Player->GetPlaybackContext();
+	UWorld* ContextWorld = PlaybackContext ? PlaybackContext->GetWorld() : nullptr;
+
+	// Only handle editor world/viewports.
+	if (!FPreAnimatedCameraCutEditorTraits::ShouldHandleViewportCameraCuts(ContextWorld))
+	{
+		return;
+	}
+
 	TSharedPtr<FPreAnimatedCameraCutEditorStorage> PreAnimatedStorage = Linker->PreAnimatedState.GetOrCreateStorage<FPreAnimatedCameraCutEditorStorage>();
 
 	for (FLevelEditorViewportClient* LevelVC : GEditor->GetLevelViewportClients())
 	{
-		if (!FPreAnimatedCameraCutEditorTraits::ShouldHandleViewportCameraCuts(LevelVC))
+		// Only handle the viewports that are tied to our playback context.
+		if (!LevelVC || LevelVC->GetWorld() != ContextWorld)
 		{
 			continue;
 		}
@@ -175,6 +197,16 @@ void FCameraCutEditorHandler::SetCameraCut(
 		return;
 	}
 
+	IMovieScenePlayer* Player = SequenceInstance.GetPlayer();
+	UObject* PlaybackContext = Player->GetPlaybackContext();
+	UWorld* ContextWorld = PlaybackContext ? PlaybackContext->GetWorld() : nullptr;
+
+	// Only handle editor world/viewports.
+	if (!FPreAnimatedCameraCutEditorTraits::ShouldHandleViewportCameraCuts(ContextWorld))
+	{
+		return;
+	}
+
 	FCameraCutPlaybackCapabilityCompatibilityWrapper Wrapper(SequenceInstance);
 
 	// If we don't want to update camera cuts, let's remember it and release the viewports
@@ -193,7 +225,8 @@ void FCameraCutEditorHandler::SetCameraCut(
 
 	for (FLevelEditorViewportClient* LevelVC : GEditor->GetLevelViewportClients())
 	{
-		if (LevelVC == nullptr || !FPreAnimatedCameraCutEditorTraits::ShouldHandleViewportCameraCuts(LevelVC))
+		// Only handle the viewports that are tied to our playback context.
+		if (!LevelVC || LevelVC->GetWorld() != ContextWorld)
 		{
 			continue;
 		}
diff --git a/Engine/Source/Runtime/MovieSceneTracks/Private/TrackInstances/MovieSceneCameraCutEditorHandler.h b/Engine/Source/Runtime/MovieSceneTracks/Private/TrackInstances/MovieSceneCameraCutEditorHandler.h
--- a/Engine/Source/Runtime/MovieSceneTracks/Private/TrackInstances/MovieSceneCameraCutEditorHandler.h
+++ b/Engine/Source/Runtime/MovieSceneTracks/Private/TrackInstances/MovieSceneCameraCutEditorHandler.h
@@ -38,7 +38,7 @@ struct FPreAnimatedCameraCutEditorTraits : FPreAnimatedStateTraits
 	using KeyType = FLevelEditorViewportClient*;
 	using StorageType = FPreAnimatedCameraCutEditorState;
 
-	static bool ShouldHandleViewportCameraCuts(FLevelEditorViewportClient* ViewportClient);
+	static bool ShouldHandleViewportCameraCuts(UWorld* ViewportWorld);
 	static StorageType CachePreAnimatedValue(KeyType InKey);
 
 	void RestorePreAnimatedValue(KeyType InKey, const StorageType& CachedValue, const FRestoreStateParams& Params);
@@ -87,6 +87,7 @@ struct FCameraCutEditorHandler
 	/** Force cache/discard/restore pre-animated values */
 	static void ForcePreAnimatedValueOperation(
 			UMovieSceneEntitySystemLinker* Linker,
+			const FSequenceInstance& SequenceInstance,
 			EForcedCameraCutPreAnimatedStorageOperation Operation);
 
 private:
diff --git a/Engine/Source/Runtime/MovieSceneTracks/Private/TrackInstances/MovieSceneCameraCutTrackInstance.cpp b/Engine/Source/Runtime/MovieSceneTracks/Private/TrackInstances/MovieSceneCameraCutTrackInstance.cpp
--- a/Engine/Source/Runtime/MovieSceneTracks/Private/TrackInstances/MovieSceneCameraCutTrackInstance.cpp
+++ b/Engine/Source/Runtime/MovieSceneTracks/Private/TrackInstances/MovieSceneCameraCutTrackInstance.cpp
@@ -289,10 +289,12 @@ void UMovieSceneCameraCutTrackInstance::ToggleCameraCutLock(UMovieSceneEntitySys
 		using namespace UE::MovieScene;
 
 		UMovieSceneEntitySystemLinker* Linker = This->GetLinker();
+		const FInstanceRegistry* InstanceRegistry = Linker->GetInstanceRegistry();
 		for (const FCameraCutInputInfo& InputInfo : This->SortedInputInfos)
 		{
 			FScopedPreAnimatedCaptureSource CaptureSource(Linker, InputInfo.Input);
-			FCameraCutEditorHandler::ForcePreAnimatedValueOperation(Linker, Operation);
+			const FSequenceInstance& SequenceInstance = InstanceRegistry->GetInstance(InputInfo.Input.InstanceHandle);
+			FCameraCutEditorHandler::ForcePreAnimatedValueOperation(Linker, SequenceInstance, Operation);
 		}
 	};
 
@@ -614,12 +616,6 @@ void UMovieSceneCameraCutTrackInstance::OnEndUpdateInputs()
 void UMovieSceneCameraCutTrackInstance::OnDestroyed()
 {
 #if WITH_EDITOR
-	using namespace UE::MovieScene;
-
-	// Discard persistent storage on exit so that we don't restore whatever viewports we had
-	// way back when we opened the sequencer.
-	//FCameraCutEditorHandler::ForcePreAnimatedValueOperation(GetLinker(), EForcedCameraCutPreAnimatedStorageOperation::DiscardPersistent);
-
 	// Make sure we don't have any viewport modifiers registered anymore.
 	ViewportPreviewer->ToggleViewportPreviewModifiers(false);
 	ViewportPreviewer.Reset();
