Ignore false positive in static code analysis.

diff --git a/Engine/Source/Runtime/MovieSceneTracks/Private/TrackInstances/MovieSceneCameraCutGameHandler.cpp b/Engine/Source/Runtime/MovieSceneTracks/Private/TrackInstances/MovieSceneCameraCutGameHandler.cpp
--- a/Engine/Source/Runtime/MovieSceneTracks/Private/TrackInstances/MovieSceneCameraCutGameHandler.cpp
+++ b/Engine/Source/Runtime/MovieSceneTracks/Private/TrackInstances/MovieSceneCameraCutGameHandler.cpp
@@ -35,6 +35,7 @@ FPreAnimatedCameraCutState FPreAnimatedCameraCutTraits::CachePreAnimatedValue(
 	UWorld* World = PlaybackContext ? PlaybackContext->GetWorld() : nullptr;
 	if (ShouldHandleWorldCameraCuts(World))
 	{
+		CA_SUPPRESS(6011);
 		APlayerController* PC = World->GetGameInstance()->GetFirstLocalPlayerController();
 
 		// Save previous view target.
@@ -188,6 +189,7 @@ void FCameraCutGameHandler::SetCameraCut(
 		return;
 	}
 
+	CA_SUPPRESS(6011);
 	APlayerController* PC = World->GetGameInstance()->GetFirstLocalPlayerController();
 	APlayerCameraManager* CameraManager = (PC != nullptr) ? PC->PlayerCameraManager : nullptr;
 
