Sequencer: fix upgrade of default values from bool sections to visibility sections.

diff --git a/Engine/Source/Runtime/MovieSceneTracks/Private/Tracks/MovieSceneVisibilityTrack.cpp b/Engine/Source/Runtime/MovieSceneTracks/Private/Tracks/MovieSceneVisibilityTrack.cpp
--- a/Engine/Source/Runtime/MovieSceneTracks/Private/Tracks/MovieSceneVisibilityTrack.cpp
+++ b/Engine/Source/Runtime/MovieSceneTracks/Private/Tracks/MovieSceneVisibilityTrack.cpp
@@ -32,6 +32,10 @@ void UMovieSceneVisibilityTrack::PostLoad()
 
 			FObjectWriter(BoolSection, Bytes);
 			UMovieSceneVisibilitySection* NewSection = NewObject<UMovieSceneVisibilitySection>(this, NAME_None, RF_Transactional);
+			// Bool sections start with DefaultValue=false and bHasDefaultValue=false, so we need
+			// to match this in order for the delta-serialization to do the right thing.
+			NewSection->GetChannel().SetDefault(false);
+			NewSection->GetChannel().RemoveDefault();
 			FObjectReader(NewSection, Bytes);
 
 			Sections[Index] = NewSection;
