Similar CIS fix as previously

diff --git a/Engine/Source/Runtime/MovieSceneTracks/Private/Systems/MovieSceneCameraShakeSystem.cpp b/Engine/Source/Runtime/MovieSceneTracks/Private/Systems/MovieSceneCameraShakeSystem.cpp
--- a/Engine/Source/Runtime/MovieSceneTracks/Private/Systems/MovieSceneCameraShakeSystem.cpp
+++ b/Engine/Source/Runtime/MovieSceneTracks/Private/Systems/MovieSceneCameraShakeSystem.cpp
@@ -189,7 +189,9 @@ FCameraShakePreviewer& FCameraShakePreviewerLinkerExtension::GetPreviewer(FInsta
 	NewPreviewer.RegisterViewModifiers([ContextWorld](FLevelEditorViewportClient* LevelVC) -> bool
 				{
 					return LevelVC->AllowsCinematicControl() && LevelVC->GetWorld() == ContextWorld;
-				});
+				},
+				// Pass false to help the Mac compiler along.
+				false);
 	return NewPreviewer;
 }
 
