Sequencer: resolve desired view target before checking if it's the same as the current one

When blending back to gameplay we pass a null pointer for the desired view target. We would later resolve this pointer with the pre-animated pointer, but not before checking if the desired view target was already set, in which case we would early-return. This caused the small issue that, once blending back to gameplay started, we should indeed early-return (since the camera change is already going on), but we weren't because we were checking against the null pointer, not against the resolved pre-animated pointer.

This change reorders a few things to make it all happen in the correct order.

diff --git a/Engine/Source/Runtime/MovieSceneTracks/Private/TrackInstances/MovieSceneCameraCutGameHandler.cpp b/Engine/Source/Runtime/MovieSceneTracks/Private/TrackInstances/MovieSceneCameraCutGameHandler.cpp
--- a/Engine/Source/Runtime/MovieSceneTracks/Private/TrackInstances/MovieSceneCameraCutGameHandler.cpp
+++ b/Engine/Source/Runtime/MovieSceneTracks/Private/TrackInstances/MovieSceneCameraCutGameHandler.cpp
@@ -211,8 +211,38 @@ void FCameraCutGameHandler::SetCameraCut(
 		ViewTarget = CameraManager->ViewTarget.Target;
 	}
 
+	// If UnlockIfCameraActor is valid, release lock only if currently locked to the specified object.
+	AActor* UnlockIfCameraActor = Cast<AActor>(CameraCutParams.UnlockIfCameraObject);
+	if (CameraObject == nullptr && UnlockIfCameraActor != nullptr && UnlockIfCameraActor != ViewTarget)
+	{
+		return;
+	}
+
+	// See if we need to override the aspect ratio axis constraint.
+	TOptional<EAspectRatioAxisConstraint> OverrideAspectRatioAxisConstraint;
+	if (Wrapper.CameraCutCapability)
+	{
+		OverrideAspectRatioAxisConstraint = Wrapper.CameraCutCapability->GetAspectRatioAxisConstraintOverride();
+	}
+
+	// CameraObject is null if we need to release control, which can happen here (instead of via pre-animated 
+	// state restore) if we are *blending* back to gameplay, and not cutting back to it at the end of a camera 
+	// cut section. Let's get the pre-animated value and blend back towards it.
+	if (CameraObject == nullptr)
+	{
+		TSharedPtr<FPreAnimatedCameraCutStorage> PreAnimatedStorage = Linker->PreAnimatedState.FindStorage(FPreAnimatedCameraCutStorage::StorageID);
+		FPreAnimatedStorageIndex StorageIndex = PreAnimatedStorage->FindStorageIndex(0);
+		if (ensureMsgf(StorageIndex.IsValid(), TEXT("Blending camera back to gameplay but can't find pre-animated camera info!")))
+		{
+			FPreAnimatedCameraCutState CachedValue = PreAnimatedStorage->GetCachedValue(StorageIndex);
+			CameraObject = CachedValue.LastViewTarget.ResolveObjectPtr();
+			OverrideAspectRatioAxisConstraint = CachedValue.LastAspectRatioAxisConstraint;
+		}
+	}
+
 	// If we find a camera component inside the provided object, let's make sure we are going to set
 	// its owner as the next view target, and not some component (including the camera component itself).
+	AActor* CameraActor = Cast<AActor>(CameraObject);
 	UCameraComponent* CameraComponent = MovieSceneHelpers::CameraComponentFromRuntimeObject(CameraObject);
 	if (CameraComponent && CameraComponent->GetOwner() != CameraObject)
 	{
@@ -243,39 +273,6 @@ void FCameraCutGameHandler::SetCameraCut(
 		return;
 	}
 
-	// If UnlockIfCameraActor is valid, release lock only if currently locked to the specified object.
-	AActor* UnlockIfCameraActor = Cast<AActor>(CameraCutParams.UnlockIfCameraObject);
-	if (CameraObject == nullptr && UnlockIfCameraActor != nullptr && UnlockIfCameraActor != ViewTarget)
-	{
-		return;
-	}
-
-	AActor* CameraActor = Cast<AActor>(CameraObject);
-	ULocalPlayer* LocalPlayer = (PC != nullptr) ? PC->GetLocalPlayer() : nullptr;
-
-	TOptional<EAspectRatioAxisConstraint> OverrideAspectRatioAxisConstraint;
-	if (Wrapper.CameraCutCapability)
-	{
-		OverrideAspectRatioAxisConstraint = Wrapper.CameraCutCapability->GetAspectRatioAxisConstraintOverride();
-	}
-
-	// CameraObject is null if we need to release control, which can happen here (instead of via pre-animated 
-	// state restore) if we are *blending* back to gameplay, and not cutting back to it at the end of a camera 
-	// cut section. Let's get the pre-animated value and blend back towards it.
-	if (CameraObject == nullptr)
-	{
-		TSharedPtr<FPreAnimatedCameraCutStorage> PreAnimatedStorage = Linker->PreAnimatedState.FindStorage(FPreAnimatedCameraCutStorage::StorageID);
-		FPreAnimatedStorageIndex StorageIndex = PreAnimatedStorage->FindStorageIndex(0);
-		if (ensureMsgf(StorageIndex.IsValid(), TEXT("Blending camera back to gameplay but can't find pre-animated camera info!")))
-		{
-			FPreAnimatedCameraCutState CachedValue = PreAnimatedStorage->GetCachedValue(StorageIndex);
-			CameraObject = CachedValue.LastViewTarget.ResolveObjectPtr();
-			CameraActor = Cast<AActor>(CameraObject);
-			CameraComponent = MovieSceneHelpers::CameraComponentFromRuntimeObject(CameraActor);
-			OverrideAspectRatioAxisConstraint = CachedValue.LastAspectRatioAxisConstraint;
-		}
-	}
-
 	// Time to set the camera cut! How we do it depends on whether we need to do some blending, or a straight cut.
 	bool bDoSetViewTarget = true;
 	FViewTargetTransitionParams TransitionParams;
@@ -310,12 +307,13 @@ void FCameraCutGameHandler::SetCameraCut(
 		UE_LOG(LogMovieScene, Log, TEXT("Starting new camera cut: '%s'"),
 			(CameraActor ? *CameraActor->GetName() : TEXT("None")));
 	}
-	if (bDoSetViewTarget && ensureMsgf(PC, TEXT("Can't set view target when there is no player controller!")))
+	if (bDoSetViewTarget && ensureMsgf(CameraManager, TEXT("Can't set view target when there is no player controller!")))
 	{
-		PC->SetViewTarget(CameraActor, TransitionParams);
+		CameraManager->SetViewTarget(CameraActor, TransitionParams);
 	}
 
 	// Override the aspect ratio constraint if this sequence requires it.
+	ULocalPlayer* LocalPlayer = (PC != nullptr) ? PC->GetLocalPlayer() : nullptr;
 	if (LocalPlayer != nullptr && OverrideAspectRatioAxisConstraint.IsSet())
 	{
 		LocalPlayer->AspectRatioAxisConstraint = OverrideAspectRatioAxisConstraint.GetValue();
