Sequencer: preserve camera shake instances across re-imports

This lets us re-use these camera shakes every time the sequence is recompiled, which can happen a lot when it's being edited in sequencer. Previously, this would make all shakes constantly recreated and restarted, spamming the GC and making all cameras shake even though the sequence is paused.

Also fix an issue where the root instance context was used to run the camera shake, instead of the local (sub) sequence instance. Comparing current time to shake section start time was incorrect.

diff --git a/Engine/Source/Runtime/MovieSceneTracks/Private/MovieSceneTracksComponentTypes.cpp b/Engine/Source/Runtime/MovieSceneTracks/Private/MovieSceneTracksComponentTypes.cpp
--- a/Engine/Source/Runtime/MovieSceneTracks/Private/MovieSceneTracksComponentTypes.cpp
+++ b/Engine/Source/Runtime/MovieSceneTracks/Private/MovieSceneTracksComponentTypes.cpp
@@ -1,6 +1,8 @@
 // Copyright Epic Games, Inc. All Rights Reserved.
 
 #include "MovieSceneTracksComponentTypes.h"
+#include "Camera/CameraShakeBase.h"
+#include "Camera/CameraShakeSourceComponent.h"
 #include "Components/SkeletalMeshComponent.h"
 #include "MovieSceneTracksCustomAccessors.h"
 #include "EntitySystem/BuiltInComponentTypes.h"
@@ -562,6 +564,7 @@ FMovieSceneTracksComponentTypes::FMovieSceneTracksComponentTypes()
 	ComponentRegistry->NewComponentType(&AudioTriggerName,       TEXT("Audio Trigger Name"), EComponentTypeFlags::CopyToChildren);
 
 	ComponentRegistry->NewComponentType(&CameraShake,            TEXT("Camera Shake"), EComponentTypeFlags::CopyToChildren);
+	ComponentRegistry->NewComponentType(&CameraShakeInstance,    TEXT("Camera Shake Instance"), EComponentTypeFlags::Preserved);
 
 	Tags.BoundMaterialChanged = ComponentRegistry->NewTag(TEXT("Bound Material Changed"));
 	FBuiltInComponentTypes::Get()->RequiresInstantiationMask.Set(Tags.BoundMaterialChanged);
@@ -771,6 +774,10 @@ FMovieSceneTracksComponentTypes::FMovieSceneTracksComponentTypes()
 	// Set up custom primitive data components
 	ComponentRegistry->Factories.DefineChildComponent(Tags.CustomPrimitiveData, Tags.CustomPrimitiveData);
 
+	// --------------------------------------------------------------------------------------------
+	// Set up camera shake components
+	ComponentRegistry->Factories.DefineMutuallyInclusiveComponent(CameraShake, CameraShakeInstance);
+
 	InitializeMovieSceneTracksAccessors(this);
 }
 
diff --git a/Engine/Source/Runtime/MovieSceneTracks/Private/Sections/MovieSceneCameraShakeSection.cpp b/Engine/Source/Runtime/MovieSceneTracks/Private/Sections/MovieSceneCameraShakeSection.cpp
--- a/Engine/Source/Runtime/MovieSceneTracks/Private/Sections/MovieSceneCameraShakeSection.cpp
+++ b/Engine/Source/Runtime/MovieSceneTracks/Private/Sections/MovieSceneCameraShakeSection.cpp
@@ -59,7 +59,7 @@ void UMovieSceneCameraShakeSection::ImportEntityImpl(UMovieSceneEntitySystemLink
 	const FMovieSceneTracksComponentTypes* TrackComponents = FMovieSceneTracksComponentTypes::Get();
 
 	const FGuid ObjectBindingID = Params.GetObjectBindingID();
-	FMovieSceneCameraShakeComponentData ComponentData(ShakeData, GetRange());
+	FMovieSceneCameraShakeComponentData ComponentData(ShakeData, *this);
 
 	OutImportedEntity->AddBuilder(
 		FEntityBuilder()
diff --git a/Engine/Source/Runtime/MovieSceneTracks/Private/Sections/MovieSceneCameraShakeSourceShakeSection.cpp b/Engine/Source/Runtime/MovieSceneTracks/Private/Sections/MovieSceneCameraShakeSourceShakeSection.cpp
--- a/Engine/Source/Runtime/MovieSceneTracks/Private/Sections/MovieSceneCameraShakeSourceShakeSection.cpp
+++ b/Engine/Source/Runtime/MovieSceneTracks/Private/Sections/MovieSceneCameraShakeSourceShakeSection.cpp
@@ -19,7 +19,7 @@ void UMovieSceneCameraShakeSourceShakeSection::ImportEntityImpl(UMovieSceneEntit
 	const FMovieSceneTracksComponentTypes* TrackComponents = FMovieSceneTracksComponentTypes::Get();
 
 	const FGuid ObjectBindingID = Params.GetObjectBindingID();
-	FMovieSceneCameraShakeComponentData ComponentData(ShakeData, GetRange());
+	FMovieSceneCameraShakeComponentData ComponentData(ShakeData, *this);
 
 	OutImportedEntity->AddBuilder(
 		FEntityBuilder()
diff --git a/Engine/Source/Runtime/MovieSceneTracks/Private/Systems/MovieSceneCameraShakeSystem.cpp b/Engine/Source/Runtime/MovieSceneTracks/Private/Systems/MovieSceneCameraShakeSystem.cpp
--- a/Engine/Source/Runtime/MovieSceneTracks/Private/Systems/MovieSceneCameraShakeSystem.cpp
+++ b/Engine/Source/Runtime/MovieSceneTracks/Private/Systems/MovieSceneCameraShakeSystem.cpp
@@ -328,7 +328,8 @@ void UMovieSceneCameraShakeInstantiatorSystem::OnRun(FSystemTaskPrerequisites& I
 		FMovieSceneEntityID EntityID,
 		FInstanceHandle InstanceHandle,
 		UObject* BoundObject,
-		FMovieSceneCameraShakeComponentData& ShakeData)
+		const FMovieSceneCameraShakeComponentData& ShakeData,
+		FMovieSceneCameraShakeInstanceData& ShakeInstanceData)
 	{
 		const FSequenceInstance& Instance = InstanceRegistry->GetInstance(InstanceHandle);
 		const FMovieSceneContext& Context = Instance.GetContext();
@@ -350,7 +351,8 @@ void UMovieSceneCameraShakeInstantiatorSystem::OnRun(FSystemTaskPrerequisites& I
 		}
 
 		// Get the duration of the shake and compare it to the duration of the section, to know
-		// if we need to override it.
+		// if we need to override it. We only override duration if the section is shorter than
+		// the shake's natural duration... i.e, we only make shakes shorter, not longer.
 		const FFrameRate& FrameRate = Context.GetFrameRate();
 		TOptional<float> DurationOverride;
 
@@ -377,24 +379,40 @@ void UMovieSceneCameraShakeInstantiatorSystem::OnRun(FSystemTaskPrerequisites& I
 		// Start playing the shake.
 		if (ShakeSourceComponent)
 		{
+			bool bStartShake = true;
+			if (ShakeInstanceData.SectionSignature == ShakeData.SectionSignature)
+			{
+				// Don't re-create and restart the shake if it was already running with the same
+				// parameters and this instantiation phase is just re-importing other stuff.
+				bStartShake = false;
+			}
+
 			PreAnimatedCameraSourceShakeStorage->BeginTrackingEntity(EntityID, bWantsRestoreState, RootInstanceHandle, ShakeSourceComponent);
 			PreAnimatedCameraSourceShakeStorage->CachePreAnimatedValue(CacheParams, ShakeSourceComponent);
 
-			FCameraShakeSourceComponentStartParams ComponentParams;
-			ComponentParams.ShakeClass = ShakeClass;
-			ComponentParams.Scale = ShakeData.SectionData.PlayScale;
-			ComponentParams.PlaySpace = ShakeData.SectionData.PlaySpace;
-			ComponentParams.UserPlaySpaceRot = ShakeData.SectionData.UserDefinedPlaySpace;
-			ComponentParams.DurationOverride = DurationOverride;
-			ShakeSourceComponent->StartCameraShake(ComponentParams);
+			if (bStartShake)
+			{
+				FCameraShakeSourceComponentStartParams ComponentParams;
+				ComponentParams.ShakeClass = ShakeClass;
+				ComponentParams.Scale = ShakeData.SectionData.PlayScale;
+				ComponentParams.PlaySpace = ShakeData.SectionData.PlaySpace;
+				ComponentParams.UserPlaySpaceRot = ShakeData.SectionData.UserDefinedPlaySpace;
+				ComponentParams.DurationOverride = DurationOverride;
+
+				ShakeSourceComponent->StartCameraShake(ComponentParams);
+			}
+
+			ShakeInstanceData.SectionSignature = ShakeData.SectionSignature;
 
 #if WITH_EDITOR
-			if (PreviewerExtension)
+			if (PreviewerExtension && bStartShake)
 			{
 				// Shake source components start shakes in the world, unlike the other shakes
 				// (in the `else` clause) who directly affect the bound camera. This means that
 				// the shake we have just started won't affect the Sequencer preview unless we
 				// add some shaking ourselves. Let's do that here.
+				// In fact, in the editor, the above StartCameraShake call generally does nothing
+				// since there is no player controller.
 				FCameraShakePreviewer& Previewer = PreviewerExtension->GetPreviewer(InstanceHandle);
 
 				FCameraShakePreviewerAddParams PreviewParams;
@@ -405,14 +423,34 @@ void UMovieSceneCameraShakeInstantiatorSystem::OnRun(FSystemTaskPrerequisites& I
 				PreviewParams.PlaySpace = ShakeData.SectionData.PlaySpace;
 				PreviewParams.UserPlaySpaceRot = ShakeData.SectionData.UserDefinedPlaySpace;
 				PreviewParams.DurationOverride = DurationOverride;
-				Previewer.AddCameraShake(PreviewParams);
+
+				// Stop any previous version of this shake.
+				UCameraShakeBase* OldShakeInstance = ShakeInstanceData.ShakeInstance;
+				if (OldShakeInstance)
+				{
+					Previewer.RemoveCameraShake(OldShakeInstance);
+				}
+
+				ShakeInstanceData.ShakeInstance = Previewer.AddCameraShake(PreviewParams);
 			}
 #endif  // WITH_EDITOR
 		}
 		else if (UCameraComponent* CameraComponent = MovieSceneHelpers::CameraComponentFromRuntimeObject(BoundObject))
 		{
-			UObject* OuterObject = Player->GetPlaybackContext() ? Player->GetPlaybackContext() : GetTransientPackage();
-			UCameraShakeBase* ShakeInstance = NewObject<UCameraShakeBase>(OuterObject, ShakeClass);
+			bool bStartInstance = true;
+			UCameraShakeBase* ShakeInstance = nullptr;
+			if (ShakeInstanceData.SectionSignature == ShakeData.SectionSignature)
+			{
+				// Don't re-create and restart the shake if it was already running with the same
+				// parameters and this instantiation phase is just re-importing other stuff.
+				bStartInstance = false;
+				ShakeInstance = ShakeInstanceData.ShakeInstance;
+			}
+			else
+			{
+				UObject* OuterObject = Player->GetPlaybackContext() ? Player->GetPlaybackContext() : GetTransientPackage();
+				ShakeInstance = NewObject<UCameraShakeBase>(OuterObject, ShakeClass);
+			}
 
 			PreAnimatedCameraShakeStorage->BeginTrackingEntity(EntityID, bWantsRestoreState, RootInstanceHandle, ShakeInstance);
 			PreAnimatedCameraShakeStorage->CachePreAnimatedValue(CacheParams, ShakeInstance);
@@ -420,14 +458,19 @@ void UMovieSceneCameraShakeInstantiatorSystem::OnRun(FSystemTaskPrerequisites& I
 			PreAnimatedCameraComponentShakeStorage->BeginTrackingEntity(EntityID, bWantsRestoreState, RootInstanceHandle, CameraComponent);
 			PreAnimatedCameraComponentShakeStorage->CachePreAnimatedValue(CacheParams, CameraComponent);
 
-			FCameraShakeBaseStartParams ShakeParams;
-			ShakeParams.Scale = ShakeData.SectionData.PlayScale;
-			ShakeParams.PlaySpace = ShakeData.SectionData.PlaySpace;
-			ShakeParams.UserPlaySpaceRot = ShakeData.SectionData.UserDefinedPlaySpace;
-			ShakeParams.DurationOverride = DurationOverride;
-			ShakeInstance->StartShake(ShakeParams);
+			ShakeInstanceData.ShakeInstance = ShakeInstance;
+			ShakeInstanceData.SectionSignature = ShakeData.SectionSignature;
+
+			if (bStartInstance)
+			{
+				FCameraShakeBaseStartParams ShakeParams;
+				ShakeParams.Scale = ShakeData.SectionData.PlayScale;
+				ShakeParams.PlaySpace = ShakeData.SectionData.PlaySpace;
+				ShakeParams.UserPlaySpaceRot = ShakeData.SectionData.UserDefinedPlaySpace;
+				ShakeParams.DurationOverride = DurationOverride;
 
-			ShakeData.ShakeInstance = ShakeInstance;
+				ShakeInstance->StartShake(ShakeParams);
+			}
 		}
 	};
 
@@ -435,7 +478,8 @@ void UMovieSceneCameraShakeInstantiatorSystem::OnRun(FSystemTaskPrerequisites& I
 		.ReadEntityIDs()
 		.Read(BuiltInComponents->InstanceHandle)
 		.Read(BuiltInComponents->BoundObject)
-		.Write(TrackComponents->CameraShake)
+		.Read(TrackComponents->CameraShake)
+		.Write(TrackComponents->CameraShakeInstance)
 		.FilterAll({ BuiltInComponents->Tags.NeedsLink })
 		.Iterate_PerEntity(&Linker->EntityManager, VisitNewShakes);
 
@@ -600,12 +644,17 @@ struct FEvaluateCameraShake
 	{
 	}
 
-	void ForEachAllocation(const FEntityAllocation* Allocation, TRead<FMovieSceneEntityID> EntityIDs, TRead<FRootInstanceHandle> RootInstanceHandles, TRead<FInstanceHandle> InstanceHandles, TRead<UObject*> BoundObjects, TWrite<FMovieSceneCameraShakeComponentData> ShakeComponents)
+	void ForEachAllocation(
+			const FEntityAllocation* Allocation,
+			TRead<FInstanceHandle> InstanceHandles,
+			TRead<UObject*> BoundObjects,
+			TRead<FMovieSceneCameraShakeComponentData> ShakeComponents,
+			TRead<FMovieSceneCameraShakeInstanceData> ShakeInstances)
 	{
 		const int32 Num = Allocation->Num();
 		for (int32 Index = 0; Index < Num; ++Index)
 		{
-			FRootInstanceHandle InstanceHandle = RootInstanceHandles[Index];
+			FInstanceHandle InstanceHandle = InstanceHandles[Index];
 			const FSequenceInstance& Instance = InstanceRegistry->GetInstance(InstanceHandle);
 
 			// Shakes should have been started by the instantiator system.
@@ -614,17 +663,19 @@ struct FEvaluateCameraShake
 			// along by themselves in both the player camera manager (in the game) and the camera shake
 			// previewer (in the editor). We do however need to tick the camera shakes running directly
 			// onto camera component bindings.
-			FMovieSceneCameraShakeComponentData& ShakeData = ShakeComponents[Index];
+			const FMovieSceneCameraShakeComponentData& ShakeData = ShakeComponents[Index];
+			const FMovieSceneCameraShakeInstanceData& ShakeInstance = ShakeInstances[Index];
 			if (UCameraComponent* CameraComponent = MovieSceneHelpers::CameraComponentFromRuntimeObject(BoundObjects[Index]))
 			{
-				EvaluateCameraComponentShake(CameraComponent, ShakeData, Instance);
+				EvaluateCameraComponentShake(CameraComponent, ShakeData, ShakeInstance, Instance);
 			}
 		}
 	}
 
 	void EvaluateCameraComponentShake(
 			UCameraComponent* CameraComponent,
-			FMovieSceneCameraShakeComponentData& ShakeData,
+			const FMovieSceneCameraShakeComponentData& ShakeData,
+			const FMovieSceneCameraShakeInstanceData& ShakeInstance,
 			const FSequenceInstance& Instance)
 	{
 		FMinimalViewInfo POV;
@@ -635,7 +686,7 @@ struct FEvaluateCameraShake
 		// Update shake to the new time.
 		const FMovieSceneContext& Context = Instance.GetContext();
 		const FFrameTime NewShakeTime = Context.GetTime() - ShakeData.SectionStartTime;
-		ShakeData.ShakeInstance->ScrubAndApplyCameraShake(NewShakeTime / Context.GetFrameRate(), 1.f, POV);
+		ShakeInstance.ShakeInstance->ScrubAndApplyCameraShake(NewShakeTime / Context.GetFrameRate(), 1.f, POV);
 
 		// Grab transform and FOV changes.
 		FTransform WorldToBaseCamera = CameraComponent->GetComponentToWorld().Inverse();
@@ -726,11 +777,10 @@ void UMovieSceneCameraShakeEvaluatorSystem::OnRun(FSystemTaskPrerequisites& InPr
 	const FMovieSceneTracksComponentTypes* TrackComponents = FMovieSceneTracksComponentTypes::Get();
 
 	FEntityTaskBuilder()
-		.ReadEntityIDs()
-		.Read(BuiltInComponents->RootInstanceHandle)
 		.Read(BuiltInComponents->InstanceHandle)
 		.Read(BuiltInComponents->BoundObject)
-		.Write(TrackComponents->CameraShake)
+		.Read(TrackComponents->CameraShake)
+		.Read(TrackComponents->CameraShakeInstance)
 		.SetDesiredThread(Linker->EntityManager.GetGatherThread())
 		.Dispatch_PerAllocation<FEvaluateCameraShake>(
 			&Linker->EntityManager, InPrerequisites, &Subsequents, Linker);
diff --git a/Engine/Source/Runtime/MovieSceneTracks/Public/MovieSceneTracksComponentTypes.h b/Engine/Source/Runtime/MovieSceneTracks/Public/MovieSceneTracksComponentTypes.h
--- a/Engine/Source/Runtime/MovieSceneTracks/Public/MovieSceneTracksComponentTypes.h
+++ b/Engine/Source/Runtime/MovieSceneTracks/Public/MovieSceneTracksComponentTypes.h
@@ -27,6 +27,8 @@ class UMovieSceneAudioSection;
 class UMovieSceneDataLayerSection;
 class UMovieSceneLevelVisibilitySection;
 class UMovieSceneSkeletalAnimationSection;
+struct FCameraShakeBaseStartParams;
+struct FCameraShakeSourceComponentStartParams;
 struct FMovieSceneObjectBindingID;
 
 
@@ -119,26 +121,45 @@ struct FMovieSceneCameraShakeComponentData
 	/** The range of the section that created this component */
 	UPROPERTY()
 	FFrameNumber SectionStartTime;
-
 	UPROPERTY()
 	FFrameNumber SectionEndTime;
 
-	// ~~ Internal data for the camera shake systems ~~
-
-	// Shake instance created by the shake evaluation system.
+	/** The signature of the source section at the time the shake instance was created */
 	UPROPERTY()
-	TObjectPtr<UCameraShakeBase> ShakeInstance;
+	FGuid SectionSignature;
 
 	FMovieSceneCameraShakeComponentData()
 	{}
-	FMovieSceneCameraShakeComponentData(const FMovieSceneCameraShakeSectionData& InSectionData, TRange<FFrameNumber> InSectionRange)
+	FMovieSceneCameraShakeComponentData(const FMovieSceneCameraShakeSectionData& InSectionData, const UMovieSceneSection& InSection)
 		: SectionData(InSectionData)
+		, SectionSignature(InSection.GetSignature())
 	{
-		SectionStartTime = InSectionRange.HasLowerBound() ? InSectionRange.GetLowerBoundValue() : 0;
-		SectionEndTime = InSectionRange.HasUpperBound() ? InSectionRange.GetUpperBoundValue() : FFrameNumber(TNumericLimits<int32>::Max());
+		const TRange<FFrameNumber> SectionRange = InSection.GetRange();
+		SectionStartTime = SectionRange.HasLowerBound() ? SectionRange.GetLowerBoundValue() : 0;
+		SectionEndTime = SectionRange.HasUpperBound() ? SectionRange.GetUpperBoundValue() : FFrameNumber(TNumericLimits<int32>::Max());
 	}
 };
 
+/**
+ * Component data for camera shakes created by the shake system
+ * This is separate from FMovieSceneCameraShakeComponentData because that
+ * one it imported from source shake sections, and our component data here
+ * will be preserved on reimported entities.
+ */
+USTRUCT()
+struct FMovieSceneCameraShakeInstanceData
+{
+	GENERATED_BODY()
+
+	/** Shake instance created by the shake evaluation system */
+	UPROPERTY()
+	TObjectPtr<UCameraShakeBase> ShakeInstance;
+
+	/** The signature of the source section at the time the shake instance was created */
+	UPROPERTY()
+	FGuid SectionSignature;
+};
+
 /**
  * Component data for audio tracks inputs
  * This provides the names of the inputs whose values are stored on the
@@ -617,6 +638,7 @@ struct FMovieSceneTracksComponentTypes
 	TComponentTypeID<FName> AudioTriggerName;
 
 	TComponentTypeID<FMovieSceneCameraShakeComponentData> CameraShake;
+	TComponentTypeID<FMovieSceneCameraShakeInstanceData> CameraShakeInstance;
 
 	struct
 	{
