Sequencer: better handle incorrect/unresolved camera bindings in the camera cut track

diff --git a/Engine/Source/Runtime/MovieSceneTracks/Private/TrackInstances/MovieSceneCameraCutEditorHandler.cpp b/Engine/Source/Runtime/MovieSceneTracks/Private/TrackInstances/MovieSceneCameraCutEditorHandler.cpp
--- a/Engine/Source/Runtime/MovieSceneTracks/Private/TrackInstances/MovieSceneCameraCutEditorHandler.cpp
+++ b/Engine/Source/Runtime/MovieSceneTracks/Private/TrackInstances/MovieSceneCameraCutEditorHandler.cpp
@@ -266,16 +266,16 @@ void FCameraCutEditorHandler::SetCameraCutForViewport(
 	float ViewFOV = ViewportClient.ViewFOV;
 	bool bCameraHasBeenCut = CameraCutParams.bJumpCut;
 
-	TSharedPtr<FPreAnimatedCameraCutEditorStorage> PreAnimatedStorage = Linker->PreAnimatedState.FindStorage(
+	TSharedPtr<FPreAnimatedCameraCutEditorStorage> PreAnimatedStorage = Linker->PreAnimatedState.GetOrCreateStorage(
 			FPreAnimatedCameraCutEditorStorage::StorageID);
 
 	// If that viewport wasn't locked to cinematics, see if we need to re-cache pre-animated state.
 	// This is necessary if the user released control with the camera button on the camera cut track.
 	// (see ReleaseCameraCutForViewport)
 	//
-	// Note that PreAnimatedStorage can be null here if we never cached any pre-animated state, which
+	// Note that storage index can be invalid here if we never cached any pre-animated state, which
 	// is possible if the option to restore viewports on unlock is off.
-	if (PreAnimatedStorage.IsValid() && !ViewportClient.IsLockedToCinematic())
+	if (!ViewportClient.IsLockedToCinematic())
 	{
 		FPreAnimatedStorageIndex StorageIndex = PreAnimatedStorage->FindStorageIndex(&ViewportClient);
 		if (!StorageIndex.IsValid())
@@ -301,7 +301,7 @@ void FCameraCutEditorHandler::SetCameraCutForViewport(
 		// inside pre-animated state if we are *blending* out of the cinematic. In this case, let's
 		// restore the pre-animated viewport location ourselves and enable blending.
 		FPreAnimatedStorageIndex StorageIndex = PreAnimatedStorage->FindStorageIndex(&ViewportClient);
-		if (ensureMsgf(StorageIndex.IsValid(), TEXT("Blending camera back to editor but can't find pre-animated viewport info!")))
+		if (StorageIndex.IsValid())
 		{
 			FPreAnimatedCameraCutEditorState CachedValue = PreAnimatedStorage->GetCachedValue(StorageIndex);
 			ViewLocation = CachedValue.ViewportLocation;
diff --git a/Engine/Source/Runtime/MovieSceneTracks/Private/TrackInstances/MovieSceneCameraCutTrackInstance.cpp b/Engine/Source/Runtime/MovieSceneTracks/Private/TrackInstances/MovieSceneCameraCutTrackInstance.cpp
--- a/Engine/Source/Runtime/MovieSceneTracks/Private/TrackInstances/MovieSceneCameraCutTrackInstance.cpp
+++ b/Engine/Source/Runtime/MovieSceneTracks/Private/TrackInstances/MovieSceneCameraCutTrackInstance.cpp
@@ -228,6 +228,11 @@ public:
 		const FMovieSceneContext& Context = SequenceInstance.GetContext();
 
 		UObject* CameraActor = FindBoundObject(Params.CameraBindingID, Params.OperandSequenceID, Player);
+		if (Params.CameraBindingID.IsValid() && CameraActor == nullptr)
+		{
+			// We have an unresolved or incorrect binding.
+			return false;
+		}
 
 		FMovieSceneCameraCutParams CameraCutParams;
 		CameraCutParams.bJumpCut = Context.HasJumped();
