Sequencer- Update object property track to use ECS system, and remove old object track template.
Set custom accessors for USkeletalMeshComponent::SkeletalMeshAsset and USkinnedMeshComponent::SkeletalMesh (this was why it was necessary to update the track to ECS, so we could use a custom getter).
Finally, add back the #if WITH_EDITORONLY_DATA to USkeletalMeshComponent::SkeletalMeshAsset now that we can, as the ECS custom accessors don't require a real property to be present in order to work.

diff --git a/Engine/Source/Runtime/Engine/Classes/Components/SkeletalMeshComponent.h b/Engine/Source/Runtime/Engine/Classes/Components/SkeletalMeshComponent.h
--- a/Engine/Source/Runtime/Engine/Classes/Components/SkeletalMeshComponent.h
+++ b/Engine/Source/Runtime/Engine/Classes/Components/SkeletalMeshComponent.h
@@ -316,13 +316,13 @@ class USkeletalMeshComponent : public USkinnedMeshComponent, public IInterface_C
 	friend struct FLinkedInstancesAdapter;
 	friend struct FLinkedAnimLayerClassData;
 	
-//#if WITH_EDITORONLY_DATA  // TODO: Re-add these guards once the MovieScene getters/setters are working, so that we can get rid of this redundant pointer in all cooked builds
+#if WITH_EDITORONLY_DATA 
 private:
 	/** The skeletal mesh used by this component. */
 	UE_DEPRECATED(5.1, "This property isn't deprecated, but getter and setter must be used at all times to preserve correct operations.")
 	UPROPERTY(EditAnywhere, Transient, Setter = SetSkeletalMeshAsset, BlueprintSetter = SetSkeletalMeshAsset, Getter = GetSkeletalMeshAsset, BlueprintGetter = GetSkeletalMeshAsset, Category = Mesh)
 	TObjectPtr<USkeletalMesh> SkeletalMeshAsset;
-//#endif
+#endif
 
 public:
 	/**
diff --git a/Engine/Source/Runtime/Engine/Private/Components/SkeletalMeshComponent.cpp b/Engine/Source/Runtime/Engine/Private/Components/SkeletalMeshComponent.cpp
--- a/Engine/Source/Runtime/Engine/Private/Components/SkeletalMeshComponent.cpp
+++ b/Engine/Source/Runtime/Engine/Private/Components/SkeletalMeshComponent.cpp
@@ -393,12 +393,12 @@ void USkeletalMeshComponent::PostLoad()
 		bOverrideMinLod = true;
 	}
 
-//#if WITH_EDITORONLY_DATA  // TODO: Re-add these guards once the MovieScene getters/setters are working, so that we can get rid of this redundant pointer in all cooked builds
+#if WITH_EDITORONLY_DATA
 	// Update property alias
 	PRAGMA_DISABLE_DEPRECATION_WARNINGS
 	SkeletalMeshAsset = GetSkeletalMeshAsset();
 	PRAGMA_ENABLE_DEPRECATION_WARNINGS
-//#endif
+#endif
 }
 
 void USkeletalMeshComponent::PostInitProperties()
@@ -2894,11 +2894,11 @@ void USkeletalMeshComponent::SetSkeletalMesh(USkeletalMesh* InSkelMesh, bool bRe
 	RemoveAllClothingActors();
 
 	// Update property alias
-//#if WITH_EDITORONLY_DATA  // TODO: Re-add these guards once the MovieScene getters/setters are working, so that we can get rid of this redundant pointer in all cooked builds
+#if WITH_EDITORONLY_DATA
 	PRAGMA_DISABLE_DEPRECATION_WARNINGS
 	SkeletalMeshAsset = InSkelMesh;
 	PRAGMA_ENABLE_DEPRECATION_WARNINGS
-//#endif
+#endif
 
 	// We may be doing parallel evaluation on the current anim instance
 	// Calling this here with true will block this init till that thread completes
diff --git a/Engine/Source/Runtime/MovieScene/Private/EntitySystem/MovieScenePropertyRegistry.cpp b/Engine/Source/Runtime/MovieScene/Private/EntitySystem/MovieScenePropertyRegistry.cpp
--- a/Engine/Source/Runtime/MovieScene/Private/EntitySystem/MovieScenePropertyRegistry.cpp
+++ b/Engine/Source/Runtime/MovieScene/Private/EntitySystem/MovieScenePropertyRegistry.cpp
@@ -120,6 +120,12 @@ TOptional<uint16> ComputeFastPropertyPtrOffset(UClass* ObjectClass, const FMovie
 		return TOptional<uint16>();
 	}
 
+	// For object properties, we can't use the fast path, as reference tracking may be necessary
+	if (FObjectPropertyBase* ObjectProperty = CastField<FObjectPropertyBase>(PropertyAndAddress.Property))
+	{
+		return TOptional<uint16>();
+	}
+
 	// @todo: Constructing FNames from strings is _very_ costly and we really shouldn't be doing this at runtime.
 	//        This is a little better now we use a string builder and an FNAME_Find, but it's still not ideal
 	TStringBuilder<128> SetterName;
diff --git a/Engine/Source/Runtime/MovieSceneTracks/Private/Evaluation/MovieSceneObjectPropertyTemplate.cpp b/Engine/Source/Runtime/MovieSceneTracks/Private/Evaluation/MovieSceneObjectPropertyTemplate.cpp
deleted file mode 100644
diff --git a/Engine/Source/Runtime/MovieSceneTracks/Private/Evaluation/MovieSceneObjectPropertyTemplate.h b/Engine/Source/Runtime/MovieSceneTracks/Private/Evaluation/MovieSceneObjectPropertyTemplate.h
deleted file mode 100644
diff --git a/Engine/Source/Runtime/MovieSceneTracks/Private/MovieSceneTracksComponentTypes.cpp b/Engine/Source/Runtime/MovieSceneTracks/Private/MovieSceneTracksComponentTypes.cpp
--- a/Engine/Source/Runtime/MovieSceneTracks/Private/MovieSceneTracksComponentTypes.cpp
+++ b/Engine/Source/Runtime/MovieSceneTracks/Private/MovieSceneTracksComponentTypes.cpp
@@ -169,6 +169,16 @@ void ConvertOperationalProperty(double In, float& Out)
 	Out = static_cast<float>(In);
 }
 
+void ConvertOperationalProperty(const FObjectComponent& In, UObject*& Out)
+{
+	Out = In.GetObject();
+}
+
+void ConvertOperationalProperty(UObject* In, FObjectComponent& Out)
+{
+	Out = FObjectComponent::Strong(In);
+}
+
 uint8 GetSkeletalMeshAnimationMode(const UObject* Object)
 {
 	const USkeletalMeshComponent* SkeletalMeshComponent = CastChecked<const USkeletalMeshComponent>(Object);
@@ -498,6 +508,30 @@ struct FComponentTransformHandler : TPropertyComponentHandler<FComponentTransfor
 	}
 };
 
+struct FObjectHandler : TPropertyComponentHandler<FObjectPropertyTraits, FObjectComponent>
+{
+	virtual void DispatchInitializePropertyMetaDataTasks(const FPropertyDefinition& Definition, FSystemTaskPrerequisites& InPrerequisites, FSystemSubsequentTasks& Subsequents, UMovieSceneEntitySystemLinker* Linker) override
+	{
+		FBuiltInComponentTypes* BuiltInComponents = FBuiltInComponentTypes::Get();
+		FMovieSceneTracksComponentTypes* TrackComponents = FMovieSceneTracksComponentTypes::Get();
+
+		FEntityTaskBuilder()
+			.Read(BuiltInComponents->BoundObject)
+			.Read(BuiltInComponents->PropertyBinding)
+			.Write(TrackComponents->Object.MetaDataComponents.GetType<0>())
+			.FilterAll({ BuiltInComponents->Tags.NeedsLink })
+			.Iterate_PerEntity(&Linker->EntityManager, [](UObject* Object, const FMovieScenePropertyBinding& Binding, FObjectPropertyTraits::FObjectMetadata& OutMetaData)
+				{
+					FObjectPropertyBase* BoundProperty = CastField<FObjectPropertyBase>(FTrackInstancePropertyBindings::FindProperty(Object, Binding.PropertyPath.ToString()));
+					if (ensure(BoundProperty))
+					{
+						OutMetaData.ObjectClass = BoundProperty->PropertyClass;
+						OutMetaData.bAllowsClear = !BoundProperty->HasAnyPropertyFlags(CPF_NoClear);
+					}
+				});
+	}
+};
+
 FMovieSceneTracksComponentTypes::FMovieSceneTracksComponentTypes()
 {
 	FComponentRegistry* ComponentRegistry = UMovieSceneEntitySystemLinker::GetComponents();
@@ -512,6 +546,7 @@ FMovieSceneTracksComponentTypes::FMovieSceneTracksComponentTypes()
 	ComponentRegistry->NewPropertyType(FloatVector, TEXT("float vector"));
 	ComponentRegistry->NewPropertyType(DoubleVector, TEXT("double vector"));
 	ComponentRegistry->NewPropertyType(String, TEXT("FString"));
+	ComponentRegistry->NewPropertyType(Object, TEXT("Object"));
 
 	ComponentRegistry->NewPropertyType(Transform, TEXT("FTransform"));
 	ComponentRegistry->NewPropertyType(EulerTransform, TEXT("FEulerTransform"));
@@ -524,6 +559,7 @@ FMovieSceneTracksComponentTypes::FMovieSceneTracksComponentTypes()
 	Color.MetaDataComponents.Initialize(ComponentRegistry, TEXT("Color Type"));
 	FloatVector.MetaDataComponents.Initialize(ComponentRegistry, TEXT("Num Float Vector Channels"));
 	DoubleVector.MetaDataComponents.Initialize(ComponentRegistry, TEXT("Num Double Vector Channels"));
+	Object.MetaDataComponents.Initialize(ComponentRegistry, TEXT("Object Class"));
 
 	ComponentRegistry->NewComponentType(&QuaternionRotationChannel[0], TEXT("Quaternion Rotation Channel 0"));
 	ComponentRegistry->NewComponentType(&QuaternionRotationChannel[1], TEXT("Quaternion Rotation Channel 1"));
@@ -660,6 +696,13 @@ FMovieSceneTracksComponentTypes::FMovieSceneTracksComponentTypes()
 	.AddSoleChannel(BuiltInComponents->StringResult)
 	.Commit();
 
+	// --------------------------------------------------------------------------------------------
+	// Set up Object properties
+	BuiltInComponents->PropertyRegistry.DefineProperty(Object, TEXT("Apply Object Properties"))
+	.AddSoleChannel(BuiltInComponents->ObjectResult)
+	.SetCustomAccessors(&Accessors.Object)
+	.Commit(FObjectHandler());
+
 	// --------------------------------------------------------------------------------------------
 	// Set up float parameters
 	BuiltInComponents->PropertyRegistry.DefineProperty(FloatParameter, TEXT("Apply Float Parameters"))
diff --git a/Engine/Source/Runtime/MovieSceneTracks/Private/MovieSceneTracksCustomAccessors.cpp b/Engine/Source/Runtime/MovieSceneTracks/Private/MovieSceneTracksCustomAccessors.cpp
--- a/Engine/Source/Runtime/MovieSceneTracks/Private/MovieSceneTracksCustomAccessors.cpp
+++ b/Engine/Source/Runtime/MovieSceneTracks/Private/MovieSceneTracksCustomAccessors.cpp
@@ -7,6 +7,8 @@
 #include "Components/PrimitiveComponent.h"
 #include "Components/SkyAtmosphereComponent.h"
 #include "Components/SkyLightComponent.h"
+#include "Components/SkeletalMeshComponent.h"
+#include "Engine/SkeletalMesh.h"
 
 namespace UE::MovieScene
 {
@@ -402,6 +404,34 @@ void SetSecondFogDataFogHeightOffset(UObject* Object, float InFogHeightOffset)
 	ExponentialHeightFogComponent->SetSecondFogHeightOffset(InFogHeightOffset);
 }
 
+FObjectComponent GetSkeletalMeshAsset(const UObject* Object, FObjectPropertyTraits::FObjectMetadata ObjectMetadata)
+{
+	const USkeletalMeshComponent* SkeletalMeshComponent = CastChecked<const USkeletalMeshComponent>(Object);
+	return FObjectComponent::Strong(SkeletalMeshComponent->GetSkeletalMeshAsset());
+}
+
+void SetSkeletalMeshAsset(UObject* Object, FObjectPropertyTraits::FObjectMetadata ObjectMetadata, const FObjectComponent& InSkeletalMeshAsset)
+{
+	USkeletalMeshComponent* SkeletalMeshComponent = CastChecked<USkeletalMeshComponent>(Object);
+	SkeletalMeshComponent->SetSkeletalMeshAsset(Cast<USkeletalMesh>(InSkeletalMeshAsset.GetObject()));
+}
+
+FObjectComponent GetSkeletalMesh_DEPRECATED(const UObject* Object, FObjectPropertyTraits::FObjectMetadata ObjectMetadata)
+{
+	const USkinnedMeshComponent* SkinnedMeshComponent = CastChecked<const USkinnedMeshComponent>(Object);
+	PRAGMA_DISABLE_DEPRECATION_WARNINGS
+	return FObjectComponent::Strong(SkinnedMeshComponent->GetSkeletalMesh_DEPRECATED());
+	PRAGMA_ENABLE_DEPRECATION_WARNINGS
+}
+
+void SetSkeletalMesh_DEPRECATED(UObject* Object, FObjectPropertyTraits::FObjectMetadata ObjectMetadata, const FObjectComponent& InSkeletalMeshAsset)
+{
+	USkinnedMeshComponent* SkinnedMeshComponent = CastChecked<USkinnedMeshComponent>(Object);
+	PRAGMA_DISABLE_DEPRECATION_WARNINGS
+	SkinnedMeshComponent->SetSkeletalMesh_DEPRECATED(Cast<USkeletalMesh>(InSkeletalMeshAsset.GetObject()));
+	PRAGMA_ENABLE_DEPRECATION_WARNINGS
+}
+
 void InitializeMovieSceneTracksAccessors(FMovieSceneTracksComponentTypes* TracksComponents)
 {
 	// We have some custom accessors for well-known types.
@@ -492,6 +522,14 @@ void InitializeMovieSceneTracksAccessors(FMovieSceneTracksComponentTypes* Tracks
 			GetSecondFogDataFogHeightOffset, SetSecondFogDataFogHeightOffset);
 
 	TracksComponents->Accessors.ComponentTransform.Add(USceneComponent::StaticClass(), "Transform", &GetComponentTransform, &SetComponentTransformAndVelocity);
+
+	// SkeletalMeshComponent
+	TracksComponents->Accessors.Object.Add(USkeletalMeshComponent::StaticClass(), "SkeletalMeshAsset", &GetSkeletalMeshAsset, &SetSkeletalMeshAsset);
+
+	//SkinnedMeshComponent
+	TracksComponents->Accessors.Object.Add(USkinnedMeshComponent::StaticClass(), "SkeletalMesh", &GetSkeletalMesh_DEPRECATED, &SetSkeletalMesh_DEPRECATED);
+
+	
 }
 
 } // namespace UE::MovieScene
\ No newline at end of file
diff --git a/Engine/Source/Runtime/MovieSceneTracks/Private/Sections/MovieSceneObjectPropertySection.cpp b/Engine/Source/Runtime/MovieSceneTracks/Private/Sections/MovieSceneObjectPropertySection.cpp
--- a/Engine/Source/Runtime/MovieSceneTracks/Private/Sections/MovieSceneObjectPropertySection.cpp
+++ b/Engine/Source/Runtime/MovieSceneTracks/Private/Sections/MovieSceneObjectPropertySection.cpp
@@ -44,4 +44,22 @@ void UMovieSceneObjectPropertySection::PostLoad()
 	}
 
 	Super::PostLoad();
+}
+
+bool UMovieSceneObjectPropertySection::PopulateEvaluationFieldImpl(const TRange<FFrameNumber>& EffectiveRange, const FMovieSceneEvaluationFieldEntityMetaData& InMetaData, FMovieSceneEntityComponentFieldBuilder* OutFieldBuilder)
+{
+	FMovieScenePropertyTrackEntityImportHelper::PopulateEvaluationField(*this, EffectiveRange, InMetaData, OutFieldBuilder);
+	return true;
+}
+
+void UMovieSceneObjectPropertySection::ImportEntityImpl(UMovieSceneEntitySystemLinker* EntityLinker, const FEntityImportParams& Params, FImportedEntity* OutImportedEntity)
+{
+	using namespace UE::MovieScene;
+
+	const FBuiltInComponentTypes* Components = FBuiltInComponentTypes::Get();
+	const FMovieSceneTracksComponentTypes* TracksComponents = FMovieSceneTracksComponentTypes::Get();
+
+	FPropertyTrackEntityImportHelper(TracksComponents->Object)
+		.Add(Components->ObjectPathChannel, &ObjectChannel)
+		.Commit(this, Params, OutImportedEntity);
 }
\ No newline at end of file
diff --git a/Engine/Source/Runtime/MovieSceneTracks/Private/Systems/MovieSceneObjectPropertySystem.cpp b/Engine/Source/Runtime/MovieSceneTracks/Private/Systems/MovieSceneObjectPropertySystem.cpp
new file mode 100644
--- /dev/null
+++ b/Engine/Source/Runtime/MovieSceneTracks/Private/Systems/MovieSceneObjectPropertySystem.cpp
@@ -0,0 +1,27 @@
+// Copyright Epic Games, Inc. All Rights Reserved.
+
+#include "Systems/MovieSceneObjectPropertySystem.h"
+#include "Systems/ObjectPathChannelEvaluatorSystem.h"
+#include "MovieSceneTracksComponentTypes.h"
+
+#include UE_INLINE_GENERATED_CPP_BY_NAME(MovieSceneObjectPropertySystem)
+
+UMovieSceneObjectPropertySystem::UMovieSceneObjectPropertySystem(const FObjectInitializer& ObjInit)
+	: Super(ObjInit)
+{
+	BindToProperty(UE::MovieScene::FMovieSceneTracksComponentTypes::Get()->Object);
+
+	if (HasAnyFlags(RF_ClassDefaultObject))
+	{
+		DefineImplicitPrerequisite(UObjectPathChannelEvaluatorSystem::StaticClass(), GetClass());
+
+		DefineComponentConsumer(GetClass(), UE::MovieScene::FMovieSceneTracksComponentTypes::Get()->Object.PropertyTag);
+	}
+}
+
+void UMovieSceneObjectPropertySystem::OnRun(FSystemTaskPrerequisites& InPrerequisites, FSystemSubsequentTasks& Subsequents)
+{
+	Super::OnRun(InPrerequisites, Subsequents);
+}
+
+
diff --git a/Engine/Source/Runtime/MovieSceneTracks/Private/Systems/ObjectPathChannelEvaluatorSystem.cpp b/Engine/Source/Runtime/MovieSceneTracks/Private/Systems/ObjectPathChannelEvaluatorSystem.cpp
--- a/Engine/Source/Runtime/MovieSceneTracks/Private/Systems/ObjectPathChannelEvaluatorSystem.cpp
+++ b/Engine/Source/Runtime/MovieSceneTracks/Private/Systems/ObjectPathChannelEvaluatorSystem.cpp
@@ -39,6 +39,7 @@ UObjectPathChannelEvaluatorSystem::UObjectPathChannelEvaluatorSystem(const FObje
 
 	FBuiltInComponentTypes* BuiltInComponents = FBuiltInComponentTypes::Get();
 	RelevantComponent = BuiltInComponents->ObjectPathChannel;
+	Phase = ESystemPhase::Scheduling;
 
 	if (HasAnyFlags(RF_ClassDefaultObject))
 	{
diff --git a/Engine/Source/Runtime/MovieSceneTracks/Private/Tracks/MovieSceneObjectPropertyTrack.cpp b/Engine/Source/Runtime/MovieSceneTracks/Private/Tracks/MovieSceneObjectPropertyTrack.cpp
--- a/Engine/Source/Runtime/MovieSceneTracks/Private/Tracks/MovieSceneObjectPropertyTrack.cpp
+++ b/Engine/Source/Runtime/MovieSceneTracks/Private/Tracks/MovieSceneObjectPropertyTrack.cpp
@@ -2,7 +2,6 @@
 
 #include "Tracks/MovieSceneObjectPropertyTrack.h"
 #include "Sections/MovieSceneObjectPropertySection.h"
-#include "Evaluation/MovieSceneObjectPropertyTemplate.h"
 
 #include UE_INLINE_GENERATED_CPP_BY_NAME(MovieSceneObjectPropertyTrack)
 
@@ -25,9 +24,3 @@ UMovieSceneSection* UMovieSceneObjectPropertyTrack::CreateNewSection()
 	Section->ObjectChannel.SetPropertyClass(PropertyClass);
 	return Section;
 }
-
-FMovieSceneEvalTemplatePtr UMovieSceneObjectPropertyTrack::CreateTemplateForSection(const UMovieSceneSection& InSection) const
-{
-	return FMovieSceneObjectPropertyTemplate(*CastChecked<UMovieSceneObjectPropertySection>(&InSection), *this);
-}
-
diff --git a/Engine/Source/Runtime/MovieSceneTracks/Public/MovieSceneTracksComponentTypes.h b/Engine/Source/Runtime/MovieSceneTracks/Public/MovieSceneTracksComponentTypes.h
--- a/Engine/Source/Runtime/MovieSceneTracks/Public/MovieSceneTracksComponentTypes.h
+++ b/Engine/Source/Runtime/MovieSceneTracks/Public/MovieSceneTracksComponentTypes.h
@@ -570,6 +570,80 @@ struct FBoolPropertyTraits
 	}
 };
 
+struct FObjectPropertyTraits
+{
+	static constexpr bool bIsComposite = false;
+
+	struct FObjectMetadata
+	{
+		TObjectPtr<UClass> ObjectClass = UObject::StaticClass();
+		bool bAllowsClear = true;
+	};
+
+	using StorageType = FObjectComponent;
+	//using CustomAccessorStorageType = float;
+	using MetaDataType = TPropertyMetaData<FObjectMetadata>;
+	using TraitsType = FObjectPropertyTraits;
+	using ParamType = UObject*;
+
+	using ObjectTraitsImpl = TIndirectPropertyTraits<UObject*, FObjectComponent>;
+
+	static void GetObjectPropertyValue(const UObject* InObject, FObjectMetadata ObjectMetadata, const FCustomPropertyAccessor& BaseCustomAccessor, FObjectComponent& OutValue)
+	{
+		const TCustomPropertyAccessor<FObjectPropertyTraits>& CustomAccessor = static_cast<const TCustomPropertyAccessor<FObjectPropertyTraits>&>(BaseCustomAccessor);
+		OutValue = (*CustomAccessor.Functions.Getter)(InObject, ObjectMetadata);
+	}
+	static void GetObjectPropertyValue(const UObject* InObject, FObjectMetadata ObjectMetadata, uint16 PropertyOffset, FObjectComponent& OutValue)
+	{
+		ObjectTraitsImpl::GetObjectPropertyValue(InObject, PropertyOffset, OutValue);
+	}
+	static void GetObjectPropertyValue(const UObject* InObject, FObjectMetadata ObjectMetadata, FTrackInstancePropertyBindings* PropertyBindings, FObjectComponent& OutValue)
+	{
+		ObjectTraitsImpl::GetObjectPropertyValue(InObject, PropertyBindings, OutValue);
+	}
+	static void GetObjectPropertyValue(const UObject* InObject, FObjectMetadata ObjectMetadata, const FName& PropertyPath, StorageType& OutValue)
+	{
+		ObjectTraitsImpl::GetObjectPropertyValue(InObject, PropertyPath, OutValue);
+	}
+
+	static bool CanAssignValue(const FObjectMetadata& ObjectMetadata, UObject* DesiredValue)
+	{
+		if (!ObjectMetadata.ObjectClass)
+		{
+			return false;
+		}
+		else if (!DesiredValue)
+		{
+			return ObjectMetadata.bAllowsClear;
+		}
+		else if (DesiredValue->GetClass() != nullptr)
+		{
+			return DesiredValue->GetClass()->IsChildOf(ObjectMetadata.ObjectClass);
+		}
+		return false;
+	}
+
+	static void SetObjectPropertyValue(UObject* InObject, FObjectMetadata ObjectMetadata, const FCustomPropertyAccessor& BaseCustomAccessor, const FObjectComponent& InValue)
+	{
+		const TCustomPropertyAccessor<FObjectPropertyTraits>& CustomAccessor = static_cast<const TCustomPropertyAccessor<FObjectPropertyTraits>&>(BaseCustomAccessor);
+		(*CustomAccessor.Functions.Setter)(InObject, ObjectMetadata, InValue);
+	}
+	static void SetObjectPropertyValue(UObject* InObject, FObjectMetadata ObjectMetadata, uint16 PropertyOffset, const FObjectComponent& InValue)
+	{
+		if (CanAssignValue(ObjectMetadata, InValue.GetObject()))
+		{
+			ObjectTraitsImpl::SetObjectPropertyValue(InObject, PropertyOffset, InValue);
+		}
+	}
+	static void SetObjectPropertyValue(UObject* InObject, FObjectMetadata ObjectMetadata, FTrackInstancePropertyBindings* PropertyBindings, const FObjectComponent& InValue)
+	{
+		if (CanAssignValue(ObjectMetadata, InValue.GetObject()))
+		{
+			ObjectTraitsImpl::SetObjectPropertyValue(InObject, PropertyBindings, InValue);
+		}
+	}
+};
+
 using FBytePropertyTraits               = TDirectPropertyTraits<uint8, false>;
 using FEnumPropertyTraits               = TDirectPropertyTraits<uint8, false>;
 using FIntPropertyTraits                = TDirectPropertyTraits<int32, false>;
@@ -599,6 +673,7 @@ struct FMovieSceneTracksComponentTypes
 	TPropertyComponents<FEulerTransformPropertyTraits> EulerTransform;
 	TPropertyComponents<FComponentTransformPropertyTraits> ComponentTransform;
 	TPropertyComponents<FStringPropertyTraits> String;
+	TPropertyComponents<FObjectPropertyTraits> Object;
 
 	TPropertyComponents<FFloatParameterTraits> FloatParameter;
 	TPropertyComponents<FColorParameterTraits> ColorParameter;
@@ -652,6 +727,7 @@ struct FMovieSceneTracksComponentTypes
 		TCustomPropertyRegistration<FFloatVectorPropertyTraits> FloatVector;
 		TCustomPropertyRegistration<FDoubleVectorPropertyTraits> DoubleVector;
 		TCustomPropertyRegistration<FComponentTransformPropertyTraits, 1> ComponentTransform;
+		TCustomPropertyRegistration<FObjectPropertyTraits> Object;
 	} Accessors;
 
 	struct
diff --git a/Engine/Source/Runtime/MovieSceneTracks/Public/MovieSceneTracksPropertyTypes.h b/Engine/Source/Runtime/MovieSceneTracks/Public/MovieSceneTracksPropertyTypes.h
--- a/Engine/Source/Runtime/MovieSceneTracks/Public/MovieSceneTracksPropertyTypes.h
+++ b/Engine/Source/Runtime/MovieSceneTracks/Public/MovieSceneTracksPropertyTypes.h
@@ -14,12 +14,13 @@
 
 class USceneComponent;
 
-
 namespace UE
 {
 namespace MovieScene
 {
 
+struct FObjectComponent;
+
 struct FVectorPropertyMetaData
 {
 	uint8 NumChannels = 0;
@@ -222,5 +223,8 @@ MOVIESCENETRACKS_API void ConvertOperationalProperty(const FVector2d& In, FDoubl
 MOVIESCENETRACKS_API void ConvertOperationalProperty(const FVector3d& In, FDoubleIntermediateVector& Out);
 MOVIESCENETRACKS_API void ConvertOperationalProperty(const FVector4d& In, FDoubleIntermediateVector& Out);
 
+MOVIESCENETRACKS_API void ConvertOperationalProperty(const FObjectComponent& In, UObject*& Out);
+MOVIESCENETRACKS_API void ConvertOperationalProperty(UObject* In, FObjectComponent& Out);
+
 } // namespace MovieScene
 } // namespace UE
diff --git a/Engine/Source/Runtime/MovieSceneTracks/Public/Sections/MovieSceneObjectPropertySection.h b/Engine/Source/Runtime/MovieSceneTracks/Public/Sections/MovieSceneObjectPropertySection.h
--- a/Engine/Source/Runtime/MovieSceneTracks/Public/Sections/MovieSceneObjectPropertySection.h
+++ b/Engine/Source/Runtime/MovieSceneTracks/Public/Sections/MovieSceneObjectPropertySection.h
@@ -6,13 +6,16 @@
 #include "MovieSceneSection.h"
 #include "UObject/ObjectMacros.h"
 #include "UObject/UObjectGlobals.h"
+#include "EntitySystem/IMovieSceneEntityProvider.h"
 
 #include "MovieSceneObjectPropertySection.generated.h"
 
 class UObject;
 
 UCLASS(MinimalAPI)
-class UMovieSceneObjectPropertySection : public UMovieSceneSection
+class UMovieSceneObjectPropertySection 
+	: public UMovieSceneSection
+	, public IMovieSceneEntityProvider
 {
 public:
 
@@ -27,4 +30,9 @@ public:
 
 	UPROPERTY()
 	FMovieSceneObjectPathChannel ObjectChannel;
+
+private:
+	virtual void ImportEntityImpl(UMovieSceneEntitySystemLinker* EntityLinker, const FEntityImportParams& Params, FImportedEntity* OutImportedEntity) override;
+	virtual bool PopulateEvaluationFieldImpl(const TRange<FFrameNumber>& EffectiveRange, const FMovieSceneEvaluationFieldEntityMetaData& InMetaData, FMovieSceneEntityComponentFieldBuilder* OutFieldBuilder) override;
+
 };
diff --git a/Engine/Source/Runtime/MovieSceneTracks/Public/Systems/MovieSceneObjectPropertySystem.h b/Engine/Source/Runtime/MovieSceneTracks/Public/Systems/MovieSceneObjectPropertySystem.h
new file mode 100644
--- /dev/null
+++ b/Engine/Source/Runtime/MovieSceneTracks/Public/Systems/MovieSceneObjectPropertySystem.h
@@ -0,0 +1,16 @@
+// Copyright Epic Games, Inc. All Rights Reserved.
+
+#pragma once
+
+#include "Systems/MovieScenePropertySystem.h"
+#include "MovieSceneObjectPropertySystem.generated.h"
+
+UCLASS(MinimalAPI)
+class UMovieSceneObjectPropertySystem : public UMovieScenePropertySystem
+{
+	GENERATED_BODY()
+
+	UMovieSceneObjectPropertySystem(const FObjectInitializer& ObjInit);
+
+	MOVIESCENETRACKS_API virtual void OnRun(FSystemTaskPrerequisites& InPrerequisites, FSystemSubsequentTasks& Subsequents) override;
+};
diff --git a/Engine/Source/Runtime/MovieSceneTracks/Public/Tracks/MovieSceneObjectPropertyTrack.h b/Engine/Source/Runtime/MovieSceneTracks/Public/Tracks/MovieSceneObjectPropertyTrack.h
--- a/Engine/Source/Runtime/MovieSceneTracks/Public/Tracks/MovieSceneObjectPropertyTrack.h
+++ b/Engine/Source/Runtime/MovieSceneTracks/Public/Tracks/MovieSceneObjectPropertyTrack.h
@@ -4,11 +4,10 @@
 
 #include "MovieSceneTrack.h"
 #include "Tracks/MovieScenePropertyTrack.h"
-#include "Compilation/IMovieSceneTrackTemplateProducer.h"
 #include "MovieSceneObjectPropertyTrack.generated.h"
 
 UCLASS(MinimalAPI)
-class UMovieSceneObjectPropertyTrack : public UMovieScenePropertyTrack, public IMovieSceneTrackTemplateProducer
+class UMovieSceneObjectPropertyTrack : public UMovieScenePropertyTrack
 {
 public:
 
@@ -27,5 +26,4 @@ public:
 	/*~ UMovieSceneTrack interface */
 	virtual bool SupportsType(TSubclassOf<UMovieSceneSection> SectionClass) const override;
 	virtual UMovieSceneSection* CreateNewSection() override;
-	virtual FMovieSceneEvalTemplatePtr CreateTemplateForSection(const UMovieSceneSection& InSection) const override;
 };
