Sequencer: add logging to shake previewer

#noqatesting

diff --git a/Engine/Source/Runtime/MovieSceneTracks/Private/Evaluation/MovieSceneCameraShakePreviewer.cpp b/Engine/Source/Runtime/MovieSceneTracks/Private/Evaluation/MovieSceneCameraShakePreviewer.cpp
--- a/Engine/Source/Runtime/MovieSceneTracks/Private/Evaluation/MovieSceneCameraShakePreviewer.cpp
+++ b/Engine/Source/Runtime/MovieSceneTracks/Private/Evaluation/MovieSceneCameraShakePreviewer.cpp
@@ -4,10 +4,11 @@
 
 #if WITH_EDITOR
 
-#include "LevelEditorViewport.h"
+#include "Camera/CameraModifier_CameraShake.h"
 #include "Camera/CameraShakeBase.h"
 #include "Camera/CameraShakeSourceComponent.h"
-#include "Camera/CameraModifier_CameraShake.h"
+#include "LevelEditorViewport.h"
+#include "MovieSceneFwd.h"
 
 FCameraShakePreviewer::FCameraShakePreviewer(UWorld* InWorld)
 	: World(InWorld)
@@ -35,6 +36,8 @@ UCameraShakeBase* FCameraShakePreviewer::AddCameraShake(const FCameraShakePrevie
 	UCameraShakeBase* NewShake = NewObject<UCameraShakeBase>(World, Params.ShakeClass);
 	ActiveShakes.Add({ NewShake, Params.SourceComponent, Params.GlobalStartTime });
 
+	UE_LOG(LogMovieScene, Verbose, TEXT("CameraShakePreviewer::AddCameraShake '%s'"), *GetNameSafe(NewShake));
+
 	FCameraShakeBaseStartParams StartParams;
 	StartParams.Scale = Params.Scale;
 	StartParams.PlaySpace = Params.PlaySpace;
@@ -53,6 +56,8 @@ void FCameraShakePreviewer::RemoveCameraShake(UCameraShakeBase* ShakeInstance)
 		FPreviewCameraShakeInfo& ActiveShake = ActiveShakes[i];
 		if (ActiveShake.ShakeInstance == ShakeInstance)
 		{
+			UE_LOG(LogMovieScene, Verbose, TEXT("CameraShakePreviewer::RemoveCameraShake '%s'"), *GetNameSafe(ActiveShake.ShakeInstance));
+
 			ActiveShake.ShakeInstance->StopShake(bImmediately);
 			ActiveShake.ShakeInstance->TeardownShake();
 			ActiveShakes.RemoveAt(i, 1);
@@ -63,6 +68,7 @@ void FCameraShakePreviewer::RemoveCameraShake(UCameraShakeBase* ShakeInstance)
 
 void FCameraShakePreviewer::RemoveAllCameraShakesFromSource(const UCameraShakeSourceComponent* SourceComponent)
 {
+	int32 NumRemoved = 0;
 	const bool bImmediately = true;
 	for (int32 i = ActiveShakes.Num() - 1; i >= 0; --i)
 	{
@@ -72,8 +78,11 @@ void FCameraShakePreviewer::RemoveAllCameraShakesFromSource(const UCameraShakeSo
 			ActiveShake.ShakeInstance->StopShake(bImmediately);
 			ActiveShake.ShakeInstance->TeardownShake();
 			ActiveShakes.RemoveAt(i, 1);
+			++NumRemoved;
 		}
 	}
+	UE_LOG(LogMovieScene, Verbose, TEXT("CameraShakePreviewer::RemoveAllCameraShakesFromSource '%s', %d removed"), 
+			*GetNameSafe(SourceComponent), NumRemoved);
 }
 
 void FCameraShakePreviewer::RemoveAllCameraShakes()
@@ -87,6 +96,7 @@ void FCameraShakePreviewer::RemoveAllCameraShakes()
 			ActiveShake.ShakeInstance->TeardownShake();
 		}
 	}
+	UE_LOG(LogMovieScene, Verbose, TEXT("CameraShakePreviewer::RemoveAllCameraShakes, %d removed"), ActiveShakes.Num());
 	ActiveShakes.Empty();
 }
 
@@ -199,6 +209,8 @@ void FCameraShakePreviewer::OnModifyView(FEditorViewportViewModifierParams& Para
 					ShakeInfo.ShakeInstance->TeardownShake();
 				}
 
+				UE_LOG(LogMovieScene, Verbose, TEXT("CameraShakePreviewer::OnModifyView, '%s' finished or stale"), 
+						*GetNameSafe(ShakeInfo.ShakeInstance));
 				ActiveShakes.RemoveAt(i, 1);
 			}
 		}
@@ -252,6 +264,7 @@ void FCameraShakePreviewer::RegisterViewModifier(FLevelEditorViewportClient* Vie
 		return;
 	}
 
+	UE_LOG(LogMovieScene, Verbose, TEXT("CameraShakePreviewer::RegisterViewModifier"));
 	RegisteredViewportClients.Add(ViewportClient);
 	ViewportClient->ViewModifiers.AddRaw(this, &FCameraShakePreviewer::OnModifyView);
 	if (RegisteredViewportClients.Num() == 1)
@@ -275,6 +288,8 @@ void FCameraShakePreviewer::UnRegisterViewModifiers()
 	{
 		ViewportClient->ViewModifiers.RemoveAll(this);
 	}
+	UE_LOG(LogMovieScene, Verbose, TEXT("CameraShakePreviewer::UnRegisterViewModifiers, %d unregistered"), 
+			RegisteredViewportClients.Num());
 	RegisteredViewportClients.Reset();
 }
 
@@ -283,6 +298,8 @@ void FCameraShakePreviewer::UnRegisterViewModifier(FLevelEditorViewportClient* V
 	const int32 NumRemoved = RegisteredViewportClients.Remove(ViewportClient);
 	if (ensureMsgf(NumRemoved > 0, TEXT("The given viewport client wasn't registered.")))
 	{
+		UE_LOG(LogMovieScene, Verbose, TEXT("CameraShakePreviewer::UnRegisterViewModifier"));
+
 		// If this is the last viewport, stop listening to viewports changing.
 		if (RegisteredViewportClients.IsEmpty())
 		{
