Sequencer: fix visibility section creating keys with the wrong value

diff --git a/Engine/Source/Runtime/MovieSceneTracks/Private/Sections/MovieSceneVisibilitySection.cpp b/Engine/Source/Runtime/MovieSceneTracks/Private/Sections/MovieSceneVisibilitySection.cpp
--- a/Engine/Source/Runtime/MovieSceneTracks/Private/Sections/MovieSceneVisibilitySection.cpp
+++ b/Engine/Source/Runtime/MovieSceneTracks/Private/Sections/MovieSceneVisibilitySection.cpp
@@ -32,6 +32,10 @@ UMovieSceneVisibilitySection::UMovieSceneVisibilitySection(const FObjectInitiali
 	: Super(Init)
 {
 	BoolCurve.SetDefault(true);
+
+#if WITH_EDITORONLY_DATA
+	SetIsExternallyInverted(true);
+#endif
 }
 
 void UMovieSceneVisibilitySection::ImportEntityImpl(UMovieSceneEntitySystemLinker* EntityLinker, const FEntityImportParams& Params, FImportedEntity* OutImportedEntity)
