Sequencer: handle fully-overlapped sections when arranging camera cuts

diff --git a/Engine/Source/Runtime/MovieSceneTracks/Private/Tracks/MovieSceneCameraCutTrack.cpp b/Engine/Source/Runtime/MovieSceneTracks/Private/Tracks/MovieSceneCameraCutTrack.cpp
--- a/Engine/Source/Runtime/MovieSceneTracks/Private/Tracks/MovieSceneCameraCutTrack.cpp
+++ b/Engine/Source/Runtime/MovieSceneTracks/Private/Tracks/MovieSceneCameraCutTrack.cpp
@@ -325,7 +325,16 @@ void UMovieSceneCameraCutTrack::RearrangeAllSections()
 		FFrameNumber OverlapOrGap(0);
 		if (PrevSectionRange.HasUpperBound() && CurSectionRange.HasLowerBound())
 		{
-			OverlapOrGap = (PrevSectionRange.GetUpperBoundValue() - CurSectionRange.GetLowerBoundValue());
+			// We need to handle the case of CurSection being completely overlapped by PrevSection.
+			// In that case, PrevSection's upper bound is way past CurSection's upper bound, and if it's
+			// far enough away, the mid-point that we find later (see MeetupFrame) can be past the end
+			// of CurSection, resulting in setting an invalid range. So let's clamp this.
+			FFrameNumber UpperBound = PrevSectionRange.GetUpperBoundValue();
+			if (CurSectionRange.HasUpperBound())
+			{
+				UpperBound = FMath::Min(UpperBound, CurSectionRange.GetUpperBoundValue());
+			}
+			OverlapOrGap = (UpperBound - CurSectionRange.GetLowerBoundValue());
 		}
 
 		// If there's an overlap and we don't want it, resize the sections so that they start/stop
