Sequencer: fix problem with re-imported camera shake entities ending up with unstarted shakes

If the pre-animated state restore happened before the shake instantiator system, any ongoing shakes would be stopped. But then if any of those shakes were on re-imported entities, we would skip restarting them, and the shakes would be stuck in an incorrect unstarted state.

diff --git a/Engine/Source/Runtime/MovieSceneTracks/Private/Systems/MovieSceneCameraShakeSystem.cpp b/Engine/Source/Runtime/MovieSceneTracks/Private/Systems/MovieSceneCameraShakeSystem.cpp
--- a/Engine/Source/Runtime/MovieSceneTracks/Private/Systems/MovieSceneCameraShakeSystem.cpp
+++ b/Engine/Source/Runtime/MovieSceneTracks/Private/Systems/MovieSceneCameraShakeSystem.cpp
@@ -8,10 +8,11 @@
 #include "EntitySystem/BuiltInComponentTypes.h"
 #include "EntitySystem/MovieSceneEntityManager.h"
 #include "EntitySystem/MovieSceneEntitySystemLinker.h"
+#include "EntitySystem/MovieScenePreAnimatedStateSystem.h"
+#include "Evaluation/MovieSceneCameraShakePreviewer.h"
+#include "Evaluation/PreAnimatedState/MovieScenePreAnimatedObjectStorage.h"
 #include "Evaluation/PreAnimatedState/MovieScenePreAnimatedStateStorage.h"
 #include "Evaluation/PreAnimatedState/MovieScenePreAnimatedStorageID.inl"
-#include "Evaluation/PreAnimatedState/MovieScenePreAnimatedObjectStorage.h"
-#include "Evaluation/MovieSceneCameraShakePreviewer.h"
 #include "IMovieScenePlayer.h"
 #include "MovieSceneTracksComponentTypes.h"
 #include "Sections/MovieSceneCameraShakeSection.h"
@@ -272,6 +273,10 @@ UMovieSceneCameraShakeInstantiatorSystem::UMovieSceneCameraShakeInstantiatorSyst
 	{
 		const FBuiltInComponentTypes* BuiltInComponents = FBuiltInComponentTypes::Get();
 		DefineComponentConsumer(GetClass(), BuiltInComponents->BoundObject);
+
+		// Make sure our shakes aren't stopped/restored before we get a chance to transfer
+		// them to a re-imported entity.
+		DefineImplicitPrerequisite(GetClass(), UMovieSceneRestorePreAnimatedStateSystem::StaticClass());
 	}
 }
 
