Sequencer: Check for valid GEditor

#rnx

diff --git a/Engine/Source/Runtime/MovieSceneTracks/Private/Systems/MovieSceneCameraShakeSystem.cpp b/Engine/Source/Runtime/MovieSceneTracks/Private/Systems/MovieSceneCameraShakeSystem.cpp
--- a/Engine/Source/Runtime/MovieSceneTracks/Private/Systems/MovieSceneCameraShakeSystem.cpp
+++ b/Engine/Source/Runtime/MovieSceneTracks/Private/Systems/MovieSceneCameraShakeSystem.cpp
@@ -152,8 +152,11 @@ FCameraShakePreviewerLinkerExtension::FCameraShakePreviewerLinkerExtension(UMovi
 
 FCameraShakePreviewerLinkerExtension::~FCameraShakePreviewerLinkerExtension()
 {
-	GEditor->OnLevelViewportClientListChanged().RemoveAll(this);
-
+	if (GEditor != nullptr)
+	{
+		GEditor->OnLevelViewportClientListChanged().RemoveAll(this);
+	}
+	
 	for (TPair<FInstanceHandle, FCameraShakePreviewer>& Pair : Previewers)
 	{
 		FCameraShakePreviewer& Previewer = Pair.Value;
@@ -174,7 +177,7 @@ FCameraShakePreviewer& FCameraShakePreviewerLinkerExtension::GetPreviewer(FInsta
 		return *Previewer;
 	}
 
-	if (Previewers.IsEmpty())
+	if (Previewers.IsEmpty() && GEditor != nullptr)
 	{
 		// This is our first previewer... let's start listening to viewports changing.
 		GEditor->OnLevelViewportClientListChanged().AddSP(this, &FCameraShakePreviewerLinkerExtension::OnLevelViewportClientListChanged);
