Sequencer: NLA: Fix match with next not working with blend options.

diff --git a/Engine/Source/Runtime/MovieSceneTracks/Private/Tracks/MovieSceneSkeletalAnimationTrack.cpp b/Engine/Source/Runtime/MovieSceneTracks/Private/Tracks/MovieSceneSkeletalAnimationTrack.cpp
--- a/Engine/Source/Runtime/MovieSceneTracks/Private/Tracks/MovieSceneSkeletalAnimationTrack.cpp
+++ b/Engine/Source/Runtime/MovieSceneTracks/Private/Tracks/MovieSceneSkeletalAnimationTrack.cpp
@@ -1196,7 +1196,7 @@ void UMovieSceneSkeletalAnimationTrack::MatchSectionByBoneTransform(bool bMatchW
 			{
 				SecondTransformRotation.Roll = FirstTransformRotation.Roll;
 			}
-			SecondTransformQuat = QuatFromEuler(FirstTransformRotation.Euler(), RotationOrder);
+			SecondTransformQuat = QuatFromEuler(SecondTransformRotation.Euler(), RotationOrder);
 			SecondTransform.SetRotation(SecondTransformQuat);
 
 			//GetRelativeTransformReverse returns this(-1)* Other, and parameter is Other.
