Sequencer: NLA: Split Sections will inherit root motion matched location.

diff --git a/Engine/Source/Runtime/MovieSceneTracks/Private/Sections/MovieSceneSkeletalAnimationSection.cpp b/Engine/Source/Runtime/MovieSceneTracks/Private/Sections/MovieSceneSkeletalAnimationSection.cpp
--- a/Engine/Source/Runtime/MovieSceneTracks/Private/Sections/MovieSceneSkeletalAnimationSection.cpp
+++ b/Engine/Source/Runtime/MovieSceneTracks/Private/Sections/MovieSceneSkeletalAnimationSection.cpp
@@ -304,22 +304,63 @@ void UMovieSceneSkeletalAnimationSection::TrimSection(FQualifiedFrameTime TrimTi
 
 UMovieSceneSection* UMovieSceneSkeletalAnimationSection::SplitSection(FQualifiedFrameTime SplitTime, bool bDeleteKeys)
 {
+	//handle root motion, only in editor
+#if WITH_EDITOR
+
+	TOptional<FTransform> RootTransform;
+	FName BoneName;
+	UMovieSceneSkeletalAnimationTrack* Track = GetTypedOuter<UMovieSceneSkeletalAnimationTrack>();
+	//if we are doing root motion then get the BoneName 
+	if (Track)
+	{
+		int32 BoneIndex = SetBoneIndexForRootMotionCalculations(Track->bBlendFirstChildOfRoot);
+		FMovieSceneSkeletalAnimRootMotionTrackParams* RootMotionParams = GetRootMotionParams();
+		if (RootMotionParams->bHaveRootMotion)
+		{
+			if (UAnimSequence* AnimSequence = Cast<UAnimSequence>(Params.Animation))
+			{
+				const FReferenceSkeleton& RefSkeleton = AnimSequence->GetSkeleton()->GetReferenceSkeleton();
+				RootTransform = Track->GetRootMotion(SplitTime.Time);
+				if (RootTransform.IsSet())
+				{
+					BoneName = RefSkeleton.GetBoneName(BoneIndex);
+				}
+			}
+		}
+	}
+#endif
+
+	//handle FirstLoopStartFrameOffset
 	const FFrameNumber InitialFirstLoopStartFrameOffset = Params.FirstLoopStartFrameOffset;
 
 	FFrameRate FrameRate = GetTypedOuter<UMovieScene>()->GetTickResolution();
 
 	const FFrameNumber NewOffset = HasStartFrame() ? GetFirstLoopStartOffsetAtTrimTime(SplitTime, Params, GetInclusiveStartFrame(), FrameRate) : 0;
 
-	UMovieSceneSection* NewSection = Super::SplitSection(SplitTime, bDeleteKeys);
+	UMovieSceneSkeletalAnimationSection* NewSection = Cast<UMovieSceneSkeletalAnimationSection>(Super::SplitSection(SplitTime, bDeleteKeys));
 	if (NewSection != nullptr)
 	{
-		UMovieSceneSkeletalAnimationSection* NewSkeletalSection = Cast<UMovieSceneSkeletalAnimationSection>(NewSection);
-		NewSkeletalSection->Params.FirstLoopStartFrameOffset = NewOffset;
+		NewSection->Params.FirstLoopStartFrameOffset = NewOffset;
+
+#if WITH_EDITOR
+
+		if (RootTransform.IsSet())
+		{
+			NewSection->bMatchTranslation = NewSection->bMatchIncludeZHeight = true;
+			NewSection->bMatchRotationYaw = NewSection->bMatchRotationPitch = NewSection->bMatchRotationRoll = true;
+			NewSection->bMatchWithPrevious = true;
+			NewSection->MatchedLocationOffset = FVector(0.0, 0.0, 0.0);
+			NewSection->MatchedRotationOffset = FRotator(0.0, 0.0, 0.0);
+			if (NewSection->GetRootMotionParams())
+			{
+				NewSection->GetRootMotionParams()->bRootMotionsDirty = true;
+			}
+		}
+#endif
 	}
 
 	// Restore original offset modified by splitting
 	Params.FirstLoopStartFrameOffset = InitialFirstLoopStartFrameOffset;
-
 	return NewSection;
 }
 
@@ -620,7 +661,6 @@ FTransform UMovieSceneSkeletalAnimationSection::GetRootMotionStartOffset() const
 	return FTransform::Identity;;
 }
 
-
 bool UMovieSceneSkeletalAnimationSection::GetRootMotionTransform(FAnimationPoseData& AnimationPoseData,FRootMotionTransformParam& InOutParams) const
 {
 	UAnimSequence* AnimSequence = Cast<UAnimSequence>(Params.Animation);
diff --git a/Engine/Source/Runtime/MovieSceneTracks/Private/Tracks/MovieSceneSkeletalAnimationTrack.cpp b/Engine/Source/Runtime/MovieSceneTracks/Private/Tracks/MovieSceneSkeletalAnimationTrack.cpp
--- a/Engine/Source/Runtime/MovieSceneTracks/Private/Tracks/MovieSceneSkeletalAnimationTrack.cpp
+++ b/Engine/Source/Runtime/MovieSceneTracks/Private/Tracks/MovieSceneSkeletalAnimationTrack.cpp
@@ -611,8 +611,7 @@ void UMovieSceneSkeletalAnimationTrack::SetUpRootMotions(bool bForce)
 				{
 					if (UAnimSequenceBase* PrevAnimSequence = PrevAnimSection->Params.Animation)
 					{	
-						if (AnimSection->StartLocationOffset.IsNearlyZero() == false || AnimSection->StartRotationOffset.IsNearlyZero() == false ||
-							AnimSection->MatchedLocationOffset.IsNearlyZero() == false || AnimSection->MatchedRotationOffset.IsNearlyZero() == false)
+						if (bAnySectionsHaveOffset)
 						{ 
 							RootMotionParams.RootMotionStartOffset = AnimSection->GetRootMotionStartOffset();
 							FMemMark Mark(FMemStack::Get());
