MovieSceneTracksCustomAccessors: Add accessor for DirectionalInscatteringColor and FogInscatteringColor so that the render state of the exponential heightfog component is dirtied when driving those properties.

diff --git a/Engine/Source/Runtime/MovieSceneTracks/Private/MovieSceneTracksCustomAccessors.cpp b/Engine/Source/Runtime/MovieSceneTracks/Private/MovieSceneTracksCustomAccessors.cpp
--- a/Engine/Source/Runtime/MovieSceneTracks/Private/MovieSceneTracksCustomAccessors.cpp
+++ b/Engine/Source/Runtime/MovieSceneTracks/Private/MovieSceneTracksCustomAccessors.cpp
@@ -347,6 +347,38 @@ void SetExponentialHeightFogComponentDirectionalInscatteringExponent(UObject* Ob
 	UExponentialHeightFogComponent* ExponentialHeightFogComponent = CastChecked<UExponentialHeightFogComponent>(Object);
 	ExponentialHeightFogComponent->SetDirectionalInscatteringExponent(InDirectionalInscatteringExponent);
 }
+	
+FIntermediateColor GetExponentialHeightFogComponentDirectionalInscatteringColor(const UObject* Object, EColorPropertyType InColorType)
+{
+	ensure(InColorType == EColorPropertyType::Linear);
+	
+	const UExponentialHeightFogComponent* ExponentialHeightFogComponent = CastChecked<const UExponentialHeightFogComponent>(Object);
+	return FIntermediateColor(ExponentialHeightFogComponent->DirectionalInscatteringLuminance);
+}
+
+void SetExponentialHeightFogComponentDirectionalInscatteringColor(UObject* Object, EColorPropertyType InColorType, const FIntermediateColor& InColor)
+{
+	ensure(InColorType == EColorPropertyType::Linear);
+	
+	UExponentialHeightFogComponent* ExponentialHeightFogComponent = CastChecked<UExponentialHeightFogComponent>(Object);
+	ExponentialHeightFogComponent->SetDirectionalInscatteringColor(InColor.GetLinearColor());
+}
+
+FIntermediateColor GetExponentialHeightFogComponentFogInscatteringColor(const UObject* Object, EColorPropertyType InColorType)
+{
+	ensure(InColorType == EColorPropertyType::Linear);
+	
+	const UExponentialHeightFogComponent* ExponentialHeightFogComponent = CastChecked<const UExponentialHeightFogComponent>(Object);
+	return FIntermediateColor(ExponentialHeightFogComponent->FogInscatteringLuminance);
+}
+
+void SetExponentialHeightFogComponentFogInscatteringColor(UObject* Object, EColorPropertyType InColorType, const FIntermediateColor& InColor)
+{
+	ensure(InColorType == EColorPropertyType::Linear);
+	
+	UExponentialHeightFogComponent* ExponentialHeightFogComponent = CastChecked<UExponentialHeightFogComponent>(Object);
+	ExponentialHeightFogComponent->SetFogInscatteringColor(InColor.GetLinearColor());
+}
 
 FIntermediateColor GetExponentialHeightFogVolumetricFogAlbedo(const UObject* Object, EColorPropertyType InColorType)
 {
@@ -502,6 +534,12 @@ void InitializeMovieSceneTracksAccessors(FMovieSceneTracksComponentTypes* Tracks
 	TracksComponents->Accessors.Float.Add(
 			UExponentialHeightFogComponent::StaticClass(), GET_MEMBER_NAME_CHECKED(UExponentialHeightFogComponent, DirectionalInscatteringExponent),
 			GetExponentialHeightFogComponentDirectionalInscatteringExponent, SetExponentialHeightFogComponentDirectionalInscatteringExponent);
+	TracksComponents->Accessors.Color.Add(
+			UExponentialHeightFogComponent::StaticClass(), GET_MEMBER_NAME_CHECKED(UExponentialHeightFogComponent, DirectionalInscatteringLuminance),
+			GetExponentialHeightFogComponentDirectionalInscatteringColor, SetExponentialHeightFogComponentDirectionalInscatteringColor);
+	TracksComponents->Accessors.Color.Add(
+			UExponentialHeightFogComponent::StaticClass(), GET_MEMBER_NAME_CHECKED(UExponentialHeightFogComponent, FogInscatteringLuminance),
+			GetExponentialHeightFogComponentFogInscatteringColor, SetExponentialHeightFogComponentFogInscatteringColor);
 	TracksComponents->Accessors.Float.Add(
 			UExponentialHeightFogComponent::StaticClass(), GET_MEMBER_NAME_CHECKED(UExponentialHeightFogComponent, FogMaxOpacity),
 			GetExponentialHeightFogComponentFogMaxOpacity, SetExponentialHeightFogComponentFogMaxOpacity);
