Speculative fix for Sequencer crash in UMovieSceneAudioSystem::ResetSharedData().

diff --git a/Engine/Source/Runtime/MovieSceneTracks/Private/Systems/MovieSceneAudioSystem.cpp b/Engine/Source/Runtime/MovieSceneTracks/Private/Systems/MovieSceneAudioSystem.cpp
--- a/Engine/Source/Runtime/MovieSceneTracks/Private/Systems/MovieSceneAudioSystem.cpp
+++ b/Engine/Source/Runtime/MovieSceneTracks/Private/Systems/MovieSceneAudioSystem.cpp
@@ -810,6 +810,7 @@ void UMovieSceneAudioSystem::OnSchedulePersistentTasks(UE::MovieScene::IEntitySy
 
 	TaskScheduler->AddPrerequisite(GatherInputsTask, EvaluateAudioTask);
 	TaskScheduler->AddPrerequisite(GatherTriggersTask, EvaluateAudioTask);
+	TaskScheduler->AddPrerequisite(ResetSharedDataTask, EvaluateAudioTask);
 }
 
 void UMovieSceneAudioSystem::OnRun(FSystemTaskPrerequisites& InPrerequisites, FSystemSubsequentTasks& Subsequents)
