Sequencer: Restore state for ComponentMaterialParameterSection

#rnx

diff --git a/Engine/Source/Runtime/MovieSceneTracks/Private/Sections/MovieSceneComponentMaterialParameterSection.cpp b/Engine/Source/Runtime/MovieSceneTracks/Private/Sections/MovieSceneComponentMaterialParameterSection.cpp
--- a/Engine/Source/Runtime/MovieSceneTracks/Private/Sections/MovieSceneComponentMaterialParameterSection.cpp
+++ b/Engine/Source/Runtime/MovieSceneTracks/Private/Sections/MovieSceneComponentMaterialParameterSection.cpp
@@ -47,6 +47,7 @@ UMovieSceneComponentMaterialParameterSection::UMovieSceneComponentMaterialParame
 	: Super(ObjectInitializer)
 {
 	bSupportsInfiniteRange = true;
+	EvalOptions.EnableAndSetCompletionMode(EMovieSceneCompletionMode::ProjectDefault);
 }
 
 EMovieSceneChannelProxyType UMovieSceneComponentMaterialParameterSection::CacheChannelProxy()
