Sequencer: Remove Modify() on PostLoad() of upgrading avisiblity track

#rnx

diff --git a/Engine/Source/Runtime/MovieSceneTracks/Private/Tracks/MovieSceneVisibilityTrack.cpp b/Engine/Source/Runtime/MovieSceneTracks/Private/Tracks/MovieSceneVisibilityTrack.cpp
--- a/Engine/Source/Runtime/MovieSceneTracks/Private/Tracks/MovieSceneVisibilityTrack.cpp
+++ b/Engine/Source/Runtime/MovieSceneTracks/Private/Tracks/MovieSceneVisibilityTrack.cpp
@@ -20,7 +20,6 @@ void UMovieSceneVisibilityTrack::PostLoad()
 {
 	// Upgrade bool sections to visibility sections.
 	TArray<uint8> Bytes;
-	bool bUpgraded = false;
 
 	for (int32 Index = 0; Index < Sections.Num(); ++Index)
 	{
@@ -39,15 +38,9 @@ void UMovieSceneVisibilityTrack::PostLoad()
 			FObjectReader(NewSection, Bytes);
 
 			Sections[Index] = NewSection;
-			bUpgraded = true;
 		}
 	}
 
-	if (bUpgraded)
-	{
-		Modify();
-	}
-
 	Super::PostLoad();
 }
 
