Sequencer: better handle camera shakes being recompiled

We need to correctly replace shake instance pointers saved in ECS components, and restart them with proper parameters.

diff --git a/Engine/Source/Runtime/MovieSceneTracks/Private/Systems/MovieSceneCameraShakeSystem.cpp b/Engine/Source/Runtime/MovieSceneTracks/Private/Systems/MovieSceneCameraShakeSystem.cpp
--- a/Engine/Source/Runtime/MovieSceneTracks/Private/Systems/MovieSceneCameraShakeSystem.cpp
+++ b/Engine/Source/Runtime/MovieSceneTracks/Private/Systems/MovieSceneCameraShakeSystem.cpp
@@ -25,10 +25,39 @@
 #endif  // WITH_EDITOR
 
 #include UE_INLINE_GENERATED_CPP_BY_NAME(MovieSceneCameraShakeSystem)
-
+PRAGMA_DISABLE_OPTIMIZATION
 namespace UE::MovieScene
 {
 
+TOptional<float> ComputeCameraShakeDurationOverride(const FMovieSceneContext& Context, const FMovieSceneCameraShakeComponentData& ShakeData)
+{
+	// Get the duration of the shake and compare it to the duration of the section, to know
+	// if we need to override it. We only override duration if the section is shorter than
+	// the shake's natural duration... i.e, we only make shakes shorter, not longer.
+	const FFrameRate& FrameRate = Context.GetFrameRate();
+	TSubclassOf<UCameraShakeBase> ShakeClass = ShakeData.SectionData.ShakeClass;
+
+	FCameraShakeDuration ShakeDuration;
+	UCameraShakeBase::GetCameraShakeDuration(ShakeClass, ShakeDuration);
+
+	TOptional<float> DurationOverride;
+	const FFrameTime SectionDurationFrames = (ShakeData.SectionEndTime - ShakeData.SectionStartTime);
+	if (ShakeDuration.IsFixed())
+	{
+		const FFrameTime ShakeDurationFrames = Context.GetFrameRate().AsFrameTime(ShakeDuration.Get());
+		if (ShakeDurationFrames > SectionDurationFrames)
+		{
+			DurationOverride = FrameRate.AsSeconds(SectionDurationFrames);
+		}
+	}
+	else
+	{
+		DurationOverride = FrameRate.AsSeconds(SectionDurationFrames);
+	}
+
+	return DurationOverride;
+}
+
 struct FPreAnimatedCameraShakeTraits : FBoundObjectPreAnimatedStateTraits
 {
 	using KeyType = FObjectKey;
@@ -189,7 +218,7 @@ FCameraShakePreviewer& FCameraShakePreviewerLinkerExtension::GetPreviewer(FInsta
 	UObject* PlaybackContext = SequenceInstance.GetPlayer()->GetPlaybackContext();
 	UWorld* ContextWorld = PlaybackContext ? PlaybackContext->GetWorld() : nullptr;
 
-	FCameraShakePreviewer& NewPreviewer = Previewers.Add(InstanceHandle, ContextWorld);
+	FCameraShakePreviewer& NewPreviewer = Previewers.Emplace(InstanceHandle, ContextWorld);
 	NewPreviewer.RegisterViewModifiers([ContextWorld](FLevelEditorViewportClient* LevelVC) -> bool
 				{
 					return LevelVC->AllowsCinematicControl() && LevelVC->GetWorld() == ContextWorld;
@@ -302,6 +331,9 @@ void UMovieSceneCameraShakeInstantiatorSystem::OnLink()
 		// We only need the previewer extension if there's an actual editor.
 		PreviewerExtension = FCameraShakePreviewerLinkerExtension::GetOrCreateExtension(Linker);
 	}
+
+	// Handle camera shakes being recompiled.
+	FCoreUObjectDelegates::OnObjectsReplaced.AddUObject(this, &UMovieSceneCameraShakeInstantiatorSystem::OnObjectsReplaced);
 #endif  // WITH_EDITOR
 }
 
@@ -310,6 +342,8 @@ void UMovieSceneCameraShakeInstantiatorSystem::OnUnlink()
 	using namespace UE::MovieScene;
 
 #if WITH_EDITOR
+	FCoreUObjectDelegates::OnObjectsReplaced.RemoveAll(this);
+
 	// Only the two camera shake systems hold pointers to the extension, so it should delete itself
 	// once both systems are unlinked.
 	PreviewerExtension = nullptr;
@@ -358,27 +392,7 @@ void UMovieSceneCameraShakeInstantiatorSystem::OnRun(FSystemTaskPrerequisites& I
 			return;
 		}
 
-		// Get the duration of the shake and compare it to the duration of the section, to know
-		// if we need to override it. We only override duration if the section is shorter than
-		// the shake's natural duration... i.e, we only make shakes shorter, not longer.
-		const FFrameRate& FrameRate = Context.GetFrameRate();
-		TOptional<float> DurationOverride;
-
-		FCameraShakeDuration ShakeDuration;
-		UCameraShakeBase::GetCameraShakeDuration(ShakeClass, ShakeDuration);
-		const FFrameTime SectionDurationFrames = (ShakeData.SectionEndTime - ShakeData.SectionStartTime);
-		if (ShakeDuration.IsFixed())
-		{
-			const FFrameTime ShakeDurationFrames = Context.GetFrameRate().AsFrameTime(ShakeDuration.Get());
-			if (ShakeDurationFrames > SectionDurationFrames)
-			{
-				DurationOverride = FrameRate.AsSeconds(SectionDurationFrames);
-			}
-		}
-		else
-		{
-			DurationOverride = FrameRate.AsSeconds(SectionDurationFrames);
-		}
+		TOptional<float> DurationOverride = ComputeCameraShakeDurationOverride(Context, ShakeData);
 
 		const bool bWantsRestoreState = EntityManager.HasComponent(EntityID, BuiltInComponents->Tags.RestoreState);
 		const FRootInstanceHandle RootInstanceHandle = Instance.GetRootInstanceHandle();
@@ -425,7 +439,7 @@ void UMovieSceneCameraShakeInstantiatorSystem::OnRun(FSystemTaskPrerequisites& I
 
 				FCameraShakePreviewerAddParams PreviewParams;
 				PreviewParams.ShakeClass = ShakeClass;
-				PreviewParams.GlobalStartTime = FrameRate.AsSeconds(ShakeData.SectionStartTime);
+				PreviewParams.GlobalStartTime = Context.GetFrameRate().AsSeconds(ShakeData.SectionStartTime);
 				PreviewParams.SourceComponent = ShakeSourceComponent;
 				PreviewParams.Scale = ShakeData.SectionData.PlayScale;
 				PreviewParams.PlaySpace = ShakeData.SectionData.PlaySpace;
@@ -440,6 +454,7 @@ void UMovieSceneCameraShakeInstantiatorSystem::OnRun(FSystemTaskPrerequisites& I
 				}
 
 				ShakeInstanceData.ShakeInstance = Previewer.AddCameraShake(PreviewParams);
+				ShakeInstanceData.bManagedByPreviewer = true;
 			}
 #endif  // WITH_EDITOR
 		}
@@ -468,6 +483,7 @@ void UMovieSceneCameraShakeInstantiatorSystem::OnRun(FSystemTaskPrerequisites& I
 
 			ShakeInstanceData.ShakeInstance = ShakeInstance;
 			ShakeInstanceData.SectionSignature = ShakeData.SectionSignature;
+			ShakeInstanceData.bManagedByPreviewer = false;
 
 			if (bStartInstance)
 			{
@@ -597,6 +613,50 @@ void UMovieSceneCameraShakeInstantiatorSystem::TriggerOneShotShakes()
 	TriggersByInstance.Empty();
 }
 
+#if WITH_EDITOR
+void UMovieSceneCameraShakeInstantiatorSystem::OnObjectsReplaced(const TMap<UObject*, UObject*>& ReplacementMap)
+{
+	using namespace UE::MovieScene;
+
+	const FBuiltInComponentTypes* BuiltInComponents = FBuiltInComponentTypes::Get();
+	const FMovieSceneTracksComponentTypes* TrackComponents = FMovieSceneTracksComponentTypes::Get();
+	const FInstanceRegistry* InstanceRegistry = Linker->GetInstanceRegistry();
+
+	auto ReplaceShakeInstances = [InstanceRegistry, &ReplacementMap](
+			FInstanceHandle InstanceHandle,
+			const FMovieSceneCameraShakeComponentData& ShakeData,
+			FMovieSceneCameraShakeInstanceData& ShakeInstanceData)
+	{
+		if (UObject* const * NewShakeInstance = ReplacementMap.Find(ShakeInstanceData.ShakeInstance))
+		{
+			ShakeInstanceData.ShakeInstance = Cast<UCameraShakeBase>(*NewShakeInstance);
+			
+			// Restart the new shake instance, unless it's already managed by a previewer which
+			// will do this already.
+			if (ensure(ShakeInstanceData.ShakeInstance) && !ShakeInstanceData.bManagedByPreviewer)
+			{
+				const FMovieSceneContext& Context = InstanceRegistry->GetInstance(InstanceHandle).GetContext();
+				TOptional<float> DurationOverride = ComputeCameraShakeDurationOverride(Context, ShakeData);
+
+				FCameraShakeBaseStartParams ShakeParams;
+				ShakeParams.Scale = ShakeData.SectionData.PlayScale;
+				ShakeParams.PlaySpace = ShakeData.SectionData.PlaySpace;
+				ShakeParams.UserPlaySpaceRot = ShakeData.SectionData.UserDefinedPlaySpace;
+				ShakeParams.DurationOverride = DurationOverride;
+
+				ShakeInstanceData.ShakeInstance->StartShake(ShakeParams);
+			}
+		}
+	};
+
+	FEntityTaskBuilder()
+		.Read(BuiltInComponents->InstanceHandle)
+		.Read(TrackComponents->CameraShake)
+		.Write(TrackComponents->CameraShakeInstance)
+		.Iterate_PerEntity(&Linker->EntityManager, ReplaceShakeInstances);
+}
+#endif
+
 namespace UE::MovieScene
 {
 
@@ -803,4 +863,4 @@ void UMovieSceneCameraShakeEvaluatorSystem::OnRun(FSystemTaskPrerequisites& InPr
 	}
 #endif  // WITH_EDITOR
 }
-
+PRAGMA_ENABLE_OPTIMIZATION
diff --git a/Engine/Source/Runtime/MovieSceneTracks/Public/Systems/MovieSceneCameraShakeSystem.h b/Engine/Source/Runtime/MovieSceneTracks/Public/Systems/MovieSceneCameraShakeSystem.h
--- a/Engine/Source/Runtime/MovieSceneTracks/Public/Systems/MovieSceneCameraShakeSystem.h
+++ b/Engine/Source/Runtime/MovieSceneTracks/Public/Systems/MovieSceneCameraShakeSystem.h
@@ -79,6 +79,10 @@ public:
 private:
 	void TriggerOneShotShakes();
 
+#if WITH_EDITOR
+	void OnObjectsReplaced(const TMap<UObject*, UObject*>& ReplacementMap);
+#endif
+
 private:
 	struct FTimedTrigger
 	{
