Sequencer: follow-up change because previous change missed some files.

diff --git a/Engine/Source/Runtime/MovieSceneTracks/Private/Evaluation/MovieSceneCameraShakePreviewer.cpp b/Engine/Source/Runtime/MovieSceneTracks/Private/Evaluation/MovieSceneCameraShakePreviewer.cpp
--- a/Engine/Source/Runtime/MovieSceneTracks/Private/Evaluation/MovieSceneCameraShakePreviewer.cpp
+++ b/Engine/Source/Runtime/MovieSceneTracks/Private/Evaluation/MovieSceneCameraShakePreviewer.cpp
@@ -5,7 +5,6 @@
 #if WITH_EDITOR
 
 #include "Camera/CameraModifier_CameraShake.h"
-#include "Camera/CameraShakeBase.h"
 #include "Camera/CameraShakeSourceComponent.h"
 #include "LevelEditorViewport.h"
 #include "MovieSceneFwd.h"
@@ -33,18 +32,19 @@ FCameraShakePreviewer::~FCameraShakePreviewer()
 
 UCameraShakeBase* FCameraShakePreviewer::AddCameraShake(const FCameraShakePreviewerAddParams& Params)
 {
-	UCameraShakeBase* NewShake = NewObject<UCameraShakeBase>(World, Params.ShakeClass);
-	ActiveShakes.Add({ NewShake, Params.SourceComponent, Params.GlobalStartTime });
-
-	UE_LOG(LogMovieScene, Verbose, TEXT("CameraShakePreviewer::AddCameraShake '%s'"), *GetNameSafe(NewShake));
-
 	FCameraShakeBaseStartParams StartParams;
 	StartParams.Scale = Params.Scale;
 	StartParams.PlaySpace = Params.PlaySpace;
 	StartParams.UserPlaySpaceRot = Params.UserPlaySpaceRot;
 	StartParams.DurationOverride = Params.DurationOverride;
+
+	UCameraShakeBase* NewShake = NewObject<UCameraShakeBase>(World, Params.ShakeClass);
 	NewShake->StartShake(StartParams);
 
+	ActiveShakes.Add({ StartParams, NewShake, Params.SourceComponent, Params.GlobalStartTime });
+
+	UE_LOG(LogMovieScene, Verbose, TEXT("CameraShakePreviewer::AddCameraShake '%s'"), *GetNameSafe(NewShake));
+
 	return NewShake;
 }
 
@@ -354,12 +354,18 @@ void FCameraShakePreviewer::OnObjectsReplaced(const TMap<UObject*, UObject*>& Re
 	for (int32 i = ActiveShakes.Num() - 1; i >= 0; i--)
 	{
 		FPreviewCameraShakeInfo& ActiveShake = ActiveShakes[i];
-		if (ReplacementMap.Find(ActiveShake.ShakeInstance))
+		if (UObject* const * NewShakeInstance = ReplacementMap.Find(ActiveShake.ShakeInstance))
 		{
-			// If a camera shake gets recompiled, we just stop and discard it.
+			// If a camera shake gets recompiled, stop the old version, and start the
+			// new version with the same parameters.
 			ActiveShake.ShakeInstance->StopShake(bImmediately);
 			ActiveShake.ShakeInstance->TeardownShake();
-			ActiveShakes.RemoveAt(i);
+			
+			ActiveShake.ShakeInstance = Cast<UCameraShakeBase>(*NewShakeInstance);
+			if (ensure(ActiveShake.ShakeInstance))
+			{
+				ActiveShake.ShakeInstance->StartShake(ActiveShake.StartParams);
+			}
 		}
 	}
 }
diff --git a/Engine/Source/Runtime/MovieSceneTracks/Private/Systems/MovieSceneCameraShakeSystem.cpp b/Engine/Source/Runtime/MovieSceneTracks/Private/Systems/MovieSceneCameraShakeSystem.cpp
--- a/Engine/Source/Runtime/MovieSceneTracks/Private/Systems/MovieSceneCameraShakeSystem.cpp
+++ b/Engine/Source/Runtime/MovieSceneTracks/Private/Systems/MovieSceneCameraShakeSystem.cpp
@@ -25,7 +25,7 @@
 #endif  // WITH_EDITOR
 
 #include UE_INLINE_GENERATED_CPP_BY_NAME(MovieSceneCameraShakeSystem)
-PRAGMA_DISABLE_OPTIMIZATION
+
 namespace UE::MovieScene
 {
 
@@ -863,4 +863,4 @@ void UMovieSceneCameraShakeEvaluatorSystem::OnRun(FSystemTaskPrerequisites& InPr
 	}
 #endif  // WITH_EDITOR
 }
-PRAGMA_ENABLE_OPTIMIZATION
+
diff --git a/Engine/Source/Runtime/MovieSceneTracks/Public/Evaluation/MovieSceneCameraShakePreviewer.h b/Engine/Source/Runtime/MovieSceneTracks/Public/Evaluation/MovieSceneCameraShakePreviewer.h
--- a/Engine/Source/Runtime/MovieSceneTracks/Public/Evaluation/MovieSceneCameraShakePreviewer.h
+++ b/Engine/Source/Runtime/MovieSceneTracks/Public/Evaluation/MovieSceneCameraShakePreviewer.h
@@ -4,12 +4,12 @@
 
 #if WITH_EDITOR
 
+#include "Camera/CameraShakeBase.h"
 #include "Camera/CameraTypes.h"
 #include "CoreMinimal.h"
 #include "UObject/GCObject.h"
 
 class FLevelEditorViewportClient;
-class UCameraShakeBase;
 class UCameraShakeSourceComponent;
 struct FActiveCameraShakeInfo;
 struct FEditorViewportViewModifierParams;
@@ -85,6 +85,7 @@ private:
 
 	struct FPreviewCameraShakeInfo
 	{
+		FCameraShakeBaseStartParams StartParams;
 		TObjectPtr<UCameraShakeBase> ShakeInstance;
 		TWeakObjectPtr<const UCameraShakeSourceComponent> SourceComponent;
 		float StartTime;
diff --git a/Engine/Source/Runtime/MovieSceneTracks/Public/MovieSceneTracksComponentTypes.h b/Engine/Source/Runtime/MovieSceneTracks/Public/MovieSceneTracksComponentTypes.h
--- a/Engine/Source/Runtime/MovieSceneTracks/Public/MovieSceneTracksComponentTypes.h
+++ b/Engine/Source/Runtime/MovieSceneTracks/Public/MovieSceneTracksComponentTypes.h
@@ -158,6 +158,10 @@ struct FMovieSceneCameraShakeInstanceData
 	/** The signature of the source section at the time the shake instance was created */
 	UPROPERTY()
 	FGuid SectionSignature;
+
+	/** Whether this instance is managed by a shake previewer */
+	UPROPERTY()
+	bool bManagedByPreviewer = false;
 };
 
 /**
