Fix crash when redoing redoing the creation of a Metasound Source on an Audio Track in Sequencer

diff --git a/Engine/Source/Runtime/MovieSceneTracks/Private/Sections/MovieSceneAudioSection.cpp b/Engine/Source/Runtime/MovieSceneTracks/Private/Sections/MovieSceneAudioSection.cpp
--- a/Engine/Source/Runtime/MovieSceneTracks/Private/Sections/MovieSceneAudioSection.cpp
+++ b/Engine/Source/Runtime/MovieSceneTracks/Private/Sections/MovieSceneAudioSection.cpp
@@ -603,14 +603,9 @@ void UMovieSceneAudioSection::ImportEntityImpl(UMovieSceneEntitySystemLinker* En
 
 void UMovieSceneAudioSection::PostEditChangeProperty(FPropertyChangedEvent& PropertyChangedEvent)
 {
-	const FName PropertyName = PropertyChangedEvent.GetPropertyName();
-	if (PropertyName == GET_MEMBER_NAME_CHECKED(UMovieSceneAudioSection, Sound))
-	{
-		// Invalidate channel proxy
-		ChannelProxy = nullptr;
-	}
-
 	Super::PostEditChangeProperty(PropertyChangedEvent);
+
+	CacheChannelProxy();
 }
 
 #endif
