Sequencer: Check for skel anim components being already post-evaluated before forcibly ticking for the current frame

This is a temporary fix while there are some playback APIs that cause inline evaluations of Sequencer. Ultimately we can remove this entire force-tick mechanism once Sequencer's playback API is entirely async.

diff --git a/Engine/Source/Runtime/MovieSceneTracks/Private/Systems/MovieSceneSkeletalAnimationSystem.cpp b/Engine/Source/Runtime/MovieSceneTracks/Private/Systems/MovieSceneSkeletalAnimationSystem.cpp
--- a/Engine/Source/Runtime/MovieSceneTracks/Private/Systems/MovieSceneSkeletalAnimationSystem.cpp
+++ b/Engine/Source/Runtime/MovieSceneTracks/Private/Systems/MovieSceneSkeletalAnimationSystem.cpp
@@ -577,7 +577,9 @@ private:
 		// If the skeletal component has already ticked this frame because tick prerequisites weren't set up yet or a new binding was created, forcibly tick this component to update.
 		// This resolves first frame issues where the skeletal component ticks first, then the sequencer binding is resolved which sets up tick prerequisites
 		// for the next frame.
-		if (SkeletalMeshComponent->PoseTickedThisFrame() || (SequencerInstance && SequencerInstance->GetSourceAnimInstance() != ExistingAnimInstance))
+		if (!SkeletalMeshComponent->IsPostEvaluatingAnimation() &&
+			(SkeletalMeshComponent->PoseTickedThisFrame() || (SequencerInstance && SequencerInstance->GetSourceAnimInstance() != ExistingAnimInstance))
+			)
 		{
 			SkeletalMeshComponent->TickAnimation(0.f, false);
 
