Sequencer: Fixed visibility track upgrade not re-populating the track's evaluation field with the new sections

This used to happen implicitly through a call to Modify() but that causes cook non-determinism, so we now just forcibly update the eval field since we know we can't have compiled data (yet) in PostLoad.

diff --git a/Engine/Source/Runtime/MovieScene/Private/MovieSceneTrack.cpp b/Engine/Source/Runtime/MovieScene/Private/MovieSceneTrack.cpp
--- a/Engine/Source/Runtime/MovieScene/Private/MovieSceneTrack.cpp
+++ b/Engine/Source/Runtime/MovieScene/Private/MovieSceneTrack.cpp
@@ -440,6 +440,11 @@ const FMovieSceneTrackEvaluationField& UMovieSceneTrack::GetEvaluationField()
 	return EvaluationField;
 }
 
+void UMovieSceneTrack::ForceUpdateEvaluationTree()
+{
+	UpdateEvaluationTree();
+}
+
 void UMovieSceneTrack::UpdateEvaluationTree()
 {
 	TMovieSceneEvaluationTree<FMovieSceneTrackEvaluationData> EvaluationTree;
diff --git a/Engine/Source/Runtime/MovieScene/Public/MovieSceneTrack.h b/Engine/Source/Runtime/MovieScene/Public/MovieSceneTrack.h
--- a/Engine/Source/Runtime/MovieScene/Public/MovieSceneTrack.h
+++ b/Engine/Source/Runtime/MovieScene/Public/MovieSceneTrack.h
@@ -272,6 +272,11 @@ protected:
 
 	ETreePopulationMode BuiltInTreePopulationMode;
 
+protected:
+
+	/** Forcibly update this evaluation tree without updating the signature. Does not invalidated any compiled data! */
+	MOVIESCENE_API void ForceUpdateEvaluationTree();
+
 private:
 
 	/** Sub-classes can override this method to perforum custom evaluation tree population logic. */
diff --git a/Engine/Source/Runtime/MovieSceneTracks/Private/Tracks/MovieSceneVisibilityTrack.cpp b/Engine/Source/Runtime/MovieSceneTracks/Private/Tracks/MovieSceneVisibilityTrack.cpp
--- a/Engine/Source/Runtime/MovieSceneTracks/Private/Tracks/MovieSceneVisibilityTrack.cpp
+++ b/Engine/Source/Runtime/MovieSceneTracks/Private/Tracks/MovieSceneVisibilityTrack.cpp
@@ -20,6 +20,7 @@ void UMovieSceneVisibilityTrack::PostLoad()
 {
 	// Upgrade bool sections to visibility sections.
 	TArray<uint8> Bytes;
+	bool bUpgraded = false;
 
 	for (int32 Index = 0; Index < Sections.Num(); ++Index)
 	{
@@ -38,9 +39,15 @@ void UMovieSceneVisibilityTrack::PostLoad()
 			FObjectReader(NewSection, Bytes);
 
 			Sections[Index] = NewSection;
+			bUpgraded = true;
 		}
 	}
 
+	if (bUpgraded)
+	{
+		ForceUpdateEvaluationTree();
+	}
+
 	Super::PostLoad();
 }
 
