Sequencer: VisibilityTrack upgrade needs to PreLoad sections on Serialize

diff --git a/Engine/Source/Runtime/MovieSceneTracks/Private/Tracks/MovieSceneVisibilityTrack.cpp b/Engine/Source/Runtime/MovieSceneTracks/Private/Tracks/MovieSceneVisibilityTrack.cpp
--- a/Engine/Source/Runtime/MovieSceneTracks/Private/Tracks/MovieSceneVisibilityTrack.cpp
+++ b/Engine/Source/Runtime/MovieSceneTracks/Private/Tracks/MovieSceneVisibilityTrack.cpp
@@ -51,6 +51,24 @@ void UMovieSceneVisibilityTrack::PostLoad()
 	Super::PostLoad();
 }
 
+void UMovieSceneVisibilityTrack::Serialize(FArchive& Ar)
+{
+	Super::Serialize(Ar);
+
+	// Preload BoolSections for PostLoad upgrade if necessary
+	if (Ar.IsLoading())
+	{
+		for (int32 Index = 0; Index < Sections.Num(); ++Index)
+		{
+			UMovieSceneBoolSection* BoolSection = ExactCast<UMovieSceneBoolSection>(Sections[Index]);
+			if (BoolSection)
+			{
+				Ar.Preload(BoolSection);		
+			}
+		}
+	}
+}
+
 bool UMovieSceneVisibilityTrack::SupportsType(TSubclassOf<UMovieSceneSection> SectionClass) const
 {
 	return SectionClass == UMovieSceneVisibilitySection::StaticClass();
diff --git a/Engine/Source/Runtime/MovieSceneTracks/Public/Tracks/MovieSceneVisibilityTrack.h b/Engine/Source/Runtime/MovieSceneTracks/Public/Tracks/MovieSceneVisibilityTrack.h
--- a/Engine/Source/Runtime/MovieSceneTracks/Public/Tracks/MovieSceneVisibilityTrack.h
+++ b/Engine/Source/Runtime/MovieSceneTracks/Public/Tracks/MovieSceneVisibilityTrack.h
@@ -23,6 +23,7 @@ public:
 	virtual void PostLoad() override;
 	virtual bool SupportsType(TSubclassOf<UMovieSceneSection> SectionClass) const override;
 	virtual UMovieSceneSection* CreateNewSection() override;
+	virtual void Serialize(FArchive& Ar) override;
 
 #if WITH_EDITORONLY_DATA
 	virtual FText GetDisplayName() const override;
