Sequencer: only save camera cut pre-animated states for camera cuts that have valid object bindings

diff --git a/Engine/Source/Runtime/MovieSceneTracks/Private/TrackInstances/MovieSceneCameraCutTrackInstance.cpp b/Engine/Source/Runtime/MovieSceneTracks/Private/TrackInstances/MovieSceneCameraCutTrackInstance.cpp
--- a/Engine/Source/Runtime/MovieSceneTracks/Private/TrackInstances/MovieSceneCameraCutTrackInstance.cpp
+++ b/Engine/Source/Runtime/MovieSceneTracks/Private/TrackInstances/MovieSceneCameraCutTrackInstance.cpp
@@ -107,8 +107,7 @@ struct FBlendedCameraCutEasingInfo
 /** Camera cut info struct. */
 struct FBlendedCameraCut
 {
-	FInstanceHandle InstanceHandle;
-	TObjectPtr<UMovieSceneSection> Section;
+	FMovieSceneTrackInstanceInput Input;
 
 	FMovieSceneObjectBindingID CameraBindingID;
 	FMovieSceneSequenceID OperandSequenceID;
@@ -132,8 +131,7 @@ struct FBlendedCameraCut
 	FBlendedCameraCut()
 	{}
 	FBlendedCameraCut(const FMovieSceneTrackInstanceInput& InInput, FMovieSceneObjectBindingID InCameraBindingID, FMovieSceneSequenceID InOperandSequenceID) 
-		: InstanceHandle(InInput.InstanceHandle)
-		, Section(InInput.Section)
+		: Input(InInput)
 		, CameraBindingID(InCameraBindingID)
 		, OperandSequenceID(InOperandSequenceID)
 	{}
@@ -168,15 +166,13 @@ public:
 	static bool MatchesCameraCutCache(UObject* CameraActor, const FBlendedCameraCut& Params, const FCameraCutCache& CameraCutCache)
 	{
 		return CameraActor == CameraCutCache.LastLockedCamera.Get() &&
-			Params.InstanceHandle == CameraCutCache.LastInstanceHandle &&
-			Params.Section == CameraCutCache.LastSection;
+			Params.Input.IsSameInput(CameraCutCache.LastInput);
 	}
 
 	static void UpdateCameraCutCache(UObject* CameraActor, const FBlendedCameraCut& Params, FCameraCutCache& OutCameraCutCache)
 	{
 		OutCameraCutCache.LastLockedCamera = CameraActor;
-		OutCameraCutCache.LastInstanceHandle = Params.InstanceHandle;
-		OutCameraCutCache.LastSection = Params.Section;
+		OutCameraCutCache.LastInput = Params.Input;
 	}
 	
 public:
@@ -233,6 +229,14 @@ public:
 			return false;
 		}
 
+		// Save pre-animated state only now, because we don't want to start tracking this camera cut
+		// unless we know it actually resolves to a camera actor (see above) and will actually do something.
+		FScopedPreAnimatedCaptureSource CaptureSource(Linker, Params.Input);
+		FCameraCutGameHandler::CachePreAnimatedValue(Linker, SequenceInstance);
+#if WITH_EDITOR
+		FCameraCutEditorHandler::CachePreAnimatedValue(Linker, SequenceInstance);
+#endif
+
 		FMovieSceneCameraCutParams CameraCutParams;
 		CameraCutParams.bJumpCut = Context.HasJumped();
 		CameraCutParams.BlendTime = Params.EaseIn.RootBlendTime;
@@ -529,26 +533,11 @@ void UMovieSceneCameraCutTrackInstance::OnAnimate()
 
 	if (CameraCutParams.Num() > 0)
 	{
-		const FSequenceInstance& SequenceInstance = InstanceRegistry->GetInstance(FinalCameraCut.InstanceHandle);
+		const FSequenceInstance& SequenceInstance = InstanceRegistry->GetInstance(FinalCameraCut.Input.InstanceHandle);
 		Animator.AnimateBlendedCameraCut(FinalCameraCut, Linker, SequenceInstance);
 	}
 }
 
-void UMovieSceneCameraCutTrackInstance::OnInputAdded(const FMovieSceneTrackInstanceInput& InInput)
-{
-	using namespace UE::MovieScene;
-
-	UMovieSceneEntitySystemLinker* Linker = GetLinker();
-	const FInstanceRegistry* InstanceRegistry = Linker->GetInstanceRegistry();
-	const FSequenceInstance& SequenceInstance = InstanceRegistry->GetInstance(InInput.InstanceHandle);
-
-	FScopedPreAnimatedCaptureSource CaptureSource(Linker, InInput);
-	FCameraCutGameHandler::CachePreAnimatedValue(Linker, SequenceInstance);
-#if WITH_EDITOR
-	FCameraCutEditorHandler::CachePreAnimatedValue(Linker, SequenceInstance);
-#endif
-}
-
 void UMovieSceneCameraCutTrackInstance::OnEndUpdateInputs()
 {
 	using namespace UE::MovieScene;
diff --git a/Engine/Source/Runtime/MovieSceneTracks/Public/TrackInstances/MovieSceneCameraCutTrackInstance.h b/Engine/Source/Runtime/MovieSceneTracks/Public/TrackInstances/MovieSceneCameraCutTrackInstance.h
--- a/Engine/Source/Runtime/MovieSceneTracks/Public/TrackInstances/MovieSceneCameraCutTrackInstance.h
+++ b/Engine/Source/Runtime/MovieSceneTracks/Public/TrackInstances/MovieSceneCameraCutTrackInstance.h
@@ -57,7 +57,6 @@ public:
 private:
 	virtual void OnInitialize() override;
 	virtual void OnAnimate() override;
-	virtual void OnInputAdded(const FMovieSceneTrackInstanceInput& InInput) override;
 	virtual void OnEndUpdateInputs() override;
 	virtual void OnDestroyed() override;
 
@@ -69,8 +68,7 @@ private:
 	struct FCameraCutCache
 	{
 		TWeakObjectPtr<> LastLockedCamera;
-		UE::MovieScene::FInstanceHandle LastInstanceHandle;
-		TObjectPtr<UMovieSceneSection> LastSection;
+		FMovieSceneTrackInstanceInput LastInput;
 	};
 
 	/**
