Fix in MovieSceneParameter which did not support engine versionning.

diff --git a/Engine/Source/Runtime/MovieSceneTracks/Private/Sections/MovieSceneEventSection.cpp b/Engine/Source/Runtime/MovieSceneTracks/Private/Sections/MovieSceneEventSection.cpp
--- a/Engine/Source/Runtime/MovieSceneTracks/Private/Sections/MovieSceneEventSection.cpp
+++ b/Engine/Source/Runtime/MovieSceneTracks/Private/Sections/MovieSceneEventSection.cpp
@@ -285,7 +285,9 @@ void FMovieSceneEventParameters::GetInstance(FStructOnScope& OutStruct) const
 	if (StructPtr && StructPtr->GetStructureSize() > 0 && StructBytes.Num())
 	{
 		// Deserialize the struct bytes into the struct memory
-		FEventParameterReader(StructBytes).Read(StructPtr, Memory);
+		FEventParameterReader ParamReader(StructBytes);
+		ParamReader.SetUEVer(PackageFileVersion);
+		ParamReader.Read(StructPtr, Memory);
 	}
 }
 
@@ -306,6 +308,11 @@ bool FMovieSceneEventParameters::Serialize(FArchive& Ar)
 	
 	Ar << StructBytes;
 
+	if (Ar.IsLoading())
+	{
+		PackageFileVersion = Ar.UEVer();
+	}
+
 	return true;
 }
 
diff --git a/Engine/Source/Runtime/MovieSceneTracks/Public/Sections/MovieSceneEventSection.h b/Engine/Source/Runtime/MovieSceneTracks/Public/Sections/MovieSceneEventSection.h
--- a/Engine/Source/Runtime/MovieSceneTracks/Public/Sections/MovieSceneEventSection.h
+++ b/Engine/Source/Runtime/MovieSceneTracks/Public/Sections/MovieSceneEventSection.h
@@ -23,11 +23,14 @@ struct FMovieSceneEventParameters
 {
 	GENERATED_BODY()
 
-	FMovieSceneEventParameters() {}
+	FMovieSceneEventParameters()
+		: PackageFileVersion(GPackageFileUEVersion)
+	{}
 
 	/** Construction from a struct type */
 	FMovieSceneEventParameters(UScriptStruct& InStruct)
 		: StructType(&InStruct)
+		, PackageFileVersion(GPackageFileUEVersion)
 	{
 	}
 
@@ -87,6 +90,9 @@ private:
 	/** Soft object path to the type of this parameter payload */
 	FSoftObjectPath StructType;
 
+	/** Used while loading to pass along to the FEventParameterReader */
+	FPackageFileVersion PackageFileVersion;
+
 	/** Serialized bytes that represent the payload. Serialized internally with FEventParameterArchive */
 	TArray<uint8> StructBytes;
 };
