Sequencer: Fixed grouping disparities with data layer pre-animated state

Over time there had become a disparity in the way that this pre-animated storage was handled with respect to grouping. Originally the state was always grouped under a single group handle (mostly because that was mandated by the original code) but a combination of CL 23600697 and 26143512 resulted in the data layer pre-animated state sometimes creating entries using its own group handle, and sometimes using the ones supplied by FBoundObjectPreAnimatedStateTraits.

Ultimately none of this is necessary because we do not need to use grouping here.

diff --git a/Engine/Source/Runtime/MovieSceneTracks/Private/Systems/MovieSceneDataLayerSystem.cpp b/Engine/Source/Runtime/MovieSceneTracks/Private/Systems/MovieSceneDataLayerSystem.cpp
--- a/Engine/Source/Runtime/MovieSceneTracks/Private/Systems/MovieSceneDataLayerSystem.cpp
+++ b/Engine/Source/Runtime/MovieSceneTracks/Private/Systems/MovieSceneDataLayerSystem.cpp
@@ -46,35 +46,22 @@ enum class EDataLayerUpdateFlags : uint8
 ENUM_CLASS_FLAGS(EDataLayerUpdateFlags)
 
 /** Traits class governing how pre-animated state is (re)stored for data layers */
-struct FPreAnimatedDataLayerStorageTraits : FBoundObjectPreAnimatedStateTraits
+struct FPreAnimatedDataLayerStorageTraits : FPreAnimatedStateTraits
 {
 	using KeyType = TObjectKey<UDataLayerInstance>;
 	using StorageType = EDataLayerRuntimeState;
 
 	/** Called when a previously animated data layer needs to be restored */
 	static void RestorePreAnimatedValue(const TObjectKey<UDataLayerInstance>& InKey, EDataLayerRuntimeState PreviousState, const FRestoreStateParams& Params);
+};
 
-	template<typename... T>
-	FPreAnimatedStorageGroupHandle FindGroup(const UDataLayerInstance* BoundObject, T&&... Unused)
-	{
-		return FindGroupImpl(BoundObject);
-	}
 
-	FPreAnimatedStorageGroupHandle FindGroupImpl(const UDataLayerInstance* BoundObject);
-};
 
 /** Container class for all pre-animated data layer state */
 struct FPreAnimatedDataLayerStorage
 	: TPreAnimatedStateStorage<FPreAnimatedDataLayerStorageTraits>
-	, IPreAnimatedStateGroupManager
 {
 	static TAutoRegisterPreAnimatedStorageID<FPreAnimatedDataLayerStorage> StorageID;
-	static TAutoRegisterPreAnimatedStorageID<FPreAnimatedDataLayerStorage> GroupManagerID;
-
-	/*~ IPreAnimatedStateGroupManager */
-	void InitializeGroupManager(FPreAnimatedStateExtension* Extension) override;
-	void OnGroupDestroyed(FPreAnimatedStorageGroupHandle Group) override;
-	void GatherStaleStorageGroups(TArray<FPreAnimatedStorageGroupHandle>& StaleGroupStorage) const override {}
 
 	/** Make an entry for the specified data layer */
 	FPreAnimatedStateEntry FindEntry(const UDataLayerInstance* InDataLayer);
@@ -87,10 +74,6 @@ struct FPreAnimatedDataLayerStorage
 	/** Save the value of a data layer. Should only be used for editor worlds */
 	void SavePreAnimatedStateInEditor(const UDataLayerInstance* DataLayer);
 #endif
-
-private:
-
-	FPreAnimatedStorageGroupHandle GroupHandle;
 };
 
 
@@ -154,43 +137,20 @@ void FPreAnimatedDataLayerStorageTraits::RestorePreAnimatedValue(const TObjectKe
 	}
 }
 
-FPreAnimatedStorageGroupHandle FPreAnimatedDataLayerStorageTraits::FindGroupImpl(const UDataLayerInstance* BoundObject)
-{
-	return ObjectGroupManager->FindGroupForKey(BoundObject);
-}
-
 // ---------------------------------------------------------------------
 // FPreAnimatedDataLayerStorage definitions
 TAutoRegisterPreAnimatedStorageID<FPreAnimatedDataLayerStorage> FPreAnimatedDataLayerStorage::StorageID;
-TAutoRegisterPreAnimatedStorageID<FPreAnimatedDataLayerStorage> FPreAnimatedDataLayerStorage::GroupManagerID;
-
-void FPreAnimatedDataLayerStorage::InitializeGroupManager(FPreAnimatedStateExtension* Extension)
-{}
-
-void FPreAnimatedDataLayerStorage::OnGroupDestroyed(FPreAnimatedStorageGroupHandle Group)
-{
-	ensure(Group == GroupHandle);
-	GroupHandle = FPreAnimatedStorageGroupHandle();
-}
 
 FPreAnimatedStateEntry FPreAnimatedDataLayerStorage::FindEntry(const UDataLayerInstance* InDataLayer)
 {
-	if (!GroupHandle)
-	{
-		return FPreAnimatedStateEntry();
-	}
 	FPreAnimatedStorageIndex StorageIndex = FindStorageIndex(InDataLayer);
-	return FPreAnimatedStateEntry{ GroupHandle, FPreAnimatedStateCachedValueHandle{ StorageID, StorageIndex } };
+	return FPreAnimatedStateEntry{ FPreAnimatedStorageGroupHandle(), FPreAnimatedStateCachedValueHandle{ StorageID, StorageIndex } };
 }
 
 FPreAnimatedStateEntry FPreAnimatedDataLayerStorage::MakeEntry(const UDataLayerInstance* InDataLayer)
 {
-	if (!GroupHandle)
-	{
-		GroupHandle = ParentExtension->AllocateGroup(SharedThis(this));
-	}
 	FPreAnimatedStorageIndex StorageIndex = GetOrCreateStorageIndex(InDataLayer);
-	return FPreAnimatedStateEntry{ GroupHandle, FPreAnimatedStateCachedValueHandle{ StorageID, StorageIndex } };
+	return FPreAnimatedStateEntry{ FPreAnimatedStorageGroupHandle(), FPreAnimatedStateCachedValueHandle{ StorageID, StorageIndex } };
 }
 
 void FPreAnimatedDataLayerStorage::SavePreAnimatedState(const UDataLayerInstance* DataLayer)
