Handle licensee ver in the FMovieSceneEventParameters as well.

diff --git a/Engine/Source/Runtime/MovieSceneTracks/Private/Sections/MovieSceneEventSection.cpp b/Engine/Source/Runtime/MovieSceneTracks/Private/Sections/MovieSceneEventSection.cpp
--- a/Engine/Source/Runtime/MovieSceneTracks/Private/Sections/MovieSceneEventSection.cpp
+++ b/Engine/Source/Runtime/MovieSceneTracks/Private/Sections/MovieSceneEventSection.cpp
@@ -287,6 +287,7 @@ void FMovieSceneEventParameters::GetInstance(FStructOnScope& OutStruct) const
 		// Deserialize the struct bytes into the struct memory
 		FEventParameterReader ParamReader(StructBytes);
 		ParamReader.SetUEVer(PackageFileVersion);
+		ParamReader.SetLicenseeUEVer(LicenseePackageFileVersion);
 		ParamReader.Read(StructPtr, Memory);
 	}
 }
@@ -311,6 +312,7 @@ bool FMovieSceneEventParameters::Serialize(FArchive& Ar)
 	if (Ar.IsLoading())
 	{
 		PackageFileVersion = Ar.UEVer();
+		LicenseePackageFileVersion = Ar.LicenseeUEVer();
 	}
 
 	return true;
diff --git a/Engine/Source/Runtime/MovieSceneTracks/Public/Sections/MovieSceneEventSection.h b/Engine/Source/Runtime/MovieSceneTracks/Public/Sections/MovieSceneEventSection.h
--- a/Engine/Source/Runtime/MovieSceneTracks/Public/Sections/MovieSceneEventSection.h
+++ b/Engine/Source/Runtime/MovieSceneTracks/Public/Sections/MovieSceneEventSection.h
@@ -25,12 +25,14 @@ struct FMovieSceneEventParameters
 
 	FMovieSceneEventParameters()
 		: PackageFileVersion(GPackageFileUEVersion)
+		, LicenseePackageFileVersion(GPackageFileLicenseeUEVersion )
 	{}
 
 	/** Construction from a struct type */
 	FMovieSceneEventParameters(UScriptStruct& InStruct)
 		: StructType(&InStruct)
 		, PackageFileVersion(GPackageFileUEVersion)
+		, LicenseePackageFileVersion(GPackageFileLicenseeUEVersion )
 	{
 	}
 
@@ -92,6 +94,7 @@ private:
 
 	/** Used while loading to pass along to the FEventParameterReader */
 	FPackageFileVersion PackageFileVersion;
+	int32 LicenseePackageFileVersion;
 
 	/** Serialized bytes that represent the payload. Serialized internally with FEventParameterArchive */
 	TArray<uint8> StructBytes;
