Sequencer: Corrected material parameters not removing hierarchical blend targets correctly

When a material parameter was previously subject to a hierarchical blend target continues to be animated but without the blend target, it was previously just removing the blend target component. This has the side effect of the hierarchical blend channel remaining on the entity, but it will not be visited by FinalizeBlendTargets. Changing this code to match the property instantiator by adding the RemoveHierarchicalBlendTarget tag ensures it gets tidied up correctly while maintaining the ordering dependencies:

hierarchical easing instantiator -> material parameter system -> hierarchical easing finalization system

This code should have been updated as part of CL 28275348, but was evidently overlooked

diff --git a/Engine/Source/Runtime/MovieSceneTracks/Private/Systems/MovieSceneMaterialParameterSystem.cpp b/Engine/Source/Runtime/MovieSceneTracks/Private/Systems/MovieSceneMaterialParameterSystem.cpp
--- a/Engine/Source/Runtime/MovieSceneTracks/Private/Systems/MovieSceneMaterialParameterSystem.cpp
+++ b/Engine/Source/Runtime/MovieSceneTracks/Private/Systems/MovieSceneMaterialParameterSystem.cpp
@@ -452,9 +452,9 @@ struct TOverlappingMaterialParameterHandler : Mixin
 						Linker->EntityManager.WriteComponentChecked(Input, BuiltInComponents->HierarchicalBlendTarget, BlendTarget);
 					}
 				}
-				else
+				else if (Linker->EntityManager.HasComponent(Input, BuiltInComponents->HierarchicalBlendTarget))
 				{
-					Linker->EntityManager.RemoveComponent(Input, BuiltInComponents->HierarchicalBlendTarget);
+					Linker->EntityManager.AddComponent(Input, BuiltInComponents->Tags.RemoveHierarchicalBlendTarget);
 				}
 
 				// Ensure we have the blender type tag on the inputs.
@@ -466,6 +466,11 @@ struct TOverlappingMaterialParameterHandler : Mixin
 			Linker->EntityManager.RemoveComponent(Inputs[0], BuiltInComponents->BlendChannelInput);
 
 			Mixin::InitializeSoleInput(Linker, BoundMaterial, ParameterInfo, Inputs[0], Output);
+
+			if (Linker->EntityManager.HasComponent(Inputs[0], BuiltInComponents->HierarchicalBlendTarget))
+			{
+				Linker->EntityManager.AddComponent(Inputs[0], BuiltInComponents->Tags.RemoveHierarchicalBlendTarget);
+			}
 		}
 
 		Output->NumContributors = NumContributors;
diff --git a/Engine/Source/Runtime/MovieSceneTracks/Private/Systems/WeightAndEasingEvaluatorSystem.cpp b/Engine/Source/Runtime/MovieSceneTracks/Private/Systems/WeightAndEasingEvaluatorSystem.cpp
--- a/Engine/Source/Runtime/MovieSceneTracks/Private/Systems/WeightAndEasingEvaluatorSystem.cpp
+++ b/Engine/Source/Runtime/MovieSceneTracks/Private/Systems/WeightAndEasingEvaluatorSystem.cpp
@@ -12,6 +12,7 @@
 #include "EntitySystem/MovieSceneComponentRegistry.h"
 #include "Evaluation/MovieSceneEvaluationTemplateInstance.h"
 #include "Systems/MovieScenePropertyInstantiator.h"
+#include "Systems/MovieSceneMaterialParameterSystem.h"
 #include "EntitySystem/Interrogation/MovieSceneInterrogatedPropertyInstantiator.h"
 #include "Sections/MovieSceneSubSection.h"
 #include "Systems/FloatChannelEvaluatorSystem.h"
@@ -1027,6 +1028,7 @@ UMovieSceneHierarchicalEasingFinalizationSystem::UMovieSceneHierarchicalEasingFi
 	{
 		// Has to come after property instantiation property instantiation so that initial values get created correctly
 		DefineImplicitPrerequisite(UMovieScenePropertyInstantiatorSystem::StaticClass(), GetClass());
+		DefineImplicitPrerequisite(UMovieSceneMaterialParameterInstantiatorSystem::StaticClass(), GetClass());
 		DefineImplicitPrerequisite(UMovieSceneHierarchicalEasingInstantiatorSystem::StaticClass(), GetClass());
 
 		DefineComponentConsumer(GetClass(), FBuiltInComponentTypes::Get()->HierarchicalBlendTarget);
