Sequencer: restore camera cut view target even if the previous view target isn't valid anymore.

If the previous view target isn't valid anymore, we simply pass a null pointer so that the player camera manager defaults to using the player controller. It's not any more correct in theory, but in practice it prevents sequences from ending with a stuck view target.

diff --git a/Engine/Source/Runtime/MovieSceneTracks/Private/TrackInstances/MovieSceneCameraCutGameHandler.cpp b/Engine/Source/Runtime/MovieSceneTracks/Private/TrackInstances/MovieSceneCameraCutGameHandler.cpp
--- a/Engine/Source/Runtime/MovieSceneTracks/Private/TrackInstances/MovieSceneCameraCutGameHandler.cpp
+++ b/Engine/Source/Runtime/MovieSceneTracks/Private/TrackInstances/MovieSceneCameraCutGameHandler.cpp
@@ -70,9 +70,11 @@ void FPreAnimatedCameraCutTraits::RestorePreAnimatedValue(
 	APlayerCameraManager* CameraManager = (PC != nullptr) ? PC->PlayerCameraManager : nullptr;
 
 	// Restore previous view target.
-	AActor* PreviousViewTarget = Cast<AActor>(CachedValue.LastViewTarget.ResolveObjectPtr());
-	if (CameraManager && PreviousViewTarget)
+	// If the previous view target is not valid anymore, we still set it on the camera manger. This will by
+	// default fall back to using the player controller as the view target.
+	if (CameraManager)
 	{
+		AActor* PreviousViewTarget = Cast<AActor>(CachedValue.LastViewTarget.ResolveObjectPtr());
 		CameraManager->SetViewTarget(PreviousViewTarget);
 	}
 
