Sequencer: Downgrade harmless ensure to log when trying to blend properties that can't be blended (ie. static mesh)

diff --git a/Engine/Source/Runtime/MovieSceneTracks/Private/Systems/MovieScenePropertyInstantiator.cpp b/Engine/Source/Runtime/MovieSceneTracks/Private/Systems/MovieScenePropertyInstantiator.cpp
--- a/Engine/Source/Runtime/MovieSceneTracks/Private/Systems/MovieScenePropertyInstantiator.cpp
+++ b/Engine/Source/Runtime/MovieSceneTracks/Private/Systems/MovieScenePropertyInstantiator.cpp
@@ -937,8 +937,10 @@ UMovieScenePropertyInstantiatorSystem::FSetupBlenderSystemResult UMovieSceneProp
 		NewBlenderClass = Params.PropertyDefinition->BlenderSystemClass;
 	}
 
-	if (!ensureMsgf(NewBlenderClass, TEXT("No default blender class specified on property, and no custom blender specified on entities. Falling back to double blender.")))
+	if (!NewBlenderClass)
 	{
+		UE_LOG(LogMovieScene, Warning, TEXT("No default blender class specified on property, and no custom blender specified on entities. Falling back to double blender."));
+
 		NewBlenderClass = UMovieScenePiecewiseDoubleBlenderSystem::StaticClass();
 	}
 
