Sequencer: Add support for component visibility in editor

There is existing support for the actor's temporarily hidden in editor property. This adds support for the component's visible in editor property.

diff --git a/Engine/Source/Runtime/MovieSceneTracks/Private/Systems/MovieSceneVisibilitySystem.cpp b/Engine/Source/Runtime/MovieSceneTracks/Private/Systems/MovieSceneVisibilitySystem.cpp
--- a/Engine/Source/Runtime/MovieSceneTracks/Private/Systems/MovieSceneVisibilitySystem.cpp
+++ b/Engine/Source/Runtime/MovieSceneTracks/Private/Systems/MovieSceneVisibilitySystem.cpp
@@ -20,13 +20,15 @@ namespace UE::MovieScene
 struct FPreAnimatedVisibilityState
 {
 	bool bHidden = false;
-	bool bTemporarilyHiddenInGame = false;
+	bool bActorTemporarilyHiddenInEditor = false;
+	bool bComponentIsVisibleInEditor = false;
 
 	FPreAnimatedVisibilityState()
 	{}
-	FPreAnimatedVisibilityState(bool bInHidden, bool bInTemporarilyHiddenInGame)
+	FPreAnimatedVisibilityState(bool bInHidden, bool bInActorTemporarilyHiddenInEditor, bool bInComponentIsVisibleInEditor)
 		: bHidden(bInHidden)
-		, bTemporarilyHiddenInGame(bInTemporarilyHiddenInGame)
+		, bActorTemporarilyHiddenInEditor(bInActorTemporarilyHiddenInEditor)
+		, bComponentIsVisibleInEditor(bInComponentIsVisibleInEditor)
 	{}
 };
 
@@ -41,24 +43,24 @@ struct FPreAnimatedVisibilityTraits : FBoundObjectPreAnimatedStateTraits
 		if (AActor* Actor = Cast<AActor>(ObjectPtr))
 		{
 			const bool bHidden = Actor->IsHidden();
-			const bool bTemporarilyHiddenInGame = 
+			const bool bTemporarilyHiddenInEditor = 
 #if WITH_EDITOR
 				Actor->IsTemporarilyHiddenInEditor();
 #else
 				false;
 #endif  // WITH_EDITOR
-			return StorageType(bHidden, bTemporarilyHiddenInGame);
+			return StorageType(bHidden, bTemporarilyHiddenInEditor, false);
 		}
 		else if (USceneComponent* SceneComponent = Cast<USceneComponent>(ObjectPtr))
 		{
 			const bool bHidden = SceneComponent->bHiddenInGame;
-			// Second argument doesn't matter, we don't use it when the object is a scene component.
-			return StorageType(bHidden, false);
+			const bool bVisibleInEditor = SceneComponent->IsVisibleInEditor();
+			return StorageType(bHidden, false, bVisibleInEditor);
 		}
 		else
 		{
 			// Dummy value, we won't use it if the bound object isn't an actor or scene component.
-			return StorageType(false, false);
+			return StorageType(false, false, false);
 		}
 	}
 
@@ -69,12 +71,13 @@ struct FPreAnimatedVisibilityTraits : FBoundObjectPreAnimatedStateTraits
 		{
 			Actor->SetActorHiddenInGame(InOutCachedValue.bHidden);
 #if WITH_EDITOR
-			Actor->SetIsTemporarilyHiddenInEditor(InOutCachedValue.bTemporarilyHiddenInGame);
+			Actor->SetIsTemporarilyHiddenInEditor(InOutCachedValue.bActorTemporarilyHiddenInEditor);
 #endif  // WITH_EDITOR
 		}
 		else if (USceneComponent* SceneComponent = Cast<USceneComponent>(ObjectPtr))
 		{
 			SceneComponent->SetHiddenInGame(InOutCachedValue.bHidden);
+			SceneComponent->SetVisibility(InOutCachedValue.bComponentIsVisibleInEditor);
 		}
 	}
 };
@@ -126,6 +129,7 @@ struct FVisibilityTask
 				else if (USceneComponent* SceneComponent = Cast<USceneComponent>(BoundObject))
 				{
 					SceneComponent->SetHiddenInGame(!bShouldBeVisible);
+					SceneComponent->SetVisibility(bShouldBeVisible);
 				}
 			}
 		}
