Sequencer: triggering events now uses playback capabilities instead of the player interface

diff --git a/Engine/Source/Runtime/MovieScene/Private/Evaluation/EventContextsPlaybackCapability.cpp b/Engine/Source/Runtime/MovieScene/Private/Evaluation/EventContextsPlaybackCapability.cpp
new file mode 100644
--- /dev/null
+++ b/Engine/Source/Runtime/MovieScene/Private/Evaluation/EventContextsPlaybackCapability.cpp
@@ -0,0 +1,11 @@
+// Copyright Epic Games, Inc. All Rights Reserved.
+
+#include "Evaluation/EventContextsPlaybackCapability.h"
+
+namespace UE::MovieScene
+{
+
+TPlaybackCapabilityID<IEventContextsPlaybackCapability> IEventContextsPlaybackCapability::ID = TPlaybackCapabilityID<IEventContextsPlaybackCapability>::Register();
+
+}  // namespace UE::MovieScene
+
diff --git a/Engine/Source/Runtime/MovieScene/Private/Evaluation/EventTriggerControlPlaybackCapability.cpp b/Engine/Source/Runtime/MovieScene/Private/Evaluation/EventTriggerControlPlaybackCapability.cpp
new file mode 100644
--- /dev/null
+++ b/Engine/Source/Runtime/MovieScene/Private/Evaluation/EventTriggerControlPlaybackCapability.cpp
@@ -0,0 +1,21 @@
+// Copyright Epic Games, Inc. All Rights Reserved.
+
+#include "Evaluation/EventTriggerControlPlaybackCapability.h"
+
+namespace UE::MovieScene
+{
+
+TPlaybackCapabilityID<FEventTriggerControlPlaybackCapability> FEventTriggerControlPlaybackCapability::ID = TPlaybackCapabilityID<FEventTriggerControlPlaybackCapability>::Register();
+
+bool FEventTriggerControlPlaybackCapability::IsDisablingEventTriggers(FFrameTime& OutDisabledUntilTime) const
+{
+	if (DisableEventTriggersUntilTime.IsSet())
+	{
+		OutDisabledUntilTime = DisableEventTriggersUntilTime.GetValue();
+		return true;
+	}
+	return false;
+}
+
+}  // namespace UE::MovieScene
+
diff --git a/Engine/Source/Runtime/MovieScene/Private/MovieScenePlayer.cpp b/Engine/Source/Runtime/MovieScene/Private/MovieScenePlayer.cpp
--- a/Engine/Source/Runtime/MovieScene/Private/MovieScenePlayer.cpp
+++ b/Engine/Source/Runtime/MovieScene/Private/MovieScenePlayer.cpp
@@ -2,6 +2,8 @@
 
 #include "CoreMinimal.h"
 #include "IMovieScenePlayer.h"
+#include "Evaluation/EventContextsPlaybackCapability.h"
+#include "Evaluation/EventTriggerControlPlaybackCapability.h"
 #include "Evaluation/MovieSceneEvaluationTemplateInstance.h"
 #include "EntitySystem/MovieSceneSequenceInstance.h"
 #include "EntitySystem/MovieSceneEntitySystemLinker.h"
@@ -136,19 +138,13 @@ void IMovieScenePlayer::InvalidateCachedData()
 
 TSharedPtr<UE::MovieScene::FSharedPlaybackState> IMovieScenePlayer::FindSharedPlaybackState()
 {
-	const UE::MovieScene::FSequenceInstance* RootInstance = GetEvaluationTemplate().GetRootInstance();
-	if (RootInstance)
-	{
-		return RootInstance->GetSharedPlaybackState();
-	}
-	return nullptr;
+	return GetEvaluationTemplate().GetSharedPlaybackState();
 }
 
 TSharedRef<UE::MovieScene::FSharedPlaybackState> IMovieScenePlayer::GetSharedPlaybackState()
 {
-	const UE::MovieScene::FSequenceInstance* RootInstance = GetEvaluationTemplate().GetRootInstance();
-	check(RootInstance);
-	return RootInstance->GetSharedPlaybackState();
+	// ToSharedRef will assert if evaluation template isn't initialized
+	return GetEvaluationTemplate().GetSharedPlaybackState().ToSharedRef();
 }
 
 void IMovieScenePlayer::ResetDirectorInstances()
@@ -161,6 +157,38 @@ UObject* IMovieScenePlayer::GetOrCreateDirectorInstance(TSharedRef<const UE::Mov
 	return GetEvaluationTemplate().GetOrCreateDirectorInstance(SequenceID, *this);
 }
 
+TArray<UObject*> IMovieScenePlayer::GetEventContexts() const
+{
+	using namespace UE::MovieScene;
+
+	// By default, look for the playback capability, for backwards compatibility.
+	IMovieScenePlayer* This = const_cast<IMovieScenePlayer*>(this);
+	if (TSharedPtr<const FSharedPlaybackState> SharedPlaybackState = This->FindSharedPlaybackState())
+	{
+		if (IEventContextsPlaybackCapability* EventContextsCapability = SharedPlaybackState->FindCapability<IEventContextsPlaybackCapability>())
+		{
+			return EventContextsCapability->GetEventContexts();
+		}
+	}
+	return TArray<UObject*>();
+}
+
+bool IMovieScenePlayer::IsDisablingEventTriggers(FFrameTime& DisabledUntilTime) const
+{
+	using namespace UE::MovieScene;
+
+	// By default, look for the playback capability, for backwards compatibility.
+	IMovieScenePlayer* This = const_cast<IMovieScenePlayer*>(this);
+	if (TSharedPtr<const FSharedPlaybackState> SharedPlaybackState = This->FindSharedPlaybackState())
+	{
+		if (FEventTriggerControlPlaybackCapability* TriggerControlCapability = SharedPlaybackState->FindCapability<FEventTriggerControlPlaybackCapability>())
+		{
+			return TriggerControlCapability->IsDisablingEventTriggers(DisabledUntilTime);
+		}
+	}
+	return false;
+}
+
 void IMovieScenePlayer::InitializeRootInstance(TSharedRef<UE::MovieScene::FSharedPlaybackState> NewSharedPlaybackState)
 {
 	using namespace UE::MovieScene;
diff --git a/Engine/Source/Runtime/MovieScene/Private/MovieSceneSequencePlayer.cpp b/Engine/Source/Runtime/MovieScene/Private/MovieSceneSequencePlayer.cpp
--- a/Engine/Source/Runtime/MovieScene/Private/MovieSceneSequencePlayer.cpp
+++ b/Engine/Source/Runtime/MovieScene/Private/MovieSceneSequencePlayer.cpp
@@ -17,6 +17,7 @@
 #include "EntitySystem/MovieSceneEntitySystemLinker.h"
 #include "EntitySystem/MovieSceneEntitySystemRunner.h"
 #include "Compilation/MovieSceneCompiledDataManager.h"
+#include "Evaluation/EventTriggerControlPlaybackCapability.h"
 #include "Evaluation/MovieSceneSequenceWeights.h"
 #include "UniversalObjectLocatorResolveParams.h"
 #include "GameFramework/PlayerController.h"
@@ -1298,16 +1299,6 @@ bool UMovieSceneSequencePlayer::IsValid() const
 	return RootTemplateInstance.IsValid();
 }
 
-bool UMovieSceneSequencePlayer::IsDisablingEventTriggers(FFrameTime& DisabledUntilTime) const
-{
-	if (DisableEventTriggersUntilTime.IsSet())
-	{
-		DisabledUntilTime = DisableEventTriggersUntilTime.GetValue();
-		return true;
-	}
-	return false;
-}
-
 bool UMovieSceneSequencePlayer::HasDynamicWeighting() const
 {
 	return PlaybackSettings.bDynamicWeighting;
@@ -1852,6 +1843,8 @@ void UMovieSceneSequencePlayer::PostNetReceive()
 
 void UMovieSceneSequencePlayer::UpdateNetworkSync()
 {
+	using namespace UE::MovieScene;
+
 	if (!bUpdateNetSync)
 	{
 		return;
@@ -1950,7 +1943,9 @@ void UMovieSceneSequencePlayer::UpdateNetworkSync()
 
 			// Also skip all events up to the last known position, otherwise if we skipped back in time we
 			// will re-trigger events again.
-			DisableEventTriggersUntilTime = LastPosition;
+			TSharedRef<FSharedPlaybackState> SharedPlaybackState = GetSharedPlaybackState();
+			FEventTriggerControlPlaybackCapability& TriggerControlCapability = SharedPlaybackState->SetOrAddCapability<FEventTriggerControlPlaybackCapability>();
+			TriggerControlCapability.DisableEventTriggersUntilTime = LastPosition;
 		}
 	}
 }
diff --git a/Engine/Source/Runtime/MovieScene/Public/Evaluation/EventContextsPlaybackCapability.h b/Engine/Source/Runtime/MovieScene/Public/Evaluation/EventContextsPlaybackCapability.h
new file mode 100644
--- /dev/null
+++ b/Engine/Source/Runtime/MovieScene/Public/Evaluation/EventContextsPlaybackCapability.h
@@ -0,0 +1,27 @@
+// Copyright Epic Games, Inc. All Rights Reserved.
+
+#pragma once
+
+#include "Containers/Array.h"
+#include "Evaluation/MovieScenePlaybackCapabilities.h"
+
+class UObject;
+
+namespace UE::MovieScene
+{
+
+/** Playback capability for controlling how events are triggered */
+struct MOVIESCENE_API IEventContextsPlaybackCapability
+{
+	/** Playback capability ID */
+	static TPlaybackCapabilityID<IEventContextsPlaybackCapability> ID;
+
+	/** Virtual destructor */
+	virtual ~IEventContextsPlaybackCapability() {}
+
+	/** Get the contexts used for triggering events */
+	virtual TArray<UObject*> GetEventContexts() const = 0;
+};
+
+}  // namespace UE::MovieScene
+
diff --git a/Engine/Source/Runtime/MovieScene/Public/Evaluation/EventTriggerControlPlaybackCapability.h b/Engine/Source/Runtime/MovieScene/Public/Evaluation/EventTriggerControlPlaybackCapability.h
new file mode 100644
--- /dev/null
+++ b/Engine/Source/Runtime/MovieScene/Public/Evaluation/EventTriggerControlPlaybackCapability.h
@@ -0,0 +1,29 @@
+// Copyright Epic Games, Inc. All Rights Reserved.
+
+#pragma once
+
+#include "Evaluation/MovieScenePlaybackCapabilities.h"
+#include "Misc/FrameTime.h"
+#include "Misc/Optional.h"
+
+namespace UE::MovieScene
+{
+
+/** Playback capability for controlling how events are triggered */
+struct MOVIESCENE_API FEventTriggerControlPlaybackCapability
+{
+	/** Playback capability ID */
+	static TPlaybackCapabilityID<FEventTriggerControlPlaybackCapability> ID;
+
+	/**
+	 * Returns whether triggering events should be temporarily disabled.
+	 *
+	 * @param OutDisabledUntilTime  The time until which to disable triggering events
+	 */
+	bool IsDisablingEventTriggers(FFrameTime& OutDisabledUntilTime) const;
+
+	TOptional<FFrameTime> DisableEventTriggersUntilTime;
+};
+
+}  // namespace UE::MovieScene
+
diff --git a/Engine/Source/Runtime/MovieScene/Public/IMovieScenePlayer.h b/Engine/Source/Runtime/MovieScene/Public/IMovieScenePlayer.h
--- a/Engine/Source/Runtime/MovieScene/Public/IMovieScenePlayer.h
+++ b/Engine/Source/Runtime/MovieScene/Public/IMovieScenePlayer.h
@@ -241,12 +241,12 @@ public:
 	/**
 	 * Access the event contexts for this movie scene player
 	 */
-	virtual TArray<UObject*> GetEventContexts() const { return TArray<UObject*>(); }
+	MOVIESCENE_API virtual TArray<UObject*> GetEventContexts() const;
 
 	/**
 	 * Returns whether event triggers are disabled and if so, until what time.
 	 */
-	virtual bool IsDisablingEventTriggers(FFrameTime& DisabledUntilTime) const { return false; }
+	MOVIESCENE_API virtual bool IsDisablingEventTriggers(FFrameTime& DisabledUntilTime) const;
 
 	/**
 	 * Test whether this is a preview player or not. As such, playback range becomes insignificant for things like spawnables
@@ -405,7 +405,7 @@ public:
 	 * Initializes a new root sequence instance and its shared playback state.
 	 * This adds all the player's playback capabilities to the given state.
 	 */
-	void InitializeRootInstance(TSharedRef<UE::MovieScene::FSharedPlaybackState> NewSharedPlaybackState);
+	MOVIESCENE_API virtual void InitializeRootInstance(TSharedRef<UE::MovieScene::FSharedPlaybackState> NewSharedPlaybackState);
 
 public:
 
diff --git a/Engine/Source/Runtime/MovieScene/Public/MovieSceneSequencePlayer.h b/Engine/Source/Runtime/MovieScene/Public/MovieSceneSequencePlayer.h
--- a/Engine/Source/Runtime/MovieScene/Public/MovieSceneSequencePlayer.h
+++ b/Engine/Source/Runtime/MovieScene/Public/MovieSceneSequencePlayer.h
@@ -538,7 +538,6 @@ protected:
 	virtual void GetViewportSettings(TMap<FViewportClient*, EMovieSceneViewportParams>& ViewportParamsMap) const override {}
 	MOVIESCENE_API virtual void ResolveBoundObjects(const FGuid& InBindingId, FMovieSceneSequenceID SequenceID, UMovieSceneSequence& Sequence, UObject* ResolutionContext, TArray<UObject*, TInlineAllocator<1>>& OutObjects) const override;
 	virtual IMovieScenePlaybackClient* GetPlaybackClient() override { return PlaybackClient ? &*PlaybackClient : nullptr; }
-	MOVIESCENE_API virtual bool IsDisablingEventTriggers(FFrameTime& DisabledUntilTime) const override;
 	MOVIESCENE_API virtual bool HasDynamicWeighting() const override;
 	MOVIESCENE_API virtual void PreEvaluation(const FMovieSceneContext& Context) override;
 	MOVIESCENE_API virtual void PostEvaluation(const FMovieSceneContext& Context) override;
@@ -687,9 +686,6 @@ protected:
 	/** Play position helper */
 	FMovieScenePlaybackPosition PlayPosition;
 
-	/** Disable event triggers until given time */
-	TOptional<FFrameTime> DisableEventTriggersUntilTime;
-
 	/** Spawn register */
 	TSharedPtr<FMovieSceneSpawnRegister> SpawnRegister;
 
diff --git a/Engine/Source/Runtime/MovieSceneTracks/Private/Systems/MovieSceneEventSystems.cpp b/Engine/Source/Runtime/MovieSceneTracks/Private/Systems/MovieSceneEventSystems.cpp
--- a/Engine/Source/Runtime/MovieSceneTracks/Private/Systems/MovieSceneEventSystems.cpp
+++ b/Engine/Source/Runtime/MovieSceneTracks/Private/Systems/MovieSceneEventSystems.cpp
@@ -4,9 +4,13 @@
 #include "Algo/RemoveIf.h"
 #include "EntitySystem/MovieSceneEntitySystemLinker.h"
 #include "EntitySystem/MovieSceneEntitySystemRunner.h"
-#include "EntitySystem/MovieSceneSpawnablesSystem.h"
 #include "EntitySystem/MovieSceneEntitySystemTask.h"
+#include "EntitySystem/MovieSceneSharedPlaybackState.h"
+#include "EntitySystem/MovieSceneSpawnablesSystem.h"
 
+#include "Evaluation/EventContextsPlaybackCapability.h"
+#include "Evaluation/EventTriggerControlPlaybackCapability.h"
+#include "Evaluation/SequenceDirectorPlaybackCapability.h"
 #include "Evaluation/MovieSceneEvaluationTemplateInstance.h"
 #include "MovieSceneTracksComponentTypes.h"
 #include "IMovieScenePlayer.h"
@@ -21,13 +25,11 @@
 
 #include UE_INLINE_GENERATED_CPP_BY_NAME(MovieSceneEventSystems)
 
-
 #define LOCTEXT_NAMESPACE "MovieSceneEventSystem"
 
 DECLARE_CYCLE_STAT(TEXT("Event Systems"),  MovieSceneEval_Events,        STATGROUP_MovieSceneECS);
 DECLARE_CYCLE_STAT(TEXT("Trigger Events"), MovieSceneEval_TriggerEvents, STATGROUP_MovieSceneECS);
 
-
 UMovieSceneEventSystem::UMovieSceneEventSystem(const FObjectInitializer& ObjInit)
 	: Super(ObjInit)
 {
@@ -75,18 +77,27 @@ void UMovieSceneEventSystem::TriggerAllEvents()
 
 	struct FTriggerBatch
 	{
+		FTriggerBatch(TSharedRef<const FSharedPlaybackState> InSharedPlaybackState) : SharedPlaybackState(InSharedPlaybackState) {}
 		TArray<FMovieSceneEventTriggerData> Triggers;
-		IMovieScenePlayer* Player;
+		TSharedRef<const FSharedPlaybackState> SharedPlaybackState;
 	};
 	TArray<FTriggerBatch> TriggerBatches;
 
 	for (TPair<FInstanceHandle, TArray<FMovieSceneEventTriggerData>>& Pair : EventsByRoot)
 	{
 		const FSequenceInstance& RootInstance = InstanceRegistry->GetInstance(Pair.Key);
-		IMovieScenePlayer* Player = RootInstance.GetPlayer();
+		TSharedRef<const FSharedPlaybackState> SharedPlaybackState = RootInstance.GetSharedPlaybackState();
 
 		FFrameTime SkipUntil;
-		const bool bSkipTrigger = Player->IsDisablingEventTriggers(SkipUntil);
+		bool bSkipTrigger = false;
+		if (IMovieScenePlayer* Player = FPlayerIndexPlaybackCapability::GetPlayer(SharedPlaybackState))
+		{
+			bSkipTrigger = Player->IsDisablingEventTriggers(SkipUntil);
+		}
+		else if (FEventTriggerControlPlaybackCapability* TriggerControlCapability = SharedPlaybackState->FindCapability<FEventTriggerControlPlaybackCapability>())
+		{
+			bSkipTrigger = TriggerControlCapability->IsDisablingEventTriggers(SkipUntil);
+		}
 
 		if (RootInstance.GetContext().GetDirection() == EPlayDirection::Forwards)
 		{
@@ -105,9 +116,8 @@ void UMovieSceneEventSystem::TriggerAllEvents()
 			Algo::SortBy(Pair.Value, &FMovieSceneEventTriggerData::RootTime, TGreater<>());
 		}
 
-		FTriggerBatch& TriggerBatch = TriggerBatches.Emplace_GetRef();
+		FTriggerBatch& TriggerBatch = TriggerBatches.Emplace_GetRef(SharedPlaybackState);
 		TriggerBatch.Triggers = Pair.Value;
-		TriggerBatch.Player = Player;
 	}
 
 	// We need to clean our state before actually triggering the events because one of those events could
@@ -117,25 +127,45 @@ void UMovieSceneEventSystem::TriggerAllEvents()
 
 	for (const FTriggerBatch& TriggerBatch : TriggerBatches)
 	{
-		TriggerEvents(TriggerBatch.Triggers, TriggerBatch.Player);
+		TriggerEvents(TriggerBatch.Triggers, TriggerBatch.SharedPlaybackState);
 	}
 }
 
-void UMovieSceneEventSystem::TriggerEvents(TArrayView<const FMovieSceneEventTriggerData> Events, IMovieScenePlayer* Player)
+void UMovieSceneEventSystem::TriggerEvents(TArrayView<const FMovieSceneEventTriggerData> Events, TSharedRef<const FSharedPlaybackState> SharedPlaybackState)
 {
+	using namespace UE::MovieScene;
+
 #if !NO_LOGGING
-	UMovieSceneSequence* PlayerSequence = Player->GetEvaluationTemplate().GetRootSequence();
+	UMovieSceneSequence* PlayerSequence = SharedPlaybackState->GetRootSequence();
 	FString SequenceName = PlayerSequence->GetName();
 	UE_LOG(LogMovieScene, VeryVerbose, TEXT("%s: Triggering %d events"), *SequenceName, Events.Num());
 #endif
 
-	TArray<UObject*> GlobalContexts = Player->GetEventContexts();
+	ISequenceDirectorPlaybackCapability* DirectorCapability = SharedPlaybackState->FindCapability<ISequenceDirectorPlaybackCapability>();
+	if (!ensureMsgf(DirectorCapability, TEXT("Can't trigger events on %s: no director blueprint capability found"), *SequenceName))
+	{
+		return;
+	}
+
+	TArray<UObject*> GlobalContexts;
+	if (IMovieScenePlayer* Player = FPlayerIndexPlaybackCapability::GetPlayer(SharedPlaybackState))
+	{
+		// Check the player interface first. If implemented, we get what we want. If not, it will fall back to the default
+		// implementation that queries the playback capabality (see below), which covers the cases of a player that uses
+		// the new API.
+		GlobalContexts = Player->GetEventContexts();
+	}
+	else if (IEventContextsPlaybackCapability* EventContextsCapability = SharedPlaybackState->FindCapability<IEventContextsPlaybackCapability>())
+	{
+		// No player... let's query the playback capability.
+		GlobalContexts = EventContextsCapability->GetEventContexts();
+	}
 
 	for (const FMovieSceneEventTriggerData& Event : Events)
 	{
 		SCOPE_CYCLE_COUNTER(MovieSceneEval_TriggerEvents);
 
-		UObject* DirectorInstance = Player->GetEvaluationTemplate().GetOrCreateDirectorInstance(Event.SequenceID, *Player);
+		UObject* DirectorInstance = DirectorCapability->GetOrCreateDirectorInstance(SharedPlaybackState, Event.SequenceID);
 		if (!DirectorInstance)
 		{
 #if !NO_LOGGING
@@ -174,15 +204,26 @@ void UMovieSceneEventSystem::TriggerEvents(TArrayView<const FMovieSceneEventTrig
 		}
 		else
 		{
-			TriggerEventWithParameters(DirectorInstance, Event, GlobalContexts, Player);
+			TriggerEventWithParameters(DirectorInstance, Event, GlobalContexts, SharedPlaybackState);
 		}
 	}
 }
 
-void UMovieSceneEventSystem::TriggerEventWithParameters(UObject* DirectorInstance, const FMovieSceneEventTriggerData& Event, TArrayView<UObject* const> GlobalContexts, IMovieScenePlayer* Player)
+void UMovieSceneEventSystem::TriggerEventWithParameters(UObject* DirectorInstance, const FMovieSceneEventTriggerData& Event, TArrayView<UObject* const> GlobalContexts, TSharedRef<const FSharedPlaybackState> SharedPlaybackState)
 {
 	const FMovieSceneEventPtrs& EventPtrs = Event.Ptrs;
-	if (!ensureMsgf(!EventPtrs.BoundObjectProperty.Get() || (EventPtrs.BoundObjectProperty->GetOwner<UObject>() == EventPtrs.Function && EventPtrs.BoundObjectProperty->GetOffset_ForUFunction() < EventPtrs.Function->ParmsSize), TEXT("Bound object property belongs to the wrong function or has an offset greater than the parameter size! This should never happen and indicates a BP compilation or nativization error.")))
+	if (!ensureMsgf(
+				!EventPtrs.BoundObjectProperty.Get() || 
+				(EventPtrs.BoundObjectProperty->GetOwner<UObject>() == EventPtrs.Function && 
+					EventPtrs.BoundObjectProperty->GetOffset_ForUFunction() < EventPtrs.Function->ParmsSize),
+		TEXT("Bound object property belongs to the wrong function or has an offset greater than the parameter size!"
+			 "This should never happen and indicates a BP compilation or nativization error.")))
+	{
+		return;
+	}
+
+	FMovieSceneEvaluationState* State = SharedPlaybackState->FindCapability<FMovieSceneEvaluationState>();
+	if (!ensureMsgf(State, TEXT("No evaluation state found!")))
 	{
 		return;
 	}
@@ -210,12 +251,12 @@ void UMovieSceneEventSystem::TriggerEventWithParameters(UObject* DirectorInstanc
 	// If the event exists on an object binding, only call the events for those bindings (never for the global contexts)
 	if (Event.ObjectBindingID.IsValid())
 	{
-		for (TWeakObjectPtr<> WeakBoundObject : Player->FindBoundObjects(Event.ObjectBindingID, Event.SequenceID))
+		for (TWeakObjectPtr<> WeakBoundObject : State->FindBoundObjects(Event.ObjectBindingID, Event.SequenceID, SharedPlaybackState))
 		{
 			if (UObject* BoundObject = WeakBoundObject.Get())
 			{
 				// Attempt to bind the object to the function parameters
-				if (PatchBoundObject(Parameters, BoundObject, EventPtrs.BoundObjectProperty.Get(), Player, Event.SequenceID))
+				if (PatchBoundObject(Parameters, BoundObject, EventPtrs.BoundObjectProperty.Get(), SharedPlaybackState, Event.SequenceID))
 				{
 					DirectorInstance->ProcessEvent(EventPtrs.Function, Parameters);
 				}
@@ -228,7 +269,7 @@ void UMovieSceneEventSystem::TriggerEventWithParameters(UObject* DirectorInstanc
 	{
 		for (UObject* Context : GlobalContexts)
 		{
-			if (PatchBoundObject(Parameters, Context, EventPtrs.BoundObjectProperty.Get(), Player, Event.SequenceID))
+			if (PatchBoundObject(Parameters, Context, EventPtrs.BoundObjectProperty.Get(), SharedPlaybackState, Event.SequenceID))
 			{
 				DirectorInstance->ProcessEvent(EventPtrs.Function, Parameters);
 			}
@@ -246,7 +287,7 @@ void UMovieSceneEventSystem::TriggerEventWithParameters(UObject* DirectorInstanc
 	}
 }
 
-bool UMovieSceneEventSystem::PatchBoundObject(uint8* Parameters, UObject* BoundObject, FProperty* BoundObjectProperty, IMovieScenePlayer* Player, FMovieSceneSequenceID SequenceID)
+bool UMovieSceneEventSystem::PatchBoundObject(uint8* Parameters, UObject* BoundObject, FProperty* BoundObjectProperty, TSharedRef<const FSharedPlaybackState> SharedPlaybackState, FMovieSceneSequenceID SequenceID)
 {
 	if (!BoundObjectProperty)
 	{
@@ -266,7 +307,7 @@ bool UMovieSceneEventSystem::PatchBoundObject(uint8* Parameters, UObject* BoundO
 
 		FMessageLog("PIE").Warning()
 			->AddToken(FUObjectToken::Create(BoundObjectProperty->GetOwnerUObject()))
-			->AddToken(FUObjectToken::Create(Player->GetEvaluationTemplate().GetSequence(SequenceID)))
+			->AddToken(FUObjectToken::Create(SharedPlaybackState->GetSequence(SequenceID)))
 			->AddToken(FTextToken::Create(FText::Format(LOCTEXT("LevelBP_InterfaceNotImplemented_Error", "Failed to trigger event because it does not implement the necessary interface. Function expects a '{0}'."), FText::FromName(InterfaceParameter->InterfaceClass->GetFName()))));
 		return false;
 	}
@@ -277,7 +318,7 @@ bool UMovieSceneEventSystem::PatchBoundObject(uint8* Parameters, UObject* BoundO
 		{
 			FMessageLog("PIE").Warning()
 				->AddToken(FUObjectToken::Create(BoundObjectProperty->GetOwnerUObject()))
-				->AddToken(FUObjectToken::Create(Player->GetEvaluationTemplate().GetSequence(SequenceID)))
+				->AddToken(FUObjectToken::Create(SharedPlaybackState->GetSequence(SequenceID)))
 				->AddToken(FTextToken::Create(LOCTEXT("LevelBP_LevelScriptActor_Error", "Failed to trigger event: only Interface pins are supported for root tracks within Level Sequences. Please remove the pin, or change it to an interface that is implemented on the desired level blueprint.")));
 
 			return false;
@@ -285,7 +326,7 @@ bool UMovieSceneEventSystem::PatchBoundObject(uint8* Parameters, UObject* BoundO
 		else if (!BoundObject->IsA(ObjectParameter->PropertyClass))
 		{
 			FMessageLog("PIE").Warning()
-				->AddToken(FUObjectToken::Create(Player->GetEvaluationTemplate().GetSequence(SequenceID)))
+				->AddToken(FUObjectToken::Create(SharedPlaybackState->GetSequence(SequenceID)))
 				->AddToken(FUObjectToken::Create(BoundObjectProperty->GetOwnerUObject()))
 				->AddToken(FUObjectToken::Create(BoundObject))
 				->AddToken(FTextToken::Create(FText::Format(LOCTEXT("LevelBP_InvalidCast_Error", "Failed to trigger event: Cast to {0} failed."), FText::FromName(ObjectParameter->PropertyClass->GetFName()))));
@@ -298,7 +339,7 @@ bool UMovieSceneEventSystem::PatchBoundObject(uint8* Parameters, UObject* BoundO
 	}
 
 	FMessageLog("PIE").Warning()
-		->AddToken(FUObjectToken::Create(Player->GetEvaluationTemplate().GetSequence(SequenceID)))
+		->AddToken(FUObjectToken::Create(SharedPlaybackState->GetSequence(SequenceID)))
 		->AddToken(FUObjectToken::Create(BoundObjectProperty->GetOwnerUObject()))
 		->AddToken(FUObjectToken::Create(BoundObject))
 		->AddToken(FTextToken::Create(FText::Format(LOCTEXT("LevelBP_UnsupportedProperty_Error", "Failed to trigger event: Unsupported property type for bound object: {0}."), FText::FromName(BoundObjectProperty->GetClass()->GetFName()))));
diff --git a/Engine/Source/Runtime/MovieSceneTracks/Public/Systems/MovieSceneEventSystems.h b/Engine/Source/Runtime/MovieSceneTracks/Public/Systems/MovieSceneEventSystems.h
--- a/Engine/Source/Runtime/MovieSceneTracks/Public/Systems/MovieSceneEventSystems.h
+++ b/Engine/Source/Runtime/MovieSceneTracks/Public/Systems/MovieSceneEventSystems.h
@@ -4,9 +4,15 @@
 
 #include "EntitySystem/MovieSceneEntitySystem.h"
 #include "Channels/MovieSceneEvent.h"
+#include "Evaluation/IMovieScenePlaybackCapability.h"
 #include "MovieSceneEventSystems.generated.h"
 
-class IMovieScenePlayer;
+namespace UE::MovieScene
+{
+
+struct FSharedPlaybackState;
+
+}  // namespace UE::MovieScene
 
 USTRUCT()
 struct FMovieSceneEventTriggerData
@@ -47,13 +53,15 @@ protected:
 
 private:
 
+	using FSharedPlaybackState = UE::MovieScene::FSharedPlaybackState;
+
 	MOVIESCENETRACKS_API virtual void OnRun(FSystemTaskPrerequisites& InPrerequisites, FSystemSubsequentTasks& Subsequents) override;
 	MOVIESCENETRACKS_API virtual bool IsRelevantImpl(UMovieSceneEntitySystemLinker* InLinker) const override;
 	MOVIESCENETRACKS_API virtual void OnUnlink() override final;
 
-	static void TriggerEvents(TArrayView<const FMovieSceneEventTriggerData> Events, IMovieScenePlayer* Player);
-	static void TriggerEventWithParameters(UObject* DirectorInstance, const FMovieSceneEventTriggerData& Event, TArrayView<UObject* const> GlobalContexts, IMovieScenePlayer* Player);
-	static bool PatchBoundObject(uint8* Parameters, UObject* BoundObject, FProperty* BoundObjectProperty, IMovieScenePlayer* Player, FMovieSceneSequenceID SequenceID);
+	static void TriggerEvents(TArrayView<const FMovieSceneEventTriggerData> Events, TSharedRef<const FSharedPlaybackState> SharedPlaybackState);
+	static void TriggerEventWithParameters(UObject* DirectorInstance, const FMovieSceneEventTriggerData& Event, TArrayView<UObject* const> GlobalContexts, TSharedRef<const FSharedPlaybackState> SharedPlaybackState);
+	static bool PatchBoundObject(uint8* Parameters, UObject* BoundObject, FProperty* BoundObjectProperty, TSharedRef<const FSharedPlaybackState> SharedPlaybackState, FMovieSceneSequenceID SequenceID);
 
 	/** Events grouped by root sequence instance handle */
 	TMap< UE::MovieScene::FInstanceHandle, TArray<FMovieSceneEventTriggerData> > EventsByRoot;
