Cameras: Fix shake attempting to restart after restore state

When restore state happens, the shake is stopped. OnRun(), shake is not restarted because its signature matches up with the existing shake. But it also needs to see if that shake is active.

diff --git a/Engine/Source/Runtime/MovieSceneTracks/Private/Systems/MovieSceneCameraShakeSystem.cpp b/Engine/Source/Runtime/MovieSceneTracks/Private/Systems/MovieSceneCameraShakeSystem.cpp
--- a/Engine/Source/Runtime/MovieSceneTracks/Private/Systems/MovieSceneCameraShakeSystem.cpp
+++ b/Engine/Source/Runtime/MovieSceneTracks/Private/Systems/MovieSceneCameraShakeSystem.cpp
@@ -402,7 +402,7 @@ void UMovieSceneCameraShakeInstantiatorSystem::OnRun(FSystemTaskPrerequisites& I
 		if (ShakeSourceComponent)
 		{
 			bool bStartShake = true;
-			if (ShakeInstanceData.SectionSignature == ShakeData.SectionSignature)
+			if (ShakeInstanceData.SectionSignature == ShakeData.SectionSignature && ShakeInstanceData.ShakeInstance.Get() != nullptr && ShakeInstanceData.ShakeInstance->IsActive())
 			{
 				// Don't re-create and restart the shake if it was already running with the same
 				// parameters and this instantiation phase is just re-importing other stuff.
@@ -462,7 +462,7 @@ void UMovieSceneCameraShakeInstantiatorSystem::OnRun(FSystemTaskPrerequisites& I
 		{
 			bool bStartInstance = true;
 			UCameraShakeBase* ShakeInstance = nullptr;
-			if (ShakeInstanceData.SectionSignature == ShakeData.SectionSignature)
+			if (ShakeInstanceData.SectionSignature == ShakeData.SectionSignature && ShakeInstanceData.ShakeInstance.Get() != nullptr && ShakeInstanceData.ShakeInstance->IsActive())
 			{
 				// Don't re-create and restart the shake if it was already running with the same
 				// parameters and this instantiation phase is just re-importing other stuff.
