Sequencer: Fixed hierarchical easing channels being destroyed when still in use

The issue here was new entities being linked from inside an existing subsequence that had hierarchical blending enabled (which still only really happens with Day Sequence). In this sitation the number of hierarchical easing providers has not changed, so bChannelsHaveBeenInvalidated is false, but we were previously always destroying all the channels for blend targets inside FinalizeBlendTargets and only re-constructing them from the { NeedsLink | HierarchicalBlendTarget } entities. Clearly this results in pre-existing channels that relate to blend targets being destroyed if they did not happen to have a { NeedsLink | HierarchicalBlendTarget } entity.

There are two options to fix this: firstly we could always re-create and re-assign _all_ easing channels that relate to hierarchical blend targets, or we could ensure that we don't destroy channels that relate to blend targets in this situation. I opted for the latter since the first option is way more expensive.

diff --git a/Engine/Source/Runtime/MovieScene/Public/MovieSceneFwd.h b/Engine/Source/Runtime/MovieScene/Public/MovieSceneFwd.h
--- a/Engine/Source/Runtime/MovieScene/Public/MovieSceneFwd.h
+++ b/Engine/Source/Runtime/MovieScene/Public/MovieSceneFwd.h
@@ -15,6 +15,10 @@
 	#define UE_MOVIESCENE_EVENTS WITH_EDITOR
 #endif
 
+#ifndef UE_MOVIESCENE_EXPENSIVE_CONSISTENCY_CHECKS
+	#define UE_MOVIESCENE_EXPENSIVE_CONSISTENCY_CHECKS 0
+#endif
+
 namespace UE
 {
 namespace MovieScene
diff --git a/Engine/Source/Runtime/MovieSceneTracks/Private/Systems/WeightAndEasingEvaluatorSystem.cpp b/Engine/Source/Runtime/MovieSceneTracks/Private/Systems/WeightAndEasingEvaluatorSystem.cpp
--- a/Engine/Source/Runtime/MovieSceneTracks/Private/Systems/WeightAndEasingEvaluatorSystem.cpp
+++ b/Engine/Source/Runtime/MovieSceneTracks/Private/Systems/WeightAndEasingEvaluatorSystem.cpp
@@ -888,7 +888,6 @@ void UMovieSceneHierarchicalEasingInstantiatorSystem::FinalizeBlendTargets()
 	if (bChannelsHaveBeenInvalidated || Linker->EntityManager.Contains(NewBlendTargetFilter))
 	{
 		FHierarchicalEasingChannelBuffer& Buffer = EvaluatorSystem->GetComputationBuffer();
-		Buffer.ResetBlendTargets();
 
 		const int32 NumStartingChannels = Buffer.Channels.Num();
 
@@ -900,6 +899,14 @@ void UMovieSceneHierarchicalEasingInstantiatorSystem::FinalizeBlendTargets()
 		TSortedMap<int16, FHBiasChannelData, TInlineAllocator<16>, TGreater<>> HBiasToChannelData;
 		TSet<FHierarchicalBlendTarget, DefaultKeyFuncs<FHierarchicalBlendTarget>, TInlineSetAllocator<16>> BlendTargetChannelData;
 
+		// If we have invalidated our channel list we need to recompute everything.
+		//    Otherwise we will only visit NeedsLink entities and create channels for anything that's missing
+		if (bChannelsHaveBeenInvalidated)
+		{
+			Buffer.ResetBlendTargets();
+			CachedHierarchicalBlendTargetChannels.Empty();
+		}
+
 		// Step 1: Gather all the hierarchical blend targets
 		//
 		{
@@ -949,8 +956,6 @@ void UMovieSceneHierarchicalEasingInstantiatorSystem::FinalizeBlendTargets()
 
 		// Step 3: Create chains of channels for each HBias designated as a blend target
 		{
-			CachedHierarchicalBlendTargetChannels.Empty();
-
 			for (const FHierarchicalBlendTarget& BlendTarget : BlendTargetChannelData)
 			{
 				// This is a blend target - create channels for its entire parent chain
@@ -961,6 +966,12 @@ void UMovieSceneHierarchicalEasingInstantiatorSystem::FinalizeBlendTargets()
 
 				for (int16 SourceHBias : BlendTarget.AsArray())
 				{
+					if (CachedHierarchicalBlendTargetChannels.Contains(MakeTuple(SourceHBias, BlendTarget)))
+					{
+						// This should only occur when bChannelsHaveBeenInvalidated is false
+						continue;
+					}
+
 					FHierarchicalEasingChannelData BlendTargetChannel;
 					// Get our value from the accumulator channel for the source bias which includes all the accumulated weights for that HBias level.
 					// Normally this should just be a single weight
@@ -998,6 +1009,7 @@ void UMovieSceneHierarchicalEasingInstantiatorSystem::FinalizeBlendTargets()
 		AssignEasingChannelsFilter.All({ BuiltInComponents->InstanceHandle, BuiltInComponents->HierarchicalBlendTarget });
 		if (bChannelsHaveBeenInvalidated == false)
 		{
+			// If channels have not been invalidated, only visit new blend target entities to assign their channels
 			AssignEasingChannelsFilter.All({ BuiltInComponents->Tags.NeedsLink });
 		}
 		AssignEasingChannelsFilter.None({ BuiltInComponents->Tags.NeedsUnlink, BuiltInComponents->HierarchicalEasingProvider, BuiltInComponents->Tags.ImportedEntity });
@@ -1014,6 +1026,24 @@ void UMovieSceneHierarchicalEasingInstantiatorSystem::FinalizeBlendTargets()
 		FRemoveSingleMutation RemoveTag(BuiltInComponents->Tags.RemoveHierarchicalBlendTarget);
 		Linker->EntityManager.MutateAll(FEntityComponentFilter().All({ BuiltInComponents->Tags.RemoveHierarchicalBlendTarget }), RemoveTag);
 	}
+
+
+#if UE_MOVIESCENE_EXPENSIVE_CONSISTENCY_CHECKS
+	FHierarchicalEasingChannelBuffer& Buffer = EvaluatorSystem->GetComputationBuffer();
+
+	// Check that all hierarchical easing components map to a valid channel
+	FEntityTaskBuilder()
+	.Read(BuiltInComponents->HierarchicalEasingChannel)
+	.FilterNone({ BuiltInComponents->Tags.NeedsUnlink })
+	.Iterate_PerAllocation(&Linker->EntityManager, [Buffer](const FEntityAllocation* Allocation, TRead<uint16> Channels) {
+
+		for (int32 Index = 0; Index < Allocation->Num(); ++Index)
+		{
+			const uint16 Channel = Channels[Index];
+			ensureAlways(Buffer.Channels.IsValidIndex(Channel));
+		}
+	});
+#endif
 }
 
 UMovieSceneHierarchicalEasingFinalizationSystem::UMovieSceneHierarchicalEasingFinalizationSystem(const FObjectInitializer& ObjInit)
