Sequencer: No longer consider a null material as invalidating the instantiation phase if it was previously null

diff --git a/Engine/Source/Runtime/MovieSceneTracks/Public/Systems/MovieSceneMaterialSystem.h b/Engine/Source/Runtime/MovieSceneTracks/Public/Systems/MovieSceneMaterialSystem.h
--- a/Engine/Source/Runtime/MovieSceneTracks/Public/Systems/MovieSceneMaterialSystem.h
+++ b/Engine/Source/Runtime/MovieSceneTracks/Public/Systems/MovieSceneMaterialSystem.h
@@ -157,8 +157,12 @@ struct TInitializeBoundMaterials
 
 		if (!ExistingMaterial)
 		{
-			OutDynamicMaterial = FObjectComponent::Null();
-			return true;
+			if (OutDynamicMaterial != nullptr)
+			{
+				OutDynamicMaterial = FObjectComponent::Null();
+				return true;
+			}
+			return false;
 		}
 
 		if (UMaterialInstanceDynamic* MID = Cast<UMaterialInstanceDynamic>(ExistingMaterial))
