MovieScene: Added custom MovieScene accessors to improve performance of commonly used light/sky/fog tracks.

LightComponent.IndirectLightingIntensity
DirectionalLightComponent.SetLightSourceAngle
ExponentialHeightFogComponent.StartDistance
ExponentialHeightFogComponent.VolumetricFogExtinctionScale
LocalFogVolumeComponent.RadialFogExtinction
LocalFogVolumeComponent.FogAlbedo
SkyAtmosphereComponent.MieExponentialDistribution
SkyAtmosphereComponent.OtherAbsorption

diff --git a/Engine/Source/Runtime/MovieSceneTracks/Private/MovieSceneTracksCustomAccessors.cpp b/Engine/Source/Runtime/MovieSceneTracks/Private/MovieSceneTracksCustomAccessors.cpp
--- a/Engine/Source/Runtime/MovieSceneTracks/Private/MovieSceneTracksCustomAccessors.cpp
+++ b/Engine/Source/Runtime/MovieSceneTracks/Private/MovieSceneTracksCustomAccessors.cpp
@@ -2,8 +2,10 @@
 
 #include "MovieSceneTracksCustomAccessors.h"
 #include "MovieSceneTracksComponentTypes.h"
+#include "Components/DirectionalLightComponent.h"
 #include "Components/ExponentialHeightFogComponent.h"
 #include "Components/LightComponent.h"
+#include "Components/LocalFogVolumeComponent.h"
 #include "Components/PrimitiveComponent.h"
 #include "Components/SkyAtmosphereComponent.h"
 #include "Components/SkyLightComponent.h"
@@ -116,6 +118,30 @@ void SetLightComponentVolumetricScatteringIntensity(UObject* Object, float InVol
 	LightComponent->SetVolumetricScatteringIntensity(InVolumetricScatteringIntensity);
 }
 
+float GetLightComponentIndirectLightingIntensity(const UObject* Object)
+{
+	const ULightComponent* LightComponent = CastChecked<const ULightComponent>(Object);
+	return LightComponent->IndirectLightingIntensity;
+}
+
+void SetLightComponentIndirectLightingIntensity(UObject* Object, float InIndirectLightingIntensity)
+{
+	ULightComponent* LightComponent = CastChecked<ULightComponent>(Object);
+	LightComponent->SetIndirectLightingIntensity(InIndirectLightingIntensity);
+}
+
+float GetDirectionalLightComponentLightSourceAngle(const UObject* Object)
+{
+	const UDirectionalLightComponent* DirectionalLightComponent = CastChecked<UDirectionalLightComponent>(Object);
+	return DirectionalLightComponent->LightSourceAngle;
+}
+
+void SetDirectionalLightComponentLightSourceAngle(UObject* Object, float InLightSourceAngle)
+{
+	UDirectionalLightComponent* DirectionalLightComponent = CastChecked<UDirectionalLightComponent>(Object);
+	DirectionalLightComponent->SetLightSourceAngle(InLightSourceAngle);
+}
+
 FIntermediateColor GetSkyLightComponentLightColor(const UObject* Object, EColorPropertyType InColorType)
 {
 	ensure(InColorType == EColorPropertyType::Color);
@@ -160,6 +186,18 @@ void SetSkyAtmosphereComponentMieScatteringScale(UObject* Object, float InMieSca
 	SkyAtmosphereComponent->SetMieScatteringScale(InMieScatteringScale);
 }
 
+float GetSkyAtmosphereComponentMieExponentialDistribution(const UObject* Object)
+{
+	const USkyAtmosphereComponent* SkyAtmosphereComponent = CastChecked<const USkyAtmosphereComponent>(Object);
+	return SkyAtmosphereComponent->MieExponentialDistribution;
+}
+
+void SetSkyAtmosphereComponentMieExponentialDistribution(UObject* Object, float InMieExponentialDistribution)
+{
+	USkyAtmosphereComponent* SkyAtmosphereComponent = CastChecked<USkyAtmosphereComponent>(Object);
+	SkyAtmosphereComponent->SetMieExponentialDistribution(InMieExponentialDistribution);
+}
+
 float GetSkyAtmosphereComponentOtherAbsorptionScale(const UObject* Object)
 {
 	const USkyAtmosphereComponent* SkyAtmosphereComponent = CastChecked<const USkyAtmosphereComponent>(Object);
@@ -172,6 +210,22 @@ void SetSkyAtmosphereComponentOtherAbsorptionScale(UObject* Object, float InOthe
 	SkyAtmosphereComponent->SetOtherAbsorptionScale(InOtherAbsorptionScale);
 }
 
+FIntermediateColor GetSkyAtmosphereComponentOtherAbsorption(const UObject* Object, EColorPropertyType InColorType)
+{
+	ensure(InColorType == EColorPropertyType::Linear);
+	
+	const USkyAtmosphereComponent* SkyAtmosphereComponent = CastChecked<const USkyAtmosphereComponent>(Object);
+	return FIntermediateColor(SkyAtmosphereComponent->OtherAbsorption);
+}
+
+void SetSkyAtmosphereComponentOtherAbsorption(UObject* Object, EColorPropertyType InColorType, const FIntermediateColor& InColor)
+{
+	ensure(InColorType == EColorPropertyType::Linear);
+
+	USkyAtmosphereComponent* SkyAtmosphereComponent = CastChecked<USkyAtmosphereComponent>(Object);
+	SkyAtmosphereComponent->SetOtherAbsorption(InColor.GetLinearColor());
+}
+
 float GetSkyAtmosphereComponentRayleighScatteringScale(const UObject* Object)
 {
 	const USkyAtmosphereComponent* SkyAtmosphereComponent = CastChecked<const USkyAtmosphereComponent>(Object);
@@ -347,6 +401,30 @@ void SetExponentialHeightFogComponentDirectionalInscatteringExponent(UObject* Ob
 	UExponentialHeightFogComponent* ExponentialHeightFogComponent = CastChecked<UExponentialHeightFogComponent>(Object);
 	ExponentialHeightFogComponent->SetDirectionalInscatteringExponent(InDirectionalInscatteringExponent);
 }
+
+float GetExponentialHeightFogComponentStartDistance(const UObject* Object)
+{
+	const UExponentialHeightFogComponent* ExponentialHeightFogComponent = CastChecked<const UExponentialHeightFogComponent>(Object);
+	return ExponentialHeightFogComponent->StartDistance;
+}
+
+void SetExponentialHeightFogComponentStartDistance(UObject* Object, float InStartDistance)
+{
+	UExponentialHeightFogComponent* ExponentialHeightFogComponent = CastChecked<UExponentialHeightFogComponent>(Object);
+	ExponentialHeightFogComponent->SetStartDistance(InStartDistance);
+}
+
+float GetExponentialHeightFogComponentVolumetricFogExtinctionScale(const UObject* Object)
+{
+	const UExponentialHeightFogComponent* ExponentialHeightFogComponent = CastChecked<const UExponentialHeightFogComponent>(Object);
+	return ExponentialHeightFogComponent->VolumetricFogExtinctionScale;
+}
+
+void SetExponentialHeightFogComponentVolumetricFogExtinctionScale(UObject* Object, float InVolumetricFogExtinctionScale)
+{
+	UExponentialHeightFogComponent* ExponentialHeightFogComponent = CastChecked<UExponentialHeightFogComponent>(Object);
+	ExponentialHeightFogComponent->SetVolumetricFogExtinctionScale(InVolumetricFogExtinctionScale);
+}
 	
 FIntermediateColor GetExponentialHeightFogComponentDirectionalInscatteringColor(const UObject* Object, EColorPropertyType InColorType)
 {
@@ -436,6 +514,34 @@ void SetSecondFogDataFogHeightOffset(UObject* Object, float InFogHeightOffset)
 	ExponentialHeightFogComponent->SetSecondFogHeightOffset(InFogHeightOffset);
 }
 
+float GetLocalFogVolumeComponentRadialFogExtinction(const UObject* Object)
+{
+	const ULocalFogVolumeComponent* LocalFogVolumeComponent = CastChecked<ULocalFogVolumeComponent>(Object);
+	return LocalFogVolumeComponent->RadialFogExtinction;
+}
+
+void SetLocalFogVolumeComponentRadialFogExtinction(UObject* Object, float InRadialFogExtinction)
+{
+	ULocalFogVolumeComponent* LocalFogVolumeComponent = CastChecked<ULocalFogVolumeComponent>(Object);
+	LocalFogVolumeComponent->SetRadialFogExtinction(InRadialFogExtinction);
+}
+
+FIntermediateColor GetLocalFogVolumeComponentFogAlbedo(const UObject* Object, EColorPropertyType InColorType)
+{
+	ensure(InColorType == EColorPropertyType::Linear);
+
+	const ULocalFogVolumeComponent* LocalFogVolumeComponent = CastChecked<ULocalFogVolumeComponent>(Object);
+	return FIntermediateColor(LocalFogVolumeComponent->FogAlbedo);
+}
+
+void SetLocalFogVolumeComponentFogAlbedo(UObject* Object, EColorPropertyType InColorType, const FIntermediateColor& InColor)
+{
+	ensure(InColorType == EColorPropertyType::Linear);
+
+	ULocalFogVolumeComponent* LocalFogVolumeComponent = CastChecked<ULocalFogVolumeComponent>(Object);
+	LocalFogVolumeComponent->SetFogAlbedo(InColor.GetLinearColor());
+}
+
 FObjectComponent GetSkeletalMeshAsset(const UObject* Object, FObjectPropertyTraits::FObjectMetadata ObjectMetadata)
 {
 	const USkeletalMeshComponent* SkeletalMeshComponent = CastChecked<const USkeletalMeshComponent>(Object);
@@ -481,6 +587,15 @@ void InitializeMovieSceneTracksAccessors(FMovieSceneTracksComponentTypes* Tracks
 			ULightComponent::StaticClass(), GET_MEMBER_NAME_CHECKED(ULightComponent, VolumetricScatteringIntensity),
 			GetLightComponentVolumetricScatteringIntensity, SetLightComponentVolumetricScatteringIntensity);
 
+	TracksComponents->Accessors.Float.Add(
+			ULightComponent::StaticClass(), GET_MEMBER_NAME_CHECKED(ULightComponent, IndirectLightingIntensity),
+			GetLightComponentIndirectLightingIntensity, SetLightComponentIndirectLightingIntensity);
+
+	// DirectionalLightComponent
+	TracksComponents->Accessors.Float.Add(
+			UDirectionalLightComponent::StaticClass(), GET_MEMBER_NAME_CHECKED(UDirectionalLightComponent, LightSourceAngle),
+			GetDirectionalLightComponentLightSourceAngle, SetDirectionalLightComponentLightSourceAngle);
+
 	// SkyLightComponent
 	TracksComponents->Accessors.Color.Add(
 			USkyLightComponent::StaticClass(), GET_MEMBER_NAME_CHECKED(USkyLightComponent, LightColor), 
@@ -493,6 +608,9 @@ void InitializeMovieSceneTracksAccessors(FMovieSceneTracksComponentTypes* Tracks
 	TracksComponents->Accessors.Float.Add(
 			USkyAtmosphereComponent::StaticClass(), GET_MEMBER_NAME_CHECKED(USkyAtmosphereComponent, MieScatteringScale),
 			GetSkyAtmosphereComponentMieScatteringScale, SetSkyAtmosphereComponentMieScatteringScale);
+	TracksComponents->Accessors.Float.Add(
+			USkyAtmosphereComponent::StaticClass(), GET_MEMBER_NAME_CHECKED(USkyAtmosphereComponent, MieExponentialDistribution),
+			GetSkyAtmosphereComponentMieExponentialDistribution, SetSkyAtmosphereComponentMieExponentialDistribution);
 	TracksComponents->Accessors.Float.Add(
 			USkyAtmosphereComponent::StaticClass(), GET_MEMBER_NAME_CHECKED(USkyAtmosphereComponent, OtherAbsorptionScale),
 			GetSkyAtmosphereComponentOtherAbsorptionScale, SetSkyAtmosphereComponentOtherAbsorptionScale);
@@ -511,6 +629,9 @@ void InitializeMovieSceneTracksAccessors(FMovieSceneTracksComponentTypes* Tracks
 	TracksComponents->Accessors.Color.Add(
 			USkyAtmosphereComponent::StaticClass(), GET_MEMBER_NAME_CHECKED(USkyAtmosphereComponent, MieScattering),
 			GetSkyAtmosphereComponentMieScattering, SetSkyAtmosphereComponentMieScattering);
+	TracksComponents->Accessors.Color.Add(
+			USkyAtmosphereComponent::StaticClass(), GET_MEMBER_NAME_CHECKED(USkyAtmosphereComponent, OtherAbsorption),
+			GetSkyAtmosphereComponentOtherAbsorption, SetSkyAtmosphereComponentOtherAbsorption);
 	TracksComponents->Accessors.Color.Add(
 			USkyAtmosphereComponent::StaticClass(), GET_MEMBER_NAME_CHECKED(USkyAtmosphereComponent, RayleighScattering),
 			GetSkyAtmosphereComponentRayleighScattering, SetSkyAtmosphereComponentRayleighScattering);
@@ -534,6 +655,12 @@ void InitializeMovieSceneTracksAccessors(FMovieSceneTracksComponentTypes* Tracks
 	TracksComponents->Accessors.Float.Add(
 			UExponentialHeightFogComponent::StaticClass(), GET_MEMBER_NAME_CHECKED(UExponentialHeightFogComponent, DirectionalInscatteringExponent),
 			GetExponentialHeightFogComponentDirectionalInscatteringExponent, SetExponentialHeightFogComponentDirectionalInscatteringExponent);
+	TracksComponents->Accessors.Float.Add(
+			UExponentialHeightFogComponent::StaticClass(), GET_MEMBER_NAME_CHECKED(UExponentialHeightFogComponent, StartDistance),
+			GetExponentialHeightFogComponentStartDistance, SetExponentialHeightFogComponentStartDistance);
+	TracksComponents->Accessors.Float.Add(
+			UExponentialHeightFogComponent::StaticClass(), GET_MEMBER_NAME_CHECKED(UExponentialHeightFogComponent, VolumetricFogExtinctionScale),
+			GetExponentialHeightFogComponentVolumetricFogExtinctionScale, SetExponentialHeightFogComponentVolumetricFogExtinctionScale);
 	TracksComponents->Accessors.Color.Add(
 			UExponentialHeightFogComponent::StaticClass(), GET_MEMBER_NAME_CHECKED(UExponentialHeightFogComponent, DirectionalInscatteringLuminance),
 			GetExponentialHeightFogComponentDirectionalInscatteringColor, SetExponentialHeightFogComponentDirectionalInscatteringColor);
@@ -559,6 +686,15 @@ void InitializeMovieSceneTracksAccessors(FMovieSceneTracksComponentTypes* Tracks
 			UExponentialHeightFogComponent::StaticClass(), *SecondFogDataFogHeightOffsetPath,
 			GetSecondFogDataFogHeightOffset, SetSecondFogDataFogHeightOffset);
 
+	// LocalFogVolumeComponent
+	TracksComponents->Accessors.Float.Add(
+			ULocalFogVolumeComponent::StaticClass(), GET_MEMBER_NAME_CHECKED(ULocalFogVolumeComponent, RadialFogExtinction),
+			GetLocalFogVolumeComponentRadialFogExtinction, SetLocalFogVolumeComponentRadialFogExtinction);
+	TracksComponents->Accessors.Color.Add(
+			ULocalFogVolumeComponent::StaticClass(), GET_MEMBER_NAME_CHECKED(ULocalFogVolumeComponent, FogAlbedo),
+			GetLocalFogVolumeComponentFogAlbedo, SetLocalFogVolumeComponentFogAlbedo);
+
+	// SceneComponent
 	TracksComponents->Accessors.ComponentTransform.Add(USceneComponent::StaticClass(), "Transform", &GetComponentTransform, &SetComponentTransformAndVelocity);
 
 	// SkeletalMeshComponent
