[Backout] - CL31658439
Original CL Desc
-----------------------------------------------------------------
Sequencer: No longer consider a null material as invalidating the instantiation phase if it was previously null

diff --git a/Engine/Source/Runtime/MovieSceneTracks/Public/Systems/MovieSceneMaterialSystem.h b/Engine/Source/Runtime/MovieSceneTracks/Public/Systems/MovieSceneMaterialSystem.h
--- a/Engine/Source/Runtime/MovieSceneTracks/Public/Systems/MovieSceneMaterialSystem.h
+++ b/Engine/Source/Runtime/MovieSceneTracks/Public/Systems/MovieSceneMaterialSystem.h
@@ -157,12 +157,8 @@ struct TInitializeBoundMaterials
 
 		if (!ExistingMaterial)
 		{
-			if (OutDynamicMaterial != nullptr)
-			{
-				OutDynamicMaterial = FObjectComponent::Null();
-				return true;
-			}
-			return false;
+			OutDynamicMaterial = FObjectComponent::Null();
+			return true;
 		}
 
 		if (UMaterialInstanceDynamic* MID = Cast<UMaterialInstanceDynamic>(ExistingMaterial))
