World Partition - DataLayer: Added options to the DataLayer sections for the sequencer
* When activating a DL, there are now options to force a streaming flush (or not) and/or a GC. Previously, activating a DL would always force a flush if the data layer content wasn't all loaded.
* Other minor improvements like adding edit conditions on sections properties and removing the preroll text when sections are not configured to preroll

diff --git a/Engine/Source/Editor/MovieSceneTools/Private/TrackEditors/DataLayerTrackEditor.cpp b/Engine/Source/Editor/MovieSceneTools/Private/TrackEditors/DataLayerTrackEditor.cpp
--- a/Engine/Source/Editor/MovieSceneTools/Private/TrackEditors/DataLayerTrackEditor.cpp
+++ b/Engine/Source/Editor/MovieSceneTools/Private/TrackEditors/DataLayerTrackEditor.cpp
@@ -115,17 +115,17 @@ struct FDataLayerSection
 	FText GetPrerollText() const
 	{
 		UMovieSceneDataLayerSection* Section = WeakSection.Get();
-		if (Section)
+		if (Section && Section->GetPreRollFrames() > 0)
 		{
 			switch (Section->GetPrerollState())
 			{
 			case EDataLayerRuntimeState::Unloaded:  return LOCTEXT("PrerollText_Unloaded", "(Unloaded over time in pre/post roll)");
 			case EDataLayerRuntimeState::Loaded:    return LOCTEXT("PrerollText_Loaded", "(Loaded over time in preroll)");
-			default: break;
+			case EDataLayerRuntimeState::Activated:	return LOCTEXT("PrerollText_Activated", "(Activated over time in preroll)");
 			}
 		}
 
-		return LOCTEXT("PrerollText_Activated", "(Activated over time in preroll)");
+		return FText();
 	}
 
 	FText GetLayerBarText() const
diff --git a/Engine/Source/Runtime/Engine/Classes/Engine/CoreSettings.h b/Engine/Source/Runtime/Engine/Classes/Engine/CoreSettings.h
--- a/Engine/Source/Runtime/Engine/Classes/Engine/CoreSettings.h
+++ b/Engine/Source/Runtime/Engine/Classes/Engine/CoreSettings.h
@@ -147,6 +147,8 @@ extern ENGINE_API float GLevelStreamingUnregisterComponentsTimeLimit;
 extern ENGINE_API int32 GLevelStreamingForceGCAfterLevelStreamedOut;
 /** Whether to kick off incremental GC when there are over the specified amount of levels still waiting to be purged. */
 extern ENGINE_API int32 GLevelStreamingContinuouslyIncrementalGCWhileLevelsPendingPurge;
+/** Whether to override GLevelStreamingContinuouslyIncrementalGCWhileLevelsPendingPurge temporarily until at least one level needs to be purged.  */
+extern ENGINE_API int32 GLevelStreamingContinuouslyIncrementalGCWhileLevelsPendingPurgeOverride;
 /** Enables level streaming requests while async loading (of anything) while the match is already in progress and no loading screen is up. */
 extern ENGINE_API int32 GLevelStreamingAllowLevelRequestsWhileAsyncLoadingInMatch;
 /** When we're already loading this many levels and actively in match, don't allow any more requests until one of those completes.  Set to zero to disable. */
diff --git a/Engine/Source/Runtime/Engine/Private/CoreSettings.cpp b/Engine/Source/Runtime/Engine/Private/CoreSettings.cpp
--- a/Engine/Source/Runtime/Engine/Private/CoreSettings.cpp
+++ b/Engine/Source/Runtime/Engine/Private/CoreSettings.cpp
@@ -22,6 +22,7 @@ int32 GLevelStreamingComponentsUnregistrationGranularity = 5;
 int32 GLevelStreamingRouteActorInitializationGranularity = 10;
 int32 GLevelStreamingForceGCAfterLevelStreamedOut = 1;
 int32 GLevelStreamingContinuouslyIncrementalGCWhileLevelsPendingPurge = 1;
+int32 GLevelStreamingContinuouslyIncrementalGCWhileLevelsPendingPurgeOverride = 0;
 int32 GLevelStreamingAllowLevelRequestsWhileAsyncLoadingInMatch = 1;
 int32 GLevelStreamingMaxLevelRequestsAtOnceWhileInMatch = 0;
 int32 GLevelStreamingForceVerifyLevelsGotRemovedByGC = 0;
diff --git a/Engine/Source/Runtime/Engine/Private/World.cpp b/Engine/Source/Runtime/Engine/Private/World.cpp
--- a/Engine/Source/Runtime/Engine/Private/World.cpp
+++ b/Engine/Source/Runtime/Engine/Private/World.cpp
@@ -4341,6 +4341,8 @@ void UWorld::UpdateLevelStreaming()
 	StreamingLevelsToConsider.EndConsideration();
 
 	const int32 CurrentNumLevelsPendingPurge = FLevelStreamingGCHelper::GetNumLevelsPendingPurge();
+	const int32 LevelStreamingContinuouslyIncrementalGCWhileLevelsPendingPurge = GLevelStreamingContinuouslyIncrementalGCWhileLevelsPendingPurgeOverride ? 1 : GLevelStreamingContinuouslyIncrementalGCWhileLevelsPendingPurge;
+	const bool bShouldPurgeLevels = LevelStreamingContinuouslyIncrementalGCWhileLevelsPendingPurge && CurrentNumLevelsPendingPurge >= LevelStreamingContinuouslyIncrementalGCWhileLevelsPendingPurge;
 
 	CSV_CUSTOM_STAT(LevelStreamingPendingPurge, NumlevelsPendingPurge, CurrentNumLevelsPendingPurge, ECsvCustomStatOp::Set);
 
@@ -4352,7 +4354,7 @@ void UWorld::UpdateLevelStreaming()
 	{
 		GEngine->ForceGarbageCollection(false);
 	}
-	else if (GLevelStreamingContinuouslyIncrementalGCWhileLevelsPendingPurge && CurrentNumLevelsPendingPurge >= GLevelStreamingContinuouslyIncrementalGCWhileLevelsPendingPurge)
+	else if (bShouldPurgeLevels)
 	{
 		// Request a 'soft' GC if there are levels pending purge and there are levels to be loaded. In the case of a blocking
 		// load this is going to guarantee GC firing first thing afterwards and otherwise it is going to sneak in right before
diff --git a/Engine/Source/Runtime/MovieSceneTracks/Private/Sections/MovieSceneDataLayerSection.cpp b/Engine/Source/Runtime/MovieSceneTracks/Private/Sections/MovieSceneDataLayerSection.cpp
--- a/Engine/Source/Runtime/MovieSceneTracks/Private/Sections/MovieSceneDataLayerSection.cpp
+++ b/Engine/Source/Runtime/MovieSceneTracks/Private/Sections/MovieSceneDataLayerSection.cpp
@@ -11,7 +11,9 @@ UMovieSceneDataLayerSection::UMovieSceneDataLayerSection(const FObjectInitialize
 {
 	DesiredState = EDataLayerRuntimeState::Activated;
 	PrerollState = EDataLayerRuntimeState::Activated;
+	bFlushOnActivated = true;
 	bFlushOnUnload = false;
+	bPerformGCOnUnload = false;
 	EvalOptions.EnableAndSetCompletionMode(EMovieSceneCompletionMode::RestoreState);
 }
 
@@ -35,6 +37,16 @@ void UMovieSceneDataLayerSection::SetPrerollState(EDataLayerRuntimeState InPrero
 	PrerollState = InPrerollState;
 }
 
+bool UMovieSceneDataLayerSection::GetFlushOnActivated() const
+{
+	return bFlushOnActivated;
+}
+
+void UMovieSceneDataLayerSection::SetFlushOnActivated(bool bInFlushOnActivated)
+{
+	bFlushOnActivated = bInFlushOnActivated;
+}
+
 bool UMovieSceneDataLayerSection::GetFlushOnUnload() const
 {
 	return bFlushOnUnload;
@@ -45,6 +57,16 @@ void UMovieSceneDataLayerSection::SetFlushOnUnload(bool bInFlushOnUnload)
 	bFlushOnUnload = bInFlushOnUnload;
 }
 
+bool UMovieSceneDataLayerSection::GetPerformGCOnUnload() const
+{
+	return bPerformGCOnUnload;
+}
+
+void UMovieSceneDataLayerSection::SetPerformGCOnUnload(bool bInPerformGCOnUnload)
+{
+	bPerformGCOnUnload = bInPerformGCOnUnload;
+}
+
 void UMovieSceneDataLayerSection::ImportEntityImpl(UMovieSceneEntitySystemLinker* EntityLinker, const FEntityImportParams& Params, FImportedEntity* OutImportedEntity)
 {
 	using namespace UE::MovieScene;
diff --git a/Engine/Source/Runtime/MovieSceneTracks/Private/Systems/MovieSceneDataLayerSystem.cpp b/Engine/Source/Runtime/MovieSceneTracks/Private/Systems/MovieSceneDataLayerSystem.cpp
--- a/Engine/Source/Runtime/MovieSceneTracks/Private/Systems/MovieSceneDataLayerSystem.cpp
+++ b/Engine/Source/Runtime/MovieSceneTracks/Private/Systems/MovieSceneDataLayerSystem.cpp
@@ -19,6 +19,7 @@
 #include "WorldPartition/WorldPartitionSubsystem.h"
 #include "WorldPartition/WorldPartitionRuntimeCell.h"
 #include "WorldPartition/WorldPartitionStreamingSource.h"
+#include "Engine/CoreSettings.h"
 #include "Engine/World.h"
 #include "Misc/EnumClassFlags.h"
 #include "MovieSceneSequence.h"
@@ -42,6 +43,7 @@ enum class EDataLayerUpdateFlags : uint8
 	None						= 0,
 	FlushStreamingVisibility	= 1,
 	FlushStreamingFull			= 2,
+	PerformGarbageCollect		= 4,
 };
 ENUM_CLASS_FLAGS(EDataLayerUpdateFlags)
 
@@ -81,9 +83,10 @@ struct FDataLayerState
 {
 	void Reset();
 	bool IsEmpty() const;
-	void AddRequest(int16 InBias, EDataLayerRuntimeState RequestedState, bool bRequiresStreamingFlush);
+	void AddRequest(int16 InBias, EDataLayerRuntimeState RequestedState, bool bRequiresStreamingFlush, bool bPerformGC);
 	TOptional<EDataLayerRuntimeState> ComputeDesiredState() const;
 	bool ShouldFlushStreaming(EDataLayerRuntimeState ComputedState) const;
+	bool ShouldPerformGarbageCollect(EDataLayerRuntimeState ComputedState) const;
 
 private:
 
@@ -93,6 +96,7 @@ private:
 	int32 ActivatedCount   = 0;
 	bool  bFlushUnloaded   = false;
 	bool  bFlushActivated  = false;
+	bool  bGCUnloaded      = false;
 };
 
 
@@ -104,7 +108,7 @@ struct FDesiredLayerStates
 #if WITH_EDITOR
 	void ApplyInEditor(FPreAnimatedDataLayerStorage* PreAnimatedStorage, UDataLayerEditorSubsystem* EditorSubSystem);
 #endif
-	void ApplyNewState(const UDataLayerInstance* InDataLayer, int16 HierarchicalBias, EDataLayerRuntimeState DesiredState, bool bRequiresStreamingFlush);
+	void ApplyNewState(const UDataLayerInstance* InDataLayer, int16 HierarchicalBias, EDataLayerRuntimeState DesiredState, bool bRequiresStreamingFlush, bool bPerformGC);
 
 	TMap<TObjectKey<UDataLayerInstance>, FDataLayerState> StatesByInstance;
 };
@@ -209,9 +213,10 @@ void FDataLayerState::Reset()
 	ActivatedCount   = 0;
 	bFlushUnloaded   = false;
 	bFlushActivated  = false;
+	bGCUnloaded      = false;
 }
 
-void FDataLayerState::AddRequest(int16 InBias, EDataLayerRuntimeState RequestedState, bool bRequiresStreamingFlush)
+void FDataLayerState::AddRequest(int16 InBias, EDataLayerRuntimeState RequestedState, bool bRequiresStreamingFlush, bool bPerformGC)
 {
 	if (InBias > HierarchicalBias)
 	{
@@ -223,7 +228,7 @@ void FDataLayerState::AddRequest(int16 InBias, EDataLayerRuntimeState RequestedS
 	{
 		switch (RequestedState)
 		{
-		case EDataLayerRuntimeState::Unloaded:  ++UnloadedCount;  bFlushUnloaded |= bRequiresStreamingFlush; break;
+		case EDataLayerRuntimeState::Unloaded:  ++UnloadedCount;  bFlushUnloaded |= bRequiresStreamingFlush; bGCUnloaded |= bPerformGC; break;
 		case EDataLayerRuntimeState::Loaded:    ++LoadedCount;    break;
 		case EDataLayerRuntimeState::Activated: ++ActivatedCount; bFlushActivated |= bRequiresStreamingFlush; break;
 		}
@@ -264,6 +269,16 @@ bool FDataLayerState::ShouldFlushStreaming(EDataLayerRuntimeState ComputedState)
 	return false;
 }
 
+bool FDataLayerState::ShouldPerformGarbageCollect(EDataLayerRuntimeState ComputedState) const
+{
+	switch (ComputedState)
+	{
+	case EDataLayerRuntimeState::Unloaded:  return bGCUnloaded;
+	}
+
+	return false;
+}
+
 // ---------------------------------------------------------------------
 // FDataLayerStates definitions
 bool FDesiredLayerStates::IsEmpty() const
@@ -360,6 +375,11 @@ EDataLayerUpdateFlags FDesiredLayerStates::Apply(FPreAnimatedDataLayerStorage* P
 							*StaticEnum<EDataLayerRuntimeState>()->GetDisplayNameTextByValue((int64)DesiredStateValue).ToString());
 					}
 				}
+
+				if (StateValue.ShouldPerformGarbageCollect(DesiredStateValue))
+				{
+					Flags |= EDataLayerUpdateFlags::PerformGarbageCollect;
+				}
 			}
 		}
 	}
@@ -431,7 +451,7 @@ void FDesiredLayerStates::ApplyInEditor(FPreAnimatedDataLayerStorage* PreAnimate
 }
 #endif
 
-void FDesiredLayerStates::ApplyNewState(const UDataLayerInstance* InDataLayer, int16 HierarchicalBias, EDataLayerRuntimeState DesiredState, bool bRequiresStreamingFlush)
+void FDesiredLayerStates::ApplyNewState(const UDataLayerInstance* InDataLayer, int16 HierarchicalBias, EDataLayerRuntimeState DesiredState, bool bRequiresStreamingFlush, bool bPerformGC)
 {
 	using namespace UE::MovieScene;
 
@@ -441,7 +461,7 @@ void FDesiredLayerStates::ApplyNewState(const UDataLayerInstance* InDataLayer, i
 		LayerState = &StatesByInstance.Add(InDataLayer, FDataLayerState());
 	}
 
-	LayerState->AddRequest(HierarchicalBias, DesiredState, bRequiresStreamingFlush);
+	LayerState->AddRequest(HierarchicalBias, DesiredState, bRequiresStreamingFlush, bPerformGC);
 }
 
 } // namespace MovieScene
@@ -544,6 +564,12 @@ void UMovieSceneDataLayerSystem::OnRun(FSystemTaskPrerequisites& InPrerequisites
 				World->FlushLevelStreaming(EFlushLevelStreamingType::Visibility);
 			}
 
+			if (EnumHasAnyFlags(UpdateFlags, EDataLayerUpdateFlags::PerformGarbageCollect))
+			{
+				UE_LOG(LogMovieScene, Warning, TEXT("[UMovieSceneDataLayerSystem] Forcing garbage collection"));
+				GLevelStreamingContinuouslyIncrementalGCWhileLevelsPendingPurgeOverride = 1;
+			}
+
 			UE_SUPPRESS(LogMovieScene, Warning,
 			{
 				if (FlushTypeString)
@@ -620,14 +646,16 @@ void UMovieSceneDataLayerSystem::UpdateDesiredStates()
 			if (DataLayerManager)
 			{
 				EDataLayerRuntimeState DesiredState = bPreroll ? Section->GetPrerollState() : Section->GetDesiredState();
-				const bool bRequiresStreamingFlush = (DesiredState == EDataLayerRuntimeState::Unloaded) ? Section->GetFlushOnUnload() : !bPreroll;
+				const bool bRequiresStreamingFlush = (DesiredState == EDataLayerRuntimeState::Unloaded) ? Section->GetFlushOnUnload() :
+													 (DesiredState == EDataLayerRuntimeState::Activated) ? Section->GetFlushOnActivated() : false;
+				const bool bPerformGC = (DesiredState == EDataLayerRuntimeState::Unloaded) ? Section->GetPerformGCOnUnload() : false;
 
 				for (const UDataLayerAsset* DataLayerAsset : Section->GetDataLayerAssets())
 				{
 					const UDataLayerInstance* DataLayerInstance = DataLayerManager->GetDataLayerInstanceFromAsset(DataLayerAsset);
 					if (DataLayerInstance)
 					{
-						this->DesiredLayerStates->ApplyNewState(DataLayerInstance, OptHBiases ? OptHBiases[Index] : 0, DesiredState, bRequiresStreamingFlush);
+						this->DesiredLayerStates->ApplyNewState(DataLayerInstance, OptHBiases ? OptHBiases[Index] : 0, DesiredState, bRequiresStreamingFlush, bPerformGC);
 					}
 				}
 			}
diff --git a/Engine/Source/Runtime/MovieSceneTracks/Public/Sections/MovieSceneDataLayerSection.h b/Engine/Source/Runtime/MovieSceneTracks/Public/Sections/MovieSceneDataLayerSection.h
--- a/Engine/Source/Runtime/MovieSceneTracks/Public/Sections/MovieSceneDataLayerSection.h
+++ b/Engine/Source/Runtime/MovieSceneTracks/Public/Sections/MovieSceneDataLayerSection.h
@@ -33,11 +33,23 @@ public:
 	UFUNCTION(BlueprintCallable, Category = "Sequencer|Section")
 	MOVIESCENETRACKS_API void SetPrerollState(EDataLayerRuntimeState InPrerollState);
 
+	UFUNCTION(BlueprintPure, Category = "Sequencer|Section")
+	MOVIESCENETRACKS_API bool GetFlushOnActivated() const;
+
+	UFUNCTION(BlueprintCallable, Category = "Sequencer|Section")
+	MOVIESCENETRACKS_API void SetFlushOnActivated(bool bFlushOnActivated);
+
 	UFUNCTION(BlueprintPure, Category = "Sequencer|Section")
 	MOVIESCENETRACKS_API bool GetFlushOnUnload() const;
 
 	UFUNCTION(BlueprintCallable, Category = "Sequencer|Section")
 	MOVIESCENETRACKS_API void SetFlushOnUnload(bool bFlushOnUnload);
+		
+	UFUNCTION(BlueprintPure, Category = "Sequencer|Section")
+	MOVIESCENETRACKS_API bool GetPerformGCOnUnload() const;
+
+	UFUNCTION(BlueprintCallable, Category = "Sequencer|Section")
+	MOVIESCENETRACKS_API void SetPerformGCOnUnload(bool bPerformGCOnUnload);
 
 	UFUNCTION(BlueprintPure, Category = "Sequencer|Section")
 	const TArray<UDataLayerAsset*>& GetDataLayerAssets() const { return DataLayerAssets; }
@@ -70,10 +82,21 @@ private:
 	EDataLayerRuntimeState DesiredState;
 
 	/** The desired state for the data layers on this section when the section is pre or post-rolling. */
-	UPROPERTY(EditAnywhere, Category=DataLayer)
+	UPROPERTY(EditAnywhere, Category=DataLayer, Meta=(EditCondition="HasPreRoll", EditConditionHides))
 	EDataLayerRuntimeState PrerollState;
 
+	/** Determine if we need to flush level streaming when the data layers are activated. May result in a hitch. */
+	UPROPERTY(EditAnywhere, Category=DataLayer, Meta=(EditCondition="DesiredState == EDataLayerRuntimeState::Activated", EditConditionHides))
+	bool bFlushOnActivated;
+
 	/** Determine if we need to flush level streaming when the data layers unloads. */
-	UPROPERTY(EditAnywhere, Category=DataLayer)
+	UPROPERTY(EditAnywhere, Category=DataLayer, Meta=(EditCondition="DesiredState == EDataLayerRuntimeState::Unloaded", EditConditionHides))
 	bool bFlushOnUnload;
+
+	/** Determine if we need to perform a GC when the data layers unloads. */
+	UPROPERTY(EditAnywhere, Category=DataLayer, Meta=(EditCondition="DesiredState == EDataLayerRuntimeState::Unloaded", EditConditionHides, DisplayName="Perform GC On Unload"))
+	bool bPerformGCOnUnload;
+
+	UFUNCTION()
+	bool HasPreRoll() const { return GetPreRollFrames() > 0.0f; }
 };
