Sequencer: Fix static analysis warnings from CL 32218082.

diff --git a/Engine/Plugins/MovieScene/LevelSequenceEditor/Source/LevelSequenceEditor/Private/Misc/LevelSequenceEditorSpawnRegister.cpp b/Engine/Plugins/MovieScene/LevelSequenceEditor/Source/LevelSequenceEditor/Private/Misc/LevelSequenceEditorSpawnRegister.cpp
--- a/Engine/Plugins/MovieScene/LevelSequenceEditor/Source/LevelSequenceEditor/Private/Misc/LevelSequenceEditorSpawnRegister.cpp
+++ b/Engine/Plugins/MovieScene/LevelSequenceEditor/Source/LevelSequenceEditor/Private/Misc/LevelSequenceEditorSpawnRegister.cpp
@@ -109,6 +109,7 @@ void FLevelSequenceEditorSpawnRegister::PreDestroyObject(UObject& Object, const
 	UMovieSceneSequence*  Sequence      = Sequencer.IsValid() ? Sequencer->GetEvaluationTemplate().GetSequence(TemplateID) : nullptr;
 	UObject*              SpawnedObject = FindSpawnedObject(BindingId, TemplateID, BindingIndex).Get();
 
+
 	if (SpawnedObject)
 	{
 		const FTrackedObjectState* TrackedState = TrackedObjects.Find(&Object);
@@ -117,7 +118,10 @@ void FLevelSequenceEditorSpawnRegister::PreDestroyObject(UObject& Object, const
 			// SaveDefaultSpawnableState will reset bHasBeenModified to false
 			SaveDefaultSpawnableStateImpl(BindingId, BindingIndex, Sequence, SpawnedObject, Sequencer->GetSharedPlaybackState());
 
-			Sequence->MarkPackageDirty();
+			if (Sequence)
+			{
+				Sequence->MarkPackageDirty();
+			}
 		}
 	}
 
diff --git a/Engine/Source/Runtime/MovieSceneTracks/Private/Bindings/MovieSceneSpawnableActorBinding.cpp b/Engine/Source/Runtime/MovieSceneTracks/Private/Bindings/MovieSceneSpawnableActorBinding.cpp
--- a/Engine/Source/Runtime/MovieSceneTracks/Private/Bindings/MovieSceneSpawnableActorBinding.cpp
+++ b/Engine/Source/Runtime/MovieSceneTracks/Private/Bindings/MovieSceneSpawnableActorBinding.cpp
@@ -453,6 +453,12 @@ UMovieSceneCustomBinding* UMovieSceneSpawnableActorBinding::CreateNewCustomBindi
 		}
 
 		NewCustomBinding = CreateBinding();
+
+		if (!NewCustomBinding)
+		{
+			return nullptr;
+		}
+
 		AActor* SpawnedActor = Cast<AActor>(StaticDuplicateObject(Actor, NewCustomBinding, TemplateName, RF_AllFlags));
 		SpawnedActor->DetachFromActor(FDetachmentTransformRules(EDetachmentRule::KeepRelative, false));
 #if WITH_EDITORONLY_DATA
@@ -476,6 +482,10 @@ UMovieSceneCustomBinding* UMovieSceneSpawnableActorBinding::CreateNewCustomBindi
 			return nullptr;
 		}
 		NewCustomBinding = CreateBinding();
+		if (!NewCustomBinding)
+		{
+			return nullptr;
+		}
 		NewCustomBinding->SetObjectTemplate(NewObject<UObject>(NewCustomBinding, SourceBlueprint->GeneratedClass, TemplateName, RF_Transactional));
 	}
 #if WITH_EDITORONLY_DATA
@@ -490,6 +500,11 @@ UMovieSceneCustomBinding* UMovieSceneSpawnableActorBinding::CreateNewCustomBindi
 
 			NewCustomBinding = CreateBinding();
 
+			if (!NewCustomBinding)
+			{
+				return nullptr;
+			}
+
 			NewCustomBinding->SetObjectTemplate(NewObject<UObject>(NewCustomBinding, BlueprintGeneratedBy->GeneratedClass, TemplateName, RF_Transactional));
 		}
 	}
@@ -505,6 +520,12 @@ UMovieSceneCustomBinding* UMovieSceneSpawnableActorBinding::CreateNewCustomBindi
 			}
 
 			NewCustomBinding = CreateBinding();
+
+			if (!NewCustomBinding)
+			{
+				return nullptr;
+			}
+
 			NewCustomBinding->SetObjectTemplate(NewObject<UObject>(&OwnerMovieScene, InClass, TemplateName, RF_Transactional));
 		}
 	}
