Sequencer- Fix more static analysis warnings.

diff --git a/Engine/Source/Runtime/MovieSceneTracks/Private/Bindings/MovieSceneSpawnableActorBinding.cpp b/Engine/Source/Runtime/MovieSceneTracks/Private/Bindings/MovieSceneSpawnableActorBinding.cpp
--- a/Engine/Source/Runtime/MovieSceneTracks/Private/Bindings/MovieSceneSpawnableActorBinding.cpp
+++ b/Engine/Source/Runtime/MovieSceneTracks/Private/Bindings/MovieSceneSpawnableActorBinding.cpp
@@ -454,24 +454,22 @@ UMovieSceneCustomBinding* UMovieSceneSpawnableActorBinding::CreateNewCustomBindi
 
 		NewCustomBinding = CreateBinding();
 
-		if (!NewCustomBinding)
+		if (NewCustomBinding)
 		{
-			return nullptr;
-		}
-
-		AActor* SpawnedActor = Cast<AActor>(StaticDuplicateObject(Actor, NewCustomBinding, TemplateName, RF_AllFlags));
-		SpawnedActor->DetachFromActor(FDetachmentTransformRules(EDetachmentRule::KeepRelative, false));
+			AActor* SpawnedActor = Cast<AActor>(StaticDuplicateObject(Actor, NewCustomBinding, TemplateName, RF_AllFlags));
+			SpawnedActor->DetachFromActor(FDetachmentTransformRules(EDetachmentRule::KeepRelative, false));
 #if WITH_EDITORONLY_DATA
-		SpawnedActor->bIsEditorPreviewActor = false;
+			SpawnedActor->bIsEditorPreviewActor = false;
 #endif
-		NewCustomBinding->SetObjectTemplate(SpawnedActor);
+			NewCustomBinding->SetObjectTemplate(SpawnedActor);
 
 #if WITH_EDITOR
-		BindingNameToSet = Actor->GetActorLabel();
+			BindingNameToSet = Actor->GetActorLabel();
 #endif
-		if (!bWasTransactional)
-		{
-			Actor->ClearFlags(RF_Transactional);
+			if (!bWasTransactional)
+			{
+				Actor->ClearFlags(RF_Transactional);
+			}
 		}
 	}
 	// If it's a blueprint, we need some special handling
@@ -482,11 +480,10 @@ UMovieSceneCustomBinding* UMovieSceneSpawnableActorBinding::CreateNewCustomBindi
 			return nullptr;
 		}
 		NewCustomBinding = CreateBinding();
-		if (!NewCustomBinding)
+		if (NewCustomBinding)
 		{
-			return nullptr;
+			NewCustomBinding->SetObjectTemplate(NewObject<UObject>(NewCustomBinding, SourceBlueprint->GeneratedClass, TemplateName, RF_Transactional));
 		}
-		NewCustomBinding->SetObjectTemplate(NewObject<UObject>(NewCustomBinding, SourceBlueprint->GeneratedClass, TemplateName, RF_Transactional));
 	}
 #if WITH_EDITORONLY_DATA
 	else if (UBlueprintGeneratedClass* SourceBlueprintGeneratedClass = Cast<UBlueprintGeneratedClass>(SourceObject))
@@ -500,12 +497,10 @@ UMovieSceneCustomBinding* UMovieSceneSpawnableActorBinding::CreateNewCustomBindi
 
 			NewCustomBinding = CreateBinding();
 
-			if (!NewCustomBinding)
+			if (NewCustomBinding)
 			{
-				return nullptr;
+				NewCustomBinding->SetObjectTemplate(NewObject<UObject>(NewCustomBinding, BlueprintGeneratedBy->GeneratedClass, TemplateName, RF_Transactional));
 			}
-
-			NewCustomBinding->SetObjectTemplate(NewObject<UObject>(NewCustomBinding, BlueprintGeneratedBy->GeneratedClass, TemplateName, RF_Transactional));
 		}
 	}
 #endif
@@ -520,17 +515,18 @@ UMovieSceneCustomBinding* UMovieSceneSpawnableActorBinding::CreateNewCustomBindi
 			}
 
 			NewCustomBinding = CreateBinding();
-
-			if (!NewCustomBinding)
+			if (NewCustomBinding)
 			{
-				return nullptr;
+				NewCustomBinding->SetObjectTemplate(NewObject<UObject>(&OwnerMovieScene, InClass, TemplateName, RF_Transactional));
 			}
-
-			NewCustomBinding->SetObjectTemplate(NewObject<UObject>(&OwnerMovieScene, InClass, TemplateName, RF_Transactional));
 		}
 	}
 
-	NewCustomBinding->BindingName = BindingNameToSet;
+	if (NewCustomBinding)
+	{
+		NewCustomBinding->BindingName = BindingNameToSet;
+	}
+
 	return NewCustomBinding;
 }
 
