Sequencer- Assorted fixes to the binding refactor to better handle creating new custom bindings. Also added the replaceable icon I missed.

diff --git a/Engine/Source/Editor/Sequencer/Private/SequencerUtilities.cpp b/Engine/Source/Editor/Sequencer/Private/SequencerUtilities.cpp
--- a/Engine/Source/Editor/Sequencer/Private/SequencerUtilities.cpp
+++ b/Engine/Source/Editor/Sequencer/Private/SequencerUtilities.cpp
@@ -1314,14 +1314,15 @@ FMovieScenePossessable* FSequencerUtilities::ConvertToCustomBinding(TSharedRef<I
 
 	// If we have an old-style spawnable, use the template as the object to convert instead.
 	FMovieSceneSpawnable* Spawnable = MovieScene->FindSpawnable(BindingGuid);
+	const FMovieSceneBindingReference* PreviousBindingReference = BindingReferences->GetReference(BindingGuid, BindingIndex);
 	if (Spawnable)
 	{
 		ObjectToConvert = Spawnable->GetObjectTemplate();
 		bConvertFromSpawnable = true;
 	}
-	else if (const FMovieSceneBindingReference* Reference = BindingReferences->GetReference(BindingGuid, BindingIndex))
+	else if (PreviousBindingReference)
 	{
-		if (const UMovieSceneCustomBinding* CustomBinding = Reference->CustomBinding)
+		if (const UMovieSceneCustomBinding* CustomBinding = PreviousBindingReference->CustomBinding)
 		{
 			bConvertFromSpawnable = CustomBinding->WillSpawnObject(Sequencer->GetSharedPlaybackState());
 		}
@@ -1335,8 +1336,17 @@ FMovieScenePossessable* FSequencerUtilities::ConvertToCustomBinding(TSharedRef<I
 		return CreatedPossessable;
 	}
 
-	// Create the new custom binding and then we'll slot it in.
-	UMovieSceneCustomBinding* NewCustomBinding = CustomBindingType->GetDefaultObject<UMovieSceneCustomBinding>()->CreateNewCustomBinding(ObjectToConvert, *MovieScene);
+	
+	UMovieSceneCustomBinding* NewCustomBinding = nullptr;
+	if (PreviousBindingReference)
+	{
+		NewCustomBinding = CustomBindingType->GetDefaultObject<UMovieSceneCustomBinding>()->CreateCustomBindingFromBinding(*PreviousBindingReference, ObjectToConvert, *MovieScene);
+	}
+	else
+	{
+		NewCustomBinding = CustomBindingType->GetDefaultObject<UMovieSceneCustomBinding>()->CreateNewCustomBinding(ObjectToConvert, *MovieScene);
+	}
+
 	if (!NewCustomBinding)
 	{
 		return CreatedPossessable;
@@ -2324,8 +2334,11 @@ FGuid TryCreateCustomBinding(TSharedRef<ISequencer> Sequencer, UObject* CustomBi
 	UMovieSceneCustomBinding * NewCustomBinding = nullptr;
 	if (InParams.CustomBinding)
 	{
+		const FMovieSceneBindingReference* PreviousBindingReference = BindingReferences->GetReference(InParams.ReplacementGuid, InParams.BindingIndex);
 		// We've been provided a custom binding pre-created. Ensure it supports the object given
-		if (CustomBindingObject == nullptr || InParams.CustomBinding->SupportsBindingCreationFromObject(CustomBindingObject))
+		if (CustomBindingObject == nullptr 
+			|| InParams.CustomBinding->SupportsBindingCreationFromObject(CustomBindingObject)
+			|| (PreviousBindingReference && InParams.CustomBinding->SupportsConversionFromBinding(*PreviousBindingReference, CustomBindingObject)))
 		{
 			NewCustomBinding = InParams.CustomBinding;
 		}
@@ -2421,8 +2434,6 @@ FGuid CreateGenericBinding(TSharedRef<ISequencer> Sequencer, UObject* InObject,
 
 	if (bAllowCustom)
 	{
-		
-
 		NewBindingID = TryCreateCustomBinding(Sequencer, InObject, BindingReferences, InParams, OwnerMovieScene, bSpawnable, bReplaceable);
 		if (NewBindingID.IsValid())
 		{
diff --git a/Engine/Source/Runtime/MovieScene/Public/Bindings/MovieSceneReplaceableBinding.h b/Engine/Source/Runtime/MovieScene/Public/Bindings/MovieSceneReplaceableBinding.h
--- a/Engine/Source/Runtime/MovieScene/Public/Bindings/MovieSceneReplaceableBinding.h
+++ b/Engine/Source/Runtime/MovieScene/Public/Bindings/MovieSceneReplaceableBinding.h
@@ -77,7 +77,7 @@ protected:
 	UPROPERTY()
 	FString BindingName;
 
-private:
+protected:
 
 	/* UMovieSceneCustomBinding overrides*/
 	bool WillSpawnObject(TSharedRef<const UE::MovieScene::FSharedPlaybackState> SharedPlaybackState) const override final;
diff --git a/Engine/Source/Runtime/MovieSceneTracks/Private/Bindings/MovieSceneSpawnableActorBinding.cpp b/Engine/Source/Runtime/MovieSceneTracks/Private/Bindings/MovieSceneSpawnableActorBinding.cpp
--- a/Engine/Source/Runtime/MovieSceneTracks/Private/Bindings/MovieSceneSpawnableActorBinding.cpp
+++ b/Engine/Source/Runtime/MovieSceneTracks/Private/Bindings/MovieSceneSpawnableActorBinding.cpp
@@ -258,9 +258,9 @@ UWorld* UMovieSceneSpawnableActorBinding::GetWorldContext(TSharedRef<const UE::M
 	return WorldContext;
 }
 
-FName UMovieSceneSpawnableActorBinding::GetSpawnName(const FGuid& BindingId, UMovieScene& MovieScene, FMovieSceneSequenceIDRef TemplateID, TSharedRef<const UE::MovieScene::FSharedPlaybackState> SharedPlaybackState) const
+FName UMovieSceneSpawnableActorBindingBase::GetSpawnName(const FGuid& BindingId, UMovieScene& MovieScene, FMovieSceneSequenceIDRef TemplateID, TSharedRef<const UE::MovieScene::FSharedPlaybackState> SharedPlaybackState) const
 {
-	// We use the net addressable name for spawnables on any non-editor, non-standalone world (ie, all clients, servers and PIE worlds)
+	// We use the net addressable name for spawnable actors on any non-editor, non-standalone world (ie, all clients, servers and PIE worlds)
 
 	UWorld* WorldContext = GetWorldContext(SharedPlaybackState);
 	const bool bUseNetAddressableName = bNetAddressableName && (WorldContext->WorldType != EWorldType::Editor) && (WorldContext->GetNetMode() != ENetMode::NM_Standalone);
@@ -283,15 +283,7 @@ FName UMovieSceneSpawnableActorBinding::GetSpawnName(const FGuid& BindingId, UMo
 	return NAME_None;
 }
 
-TSubclassOf<AActor> UMovieSceneSpawnableActorBinding::GetActorClass() const
-{
-	if (AActor* ActorTemplatePtr = GetActorTemplate())
-	{
-		return ActorTemplatePtr->GetClass();
-	}
-	return nullptr;
-}
-FName UMovieSceneSpawnableActorBinding::GetNetAddressableName(TSharedRef<const UE::MovieScene::FSharedPlaybackState> SharedPlaybackState, const FGuid& BindingId, FMovieSceneSequenceID SequenceID, const FString& BaseName) const
+FName UMovieSceneSpawnableActorBindingBase::GetNetAddressableName(TSharedRef<const UE::MovieScene::FSharedPlaybackState> SharedPlaybackState, const FGuid& BindingId, FMovieSceneSequenceID SequenceID, const FString& BaseName) const
 {
 	UObject* AddressingContext = nullptr;
 
@@ -335,6 +327,16 @@ FName UMovieSceneSpawnableActorBinding::GetNetAddressableName(TSharedRef<const U
 	return FName(AddressableName.Len(), AddressableName.GetData());
 }
 
+TSubclassOf<AActor> UMovieSceneSpawnableActorBinding::GetActorClass() const
+{
+	if (AActor* ActorTemplatePtr = GetActorTemplate())
+	{
+		return ActorTemplatePtr->GetClass();
+	}
+	return nullptr;
+}
+
+
 struct FIsSpawnable
 {
 	FIsSpawnable() : bIsSpawnable(false) {}
diff --git a/Engine/Source/Runtime/MovieSceneTracks/Public/Bindings/MovieSceneSpawnableActorBinding.h b/Engine/Source/Runtime/MovieSceneTracks/Public/Bindings/MovieSceneSpawnableActorBinding.h
--- a/Engine/Source/Runtime/MovieSceneTracks/Public/Bindings/MovieSceneSpawnableActorBinding.h
+++ b/Engine/Source/Runtime/MovieSceneTracks/Public/Bindings/MovieSceneSpawnableActorBinding.h
@@ -75,25 +75,33 @@ protected:
 
 	/* Override to provide Actor class to be spawned*/
 	// TODO: Make UFUNCTION
-	virtual TSubclassOf<AActor> GetActorClass() const PURE_VIRTUAL(UMovieSceneSpawnableActorBindingBase::GetActorClass, return AActor::StaticClass(););
+	MOVIESCENETRACKS_API virtual TSubclassOf<AActor> GetActorClass() const PURE_VIRTUAL(UMovieSceneSpawnableActorBindingBase::GetActorClass, return AActor::StaticClass(););
 
 	/* Optionally override to provide an Actor template to use during Spawn */
-	virtual AActor* GetActorTemplate() const { return nullptr; }
+	MOVIESCENETRACKS_API virtual AActor* GetActorTemplate() const { return nullptr; }
 
 	/* Returns the transform to spawn the actor at*/
-	virtual FTransform GetSpawnTransform() const;
+	MOVIESCENETRACKS_API  FTransform GetSpawnTransform() const;
+
+#if WITH_EDITOR
+	/* MovieSceneCustomBinding overrides*/
+	MOVIESCENETRACKS_API UClass* GetBoundObjectClass() const override { return GetActorClass(); }
+#endif
 
 protected:
 	// UMovieSceneSpawnableBindingBase overrides
 
 	/* Overridden to handle Actor-specific spawning */
-	virtual UObject* SpawnObjectInternal(UWorld* WorldContext, FName SpawnName, const FGuid& BindingId, int32 BindingIndex, UMovieScene& MovieScene, FMovieSceneSequenceIDRef TemplateID, TSharedRef<const UE::MovieScene::FSharedPlaybackState> SharedPlaybackState) override;
+	MOVIESCENETRACKS_API virtual UObject* SpawnObjectInternal(UWorld* WorldContext, FName SpawnName, const FGuid& BindingId, int32 BindingIndex, UMovieScene& MovieScene, FMovieSceneSequenceIDRef TemplateID, TSharedRef<const UE::MovieScene::FSharedPlaybackState> SharedPlaybackState) override;
 
 	/* Overridden to handle Actor-specific destruction*/
-	virtual void DestroySpawnedObjectInternal(UObject* Object) override;
+	MOVIESCENETRACKS_API virtual void DestroySpawnedObjectInternal(UObject* Object) override;
 
+	MOVIESCENETRACKS_API FName GetSpawnName(const FGuid& BindingId, UMovieScene& MovieScene, FMovieSceneSequenceIDRef TemplateID, TSharedRef<const UE::MovieScene::FSharedPlaybackState> SharedPlaybackState) const override;
 
-protected:
+
+private:
+	FName GetNetAddressableName(TSharedRef<const UE::MovieScene::FSharedPlaybackState> SharedPlaybackState, const FGuid& BindingId, FMovieSceneSequenceID SequenceID, const FString& BaseName) const;
 };
 
 /*
@@ -129,31 +137,26 @@ public:
 	bool SupportsConversionFromBinding(const FMovieSceneBindingReference& BindingReference, const UObject* SourceObject) const override;
 	UMovieSceneCustomBinding* CreateCustomBindingFromBinding(const FMovieSceneBindingReference& BindingReference, UObject* SourceObject, UMovieScene& OwnerMovieScene) override;
 	FText GetBindingTypePrettyName() const override;
-	UClass* GetBoundObjectClass() const override { return GetActorClass(); }
 #endif
 
 protected:
 
 	/* MovieSceneSpawnableBindingBase overrides*/
 	UWorld* GetWorldContext(TSharedRef<const UE::MovieScene::FSharedPlaybackState> SharedPlaybackState) const override;
-	FName GetSpawnName(const FGuid& BindingId, UMovieScene& MovieScene, FMovieSceneSequenceIDRef TemplateID, TSharedRef<const UE::MovieScene::FSharedPlaybackState> SharedPlaybackState) const override;
-
+	
 	/* MovieSceneSpawnableActorBindingBase overrides*/
 	TSubclassOf<AActor> GetActorClass() const override;
 	AActor* GetActorTemplate() const override { return ActorTemplate; }
 
-private:
-
-	FName GetNetAddressableName(TSharedRef<const UE::MovieScene::FSharedPlaybackState> SharedPlaybackState, const FGuid& BindingId, FMovieSceneSequenceID SequenceID, const FString& BaseName) const;
-
 
 private:
 
 	UPROPERTY()
-	TObjectPtr<AActor> ActorTemplate;
+	TObjectPtr<AActor> ActorTemplate;	
 
 	UPROPERTY()
 	FString BindingName;
 
+
 };
 
