Sequencer- Fix copy/paste of custom bindings.

diff --git a/Engine/Source/Editor/Sequencer/Private/SequencerUtilities.cpp b/Engine/Source/Editor/Sequencer/Private/SequencerUtilities.cpp
--- a/Engine/Source/Editor/Sequencer/Private/SequencerUtilities.cpp
+++ b/Engine/Source/Editor/Sequencer/Private/SequencerUtilities.cpp
@@ -2191,6 +2191,7 @@ void FSequencerUtilities::CopyBindings(TSharedRef<ISequencer> Sequencer, const T
 								CopyableBinding->SpawnableObjectTemplates.Add(SpawnableBinding->GetObjectTemplate());
 							}
 						}
+						CopyableBinding->CustomBindings.Add(CustomBinding);
 					}
 				}
 			}
diff --git a/Engine/Source/Runtime/MovieSceneTracks/Private/Bindings/MovieSceneSpawnableActorBinding.cpp b/Engine/Source/Runtime/MovieSceneTracks/Private/Bindings/MovieSceneSpawnableActorBinding.cpp
--- a/Engine/Source/Runtime/MovieSceneTracks/Private/Bindings/MovieSceneSpawnableActorBinding.cpp
+++ b/Engine/Source/Runtime/MovieSceneTracks/Private/Bindings/MovieSceneSpawnableActorBinding.cpp
@@ -352,9 +352,11 @@ void UMovieSceneSpawnableActorBinding::SetObjectTemplate(UObject* InObjectTempla
 	ensure(InObjectTemplate == nullptr || InObjectTemplate->IsA<AActor>());
 	checkf(InObjectTemplate == nullptr || !InObjectTemplate->HasAnyFlags(RF_ClassDefaultObject), TEXT("Setting CDOs as object templates is not supported. Please use the class directly."));
 	ActorTemplate = Cast<AActor>(InObjectTemplate);
-
-	// TODO: We should move this out of FMovieSceneSpawnable eventually
-	FMovieSceneSpawnable::MarkSpawnableTemplate(*ActorTemplate);
+	if (ActorTemplate)
+	{
+		// TODO: We should move this out of FMovieSceneSpawnable eventually
+		FMovieSceneSpawnable::MarkSpawnableTemplate(*ActorTemplate);
+	}
 	AutoSetNetAddressableName();
 }
 
