Sequencer: handle multiple instances of a sequence playing audio at the same time.

diff --git a/Engine/Source/Runtime/MovieSceneTracks/Private/Systems/MovieSceneAudioSystem.cpp b/Engine/Source/Runtime/MovieSceneTracks/Private/Systems/MovieSceneAudioSystem.cpp
--- a/Engine/Source/Runtime/MovieSceneTracks/Private/Systems/MovieSceneAudioSystem.cpp
+++ b/Engine/Source/Runtime/MovieSceneTracks/Private/Systems/MovieSceneAudioSystem.cpp
@@ -112,6 +112,7 @@ enum class EAudioEvaluationType
 
 struct FGatherAudioInputs
 {
+	using FInstanceObjectKey = UMovieSceneAudioSystem::FInstanceObjectKey;
 	using FAudioInputsBySectionKey = UMovieSceneAudioSystem::FAudioInputsBySectionKey;
 	using FAudioComponentInputEvaluationData = UMovieSceneAudioSystem::FAudioComponentInputEvaluationData;
 
@@ -124,6 +125,7 @@ struct FGatherAudioInputs
 
 	void ForEachAllocation(
 		const FEntityAllocation* Allocation,
+		TRead<FInstanceHandle> InstanceHandles,
 		TRead<FMovieSceneAudioComponentData> AudioDatas,
 		TRead<FMovieSceneAudioInputData> AudioInputDatas,
 		TReadOneOrMoreOf<
@@ -156,7 +158,7 @@ struct FGatherAudioInputs
 			const FMovieSceneAudioComponentData& AudioData = AudioDatas[Index];
 			const FMovieSceneAudioInputData& AudioInputNames = AudioInputDatas[Index];
 
-			FObjectKey SectionKey(AudioData.Section);
+			FInstanceObjectKey SectionKey(InstanceHandles[Index], FObjectKey(AudioData.Section));
 			FAudioComponentInputEvaluationData& AudioInputValues = AudioInputsBySectionKey.FindOrAdd(SectionKey);
 			
 			// Gather float inputs.
@@ -191,6 +193,7 @@ struct FGatherAudioInputs
 
 struct FGatherAudioTriggers
 {
+	using FInstanceObjectKey = UMovieSceneAudioSystem::FInstanceObjectKey;
 	using FAudioInputsBySectionKey = UMovieSceneAudioSystem::FAudioInputsBySectionKey;
 	using FAudioComponentInputEvaluationData = UMovieSceneAudioSystem::FAudioComponentInputEvaluationData;
 
@@ -203,6 +206,7 @@ struct FGatherAudioTriggers
 
 	void ForEachAllocation(
 		const FEntityAllocation* Allocation,
+		TRead<FInstanceHandle> InstanceHandles,
 		TRead<FMovieSceneAudioComponentData> AudioDatas,
 		TRead<FName> AudioTriggerNames) const
 	{
@@ -213,7 +217,7 @@ struct FGatherAudioTriggers
 			const FMovieSceneAudioComponentData& AudioData = AudioDatas[Index];
 			const FName& AudioTriggerName = AudioTriggerNames[Index];
 
-			FObjectKey SectionKey(AudioData.Section);
+			FInstanceObjectKey SectionKey(InstanceHandles[Index], FObjectKey(AudioData.Section));
 			FAudioComponentInputEvaluationData& AudioInputValues = AudioInputsBySectionKey.FindOrAdd(SectionKey);
 
 			AudioInputValues.Inputs_Trigger.Add(AudioTriggerName);
@@ -314,17 +318,18 @@ private:
 			return;
 		}
 
+		FInstanceHandle InstanceHandle(Instance.GetInstanceHandle());
 		FObjectKey ActorKey(BoundObject);
 		FObjectKey SectionKey(AudioSection);
 
 		const EAudioEvaluationType EvalType = GetAudioEvaluationType(Context);
 		if (EvalType == EAudioEvaluationType::StopAndPlay)
 		{
-			AudioSystem->StopSound(ActorKey, AudioData.Section);
+			AudioSystem->StopSound(InstanceHandle, ActorKey, AudioData.Section);
 		}
 		else if (EvalType == EAudioEvaluationType::Stop)
 		{
-			AudioSystem->StopSound(ActorKey, AudioData.Section);
+			AudioSystem->StopSound(InstanceHandle, ActorKey, AudioData.Section);
 			return;
 		}
 		else if (EvalType == EAudioEvaluationType::Skip)
@@ -360,12 +365,12 @@ private:
 				}
 			}
 
-			FAudioComponentEvaluationData* EvaluationData = AudioSystem->GetAudioComponentEvaluationData(FObjectKey(), SectionKey);
+			FAudioComponentEvaluationData* EvaluationData = AudioSystem->GetAudioComponentEvaluationData(InstanceHandle, FObjectKey(), SectionKey);
 			if (!EvaluationData)
 			{
 				// Initialize the sound
 				UWorld* World = PlaybackContext ? PlaybackContext->GetWorld() : nullptr;
-				EvaluationData = AudioSystem->AddRootAudioComponent(AudioSection, World);
+				EvaluationData = AudioSystem->AddRootAudioComponent(InstanceHandle, AudioSection, World);
 				UAudioComponent* AudioComponent = EvaluationData ? EvaluationData->AudioComponent.Get() : nullptr;
 
 				if (ensure(AudioComponent))
@@ -403,18 +408,18 @@ private:
 				EvaluationData->VolumeMultiplier = VolumeMultiplier;
 				EvaluationData->PitchMultiplier = PitchMultiplier;
 
-				EnsureAudioIsPlaying(nullptr, *AudioSection, *EvaluationData, Context, *Player);
+				EnsureAudioIsPlaying(nullptr, InstanceHandle, *AudioSection, *EvaluationData, Context, *Player);
 			}
 		}
 
 		// Object binding audio track
 		else
 		{
-			FAudioComponentEvaluationData* EvaluationData = AudioSystem->GetAudioComponentEvaluationData(ActorKey, SectionKey);
+			FAudioComponentEvaluationData* EvaluationData = AudioSystem->GetAudioComponentEvaluationData(InstanceHandle, ActorKey, SectionKey);
 			if (!EvaluationData)
 			{
 				// Initialize the sound
-				EvaluationData = AudioSystem->AddBoundObjectAudioComponent(AudioSection, BoundObject);
+				EvaluationData = AudioSystem->AddBoundObjectAudioComponent(InstanceHandle, AudioSection, BoundObject);
 				UAudioComponent* AudioComponent = EvaluationData ? EvaluationData->AudioComponent.Get() : nullptr;
 
 				if (AudioComponent)
@@ -444,18 +449,20 @@ private:
 				EvaluationData->VolumeMultiplier = VolumeMultiplier;
 				EvaluationData->PitchMultiplier = PitchMultiplier;
 
-				EnsureAudioIsPlaying(BoundObject, *AudioSection, *EvaluationData, Context, *Player);
+				EnsureAudioIsPlaying(BoundObject, InstanceHandle, *AudioSection, *EvaluationData, Context, *Player);
 			}
 		}
 	}
 
 	void EnsureAudioIsPlaying(
 			UObject* BoundObject,
+			FInstanceHandle InstanceHandle,
 			UMovieSceneAudioSection& AudioSection,
 			FAudioComponentEvaluationData& EvaluationData,
 			const FMovieSceneContext& Context, 
 			IMovieScenePlayer& Player) const
 	{
+		using FInstanceObjectKey = UMovieSceneAudioSystem::FInstanceObjectKey;
 		using FAudioInputsBySectionKey = UMovieSceneAudioSystem::FAudioInputsBySectionKey;
 		using FAudioComponentInputEvaluationData = UMovieSceneAudioSystem::FAudioComponentInputEvaluationData;
 
@@ -483,7 +490,7 @@ private:
 
 		// Apply the input params.
 		FAudioInputsBySectionKey& AudioInputsBySectionKey = AudioSystem->AudioInputsBySectionKey;
-		FObjectKey SectionKey(&AudioSection);
+		FInstanceObjectKey SectionKey(InstanceHandle, FObjectKey(&AudioSection));
 		FAudioComponentInputEvaluationData* AudioInputs = AudioInputsBySectionKey.Find(SectionKey);
 		if (AudioInputs)
 		{
@@ -508,6 +515,7 @@ private:
 
 			if (!AudioSection.GetLooping() && AudioTime > Duration && Duration != 0.f)
 			{
+				UE_LOG(LogMovieScene, Verbose, TEXT("Audio Component reached end of playback. Component: %s Sound: %s"), *AudioComponent.GetName(), *GetNameSafe(AudioComponent.Sound));
 				AudioComponent.Stop();
 				return;
 			}
@@ -518,7 +526,8 @@ private:
 		// If the audio component is not playing we (may) need a state change. If the audio component is playing
 		// the wrong sound then we need a state change. If the audio playback time is significantly out of sync 
 		// with the desired time then we need a state change.
-		bool bSoundNeedsStateChange = !AudioComponent.IsPlaying() || AudioComponent.Sound != Sound;
+		const bool bSoundsNeedPlaying = !AudioComponent.IsPlaying();
+		const bool bSoundNeedsStateChange =  AudioComponent.Sound != Sound;
 		bool bSoundNeedsTimeSync = false;
 
 		UObject* PlaybackContext = Player.GetPlaybackContext();
@@ -555,14 +564,31 @@ private:
 
 				if (!FMath::IsNearlyZero(MaxSequenceAudioDesyncToleranceCVar) && FMath::Abs(Desync) > MaxSequenceAudioDesyncToleranceCVar)
 				{
-					UE_LOG(LogMovieScene, Verbose, TEXT("Audio Component detected a significant mismatch in (assumed) playback time versus the desired time. Desync: %6.2f(s) Desired Time: %6.2f(s). Component: %s sound: %s"), Desync, AudioTime, *AudioComponent.GetName(), *GetNameSafe(AudioComponent.Sound));
+					UE_LOG(LogMovieScene, Verbose, TEXT("Audio Component detected a significant mismatch in (assumed) playback time versus the desired time. Desync: %6.2f(s) Desired Time: %6.2f(s). Component: %s Sound: %s"), Desync, AudioTime, *AudioComponent.GetName(), *GetNameSafe(AudioComponent.Sound));
 					bSoundNeedsTimeSync = true;
 				}
 			}
 		}
 
-		if (bSoundNeedsStateChange || bSoundNeedsTimeSync)
+		if (bSoundsNeedPlaying || bSoundNeedsStateChange || bSoundNeedsTimeSync)
 		{
+#if !NO_LOGGING
+			FString ReasonMessage;
+			if (bSoundsNeedPlaying)
+			{
+				ReasonMessage += TEXT("playing");
+			}
+			else if (bSoundNeedsStateChange)
+			{
+				ReasonMessage += TEXT("state change");
+			}
+			else
+			{
+				ReasonMessage += TEXT("time sync");
+			}
+			UE_LOG(LogMovieScene, Verbose, TEXT("Audio component needs %s. Component: %s"), *ReasonMessage, *AudioComponent.GetName());
+#endif
+
 			AudioComponent.bAllowSpatialization = bAllowSpatialization;
 
 			if (AudioSection.GetOverrideAttenuation())
@@ -574,7 +600,7 @@ private:
 			// stop calls when a sound cue with a duration of zero is played.
 			if (AudioComponent.IsPlaying() || bSoundNeedsTimeSync)
 			{
-				UE_LOG(LogMovieScene, Verbose, TEXT("Audio Component stopped due to needing a state change bIsPlaying: %d bNeedsTimeSync: %d. Component: %s sound: %s"), AudioComponent.IsPlaying(), bSoundNeedsTimeSync, *AudioComponent.GetName(), *GetNameSafe(AudioComponent.Sound));
+				UE_LOG(LogMovieScene, Verbose, TEXT("Audio Component stopped due to needing a state change bIsPlaying: %d bNeedsTimeSync: %d. Component: %s Sound: %s"), AudioComponent.IsPlaying(), bSoundNeedsTimeSync, *AudioComponent.GetName(), *GetNameSafe(AudioComponent.Sound));
 				AudioComponent.Stop();
 			}
 
@@ -598,7 +624,7 @@ private:
 
 			if (AudioTime >= 0.f)
 			{
-				UE_LOG(LogMovieScene, Verbose, TEXT("Audio Component Play at Local Time: %6.2f CurrentTime: %6.2f(s) SectionStart: %6.2f(s), SoundDur: %6.2f OffsetIntoClip: %6.2f sound: %s"), AudioTime, (Context.GetTime() / Context.GetFrameRate()), SectionStartTimeSeconds, AudioComponent.Sound ? AudioComponent.Sound->GetDuration() : 0.0f, (float)Context.GetFrameRate().AsSeconds(AudioStartOffset), *GetNameSafe(AudioComponent.Sound));
+				UE_LOG(LogMovieScene, Verbose, TEXT("Audio Component Play at Local Time: %6.2f CurrentTime: %6.2f(s) SectionStart: %6.2f(s), SoundDur: %6.2f OffsetIntoClip: %6.2f Component: %s Sound: %s"), AudioTime, (Context.GetTime() / Context.GetFrameRate()), SectionStartTimeSeconds, AudioComponent.Sound ? AudioComponent.Sound->GetDuration() : 0.0f, (float)Context.GetFrameRate().AsSeconds(AudioStartOffset), *AudioComponent.GetName(), *GetNameSafe(AudioComponent.Sound));
 				AudioComponent.Play(AudioTime);
 
 				// Keep track of when we asked this audio clip to play (in game time) so that we can figure 
@@ -726,13 +752,13 @@ void UMovieSceneAudioSystem::OnUnlink()
 
 	for (const TPair<FObjectKey, FAudioComponentBySectionKey>& AudioComponentsForActor : AudioComponentsByActorKey)
 	{
-		for (const TPair<FObjectKey, FAudioComponentEvaluationData>& AudioComponentForSection : AudioComponentsForActor.Value)
+		for (const TPair<FInstanceObjectKey, FAudioComponentEvaluationData>& AudioComponentForSection : AudioComponentsForActor.Value)
 		{
 			UAudioComponent* AudioComponent = AudioComponentForSection.Value.AudioComponent.Get();
 			if (AudioComponent)
 			{
 				UObject* Actor = AudioComponentsForActor.Key.ResolveObjectPtr();
-				UObject* Section = AudioComponentForSection.Key.ResolveObjectPtr();
+				UObject* Section = AudioComponentForSection.Key.Value.ResolveObjectPtr();
 				UE_LOG(LogMovieScene, Warning, TEXT("Cleaning audio component '%s' for section '%s' on actor '%s'"),
 						*AudioComponent->GetPathName(),
 						Section ? *Section->GetPathName() : TEXT("<null>"),
@@ -751,7 +777,7 @@ void UMovieSceneAudioSystem::ResetSharedData()
 	AudioInputsBySectionKey.Reset();
 	for (TPair<FObjectKey, FAudioComponentBySectionKey>& AudioComponentsForActor : AudioComponentsByActorKey)
 	{
-		for (TPair<FObjectKey, FAudioComponentEvaluationData>& AudioComponentForSection : AudioComponentsForActor.Value)
+		for (TPair<FInstanceObjectKey, FAudioComponentEvaluationData>& AudioComponentForSection : AudioComponentsForActor.Value)
 		{
 			AudioComponentForSection.Value.bEvaluatedThisFrame = false;
 		}
@@ -775,6 +801,7 @@ void UMovieSceneAudioSystem::OnSchedulePersistentTasks(UE::MovieScene::IEntitySy
 
 	// Gather audio input values computed by the channel evaluators.
 	FTaskID GatherInputsTask = FEntityTaskBuilder()
+	.Read(BuiltInComponents->InstanceHandle)
 	.Read(TrackComponents->Audio)
 	.Read(TrackComponents->AudioInputs)
 	.ReadOneOrMoreOf(
@@ -790,6 +817,7 @@ void UMovieSceneAudioSystem::OnSchedulePersistentTasks(UE::MovieScene::IEntitySy
 
 	// Gather up audio triggers
 	FTaskID GatherTriggersTask = FEntityTaskBuilder()
+	.Read(BuiltInComponents->InstanceHandle)
 	.Read(TrackComponents->Audio)
 	.Read(TrackComponents->AudioTriggerName)
 	.Schedule_PerAllocation<FGatherAudioTriggers>(&Linker->EntityManager, TaskScheduler, this);
@@ -834,6 +862,7 @@ void UMovieSceneAudioSystem::OnRun(FSystemTaskPrerequisites& InPrerequisites, FS
 	FSystemTaskPrerequisites Prereqs;
 
 	FGraphEventRef Task = FEntityTaskBuilder()
+	.Read(BuiltInComponents->InstanceHandle)
 	.Read(TrackComponents->Audio)
 	.Read(TrackComponents->AudioInputs)
 	.ReadOneOrMoreOf(
@@ -850,6 +879,7 @@ void UMovieSceneAudioSystem::OnRun(FSystemTaskPrerequisites& InPrerequisites, FS
 	}
 
 	Task = FEntityTaskBuilder()
+	.Read(BuiltInComponents->InstanceHandle)
 	.Read(TrackComponents->Audio)
 	.Read(TrackComponents->AudioTriggerName)
 	.template Dispatch_PerAllocation<FGatherAudioTriggers>(&Linker->EntityManager, InPrerequisites, nullptr, this);
@@ -871,13 +901,14 @@ void UMovieSceneAudioSystem::OnRun(FSystemTaskPrerequisites& InPrerequisites, FS
 	.template Dispatch_PerAllocation<FEvaluateAudio>(&Linker->EntityManager, Prereqs, &Subsequents, this);
 }
 
-UMovieSceneAudioSystem::FAudioComponentEvaluationData* UMovieSceneAudioSystem::GetAudioComponentEvaluationData(FObjectKey ActorKey, FObjectKey SectionKey)
+UMovieSceneAudioSystem::FAudioComponentEvaluationData* UMovieSceneAudioSystem::GetAudioComponentEvaluationData(FInstanceHandle InstanceHandle, FObjectKey ActorKey, FObjectKey SectionKey)
 {
 	FAudioComponentBySectionKey* Map = AudioComponentsByActorKey.Find(ActorKey);
 	if (Map != nullptr)
 	{
 		// First, check for an exact match for this entity
-		FAudioComponentEvaluationData* ExistingData = Map->Find(SectionKey);
+		FInstanceObjectKey DataKey{ InstanceHandle, SectionKey };
+		FAudioComponentEvaluationData* ExistingData = Map->Find(DataKey);
 		if (ExistingData != nullptr)
 		{
 			if (ExistingData->AudioComponent.IsValid())
@@ -896,7 +927,7 @@ UMovieSceneAudioSystem::FAudioComponentEvaluationData* UMovieSceneAudioSystem::G
 				FAudioComponentEvaluationData MovedData(Pair.Value);
 				Map->Remove(Pair.Key);
 				MovedData.PartialDesyncComputation.Reset();
-				return &Map->Add(SectionKey, MovedData);
+				return &Map->Add(DataKey, MovedData);
 			}
 		}
 	}
@@ -904,7 +935,7 @@ UMovieSceneAudioSystem::FAudioComponentEvaluationData* UMovieSceneAudioSystem::G
 	return nullptr;
 }
 
-UMovieSceneAudioSystem::FAudioComponentEvaluationData* UMovieSceneAudioSystem::AddBoundObjectAudioComponent(UMovieSceneAudioSection* Section, UObject* PrincipalObject)
+UMovieSceneAudioSystem::FAudioComponentEvaluationData* UMovieSceneAudioSystem::AddBoundObjectAudioComponent(FInstanceHandle InstanceHandle, UMovieSceneAudioSection* Section, UObject* PrincipalObject)
 {
 	using namespace UE::MovieScene;
 
@@ -913,7 +944,7 @@ UMovieSceneAudioSystem::FAudioComponentEvaluationData* UMovieSceneAudioSystem::A
 
 	FAudioComponentBySectionKey& ActorAudioComponentMap = AudioComponentsByActorKey.FindOrAdd(ObjectKey);
 
-	FAudioComponentEvaluationData* ExistingData = GetAudioComponentEvaluationData(ObjectKey, SectionKey);
+	FAudioComponentEvaluationData* ExistingData = GetAudioComponentEvaluationData(InstanceHandle, ObjectKey, SectionKey);
 	if (!ExistingData)
 	{
 		USoundCue* TempPlaybackAudioCue = NewObject<USoundCue>();
@@ -961,14 +992,15 @@ UMovieSceneAudioSystem::FAudioComponentEvaluationData* UMovieSceneAudioSystem::A
 		NewComponent->SetFlags(RF_Transient);
 		NewComponent->AttachToComponent(SceneComponent, FAttachmentTransformRules::KeepRelativeTransform);
 
-		ExistingData = &ActorAudioComponentMap.Add(SectionKey);
+		FInstanceObjectKey DataKey{ InstanceHandle, SectionKey };
+		ExistingData = &ActorAudioComponentMap.Add(DataKey);
 		ExistingData->AudioComponent = NewComponent;
 	}
 
 	return ExistingData;
 }
 
-UMovieSceneAudioSystem::FAudioComponentEvaluationData* UMovieSceneAudioSystem::AddRootAudioComponent(UMovieSceneAudioSection* Section, UWorld* World)
+UMovieSceneAudioSystem::FAudioComponentEvaluationData* UMovieSceneAudioSystem::AddRootAudioComponent(FInstanceHandle InstanceHandle, UMovieSceneAudioSection* Section, UWorld* World)
 {
 	using namespace UE::MovieScene;
 
@@ -977,7 +1009,7 @@ UMovieSceneAudioSystem::FAudioComponentEvaluationData* UMovieSceneAudioSystem::A
 
 	FAudioComponentBySectionKey& RootAudioComponentMap = AudioComponentsByActorKey.FindOrAdd(NullKey);
 
-	FAudioComponentEvaluationData* ExistingData = GetAudioComponentEvaluationData(NullKey, SectionKey);
+	FAudioComponentEvaluationData* ExistingData = GetAudioComponentEvaluationData(InstanceHandle, NullKey, SectionKey);
 	if (!ExistingData)
 	{
 		USoundCue* TempPlaybackAudioCue = NewObject<USoundCue>();
@@ -995,18 +1027,20 @@ UMovieSceneAudioSystem::FAudioComponentEvaluationData* UMovieSceneAudioSystem::A
 
 		NewComponent->SetFlags(RF_Transient);
 
-		ExistingData = &RootAudioComponentMap.Add(SectionKey);
+		FInstanceObjectKey DataKey{ InstanceHandle, SectionKey };
+		ExistingData = &RootAudioComponentMap.Add(DataKey);
 		ExistingData->AudioComponent = NewComponent;
 	}
 
 	return ExistingData;
 }
 
-void UMovieSceneAudioSystem::StopSound(FObjectKey ActorKey, FObjectKey SectionKey)
+void UMovieSceneAudioSystem::StopSound(FInstanceHandle InstanceHandle, FObjectKey ActorKey, FObjectKey SectionKey)
 {
 	if (FAudioComponentBySectionKey* Map = AudioComponentsByActorKey.Find(ActorKey))
 	{
-		if (FAudioComponentEvaluationData* Data = Map->Find(SectionKey))
+		FInstanceObjectKey DataKey{ InstanceHandle, SectionKey };
+		if (FAudioComponentEvaluationData* Data = Map->Find(DataKey))
 		{
 			if (UAudioComponent* AudioComponent = Data->AudioComponent.Get())
 			{
diff --git a/Engine/Source/Runtime/MovieSceneTracks/Private/Systems/MovieSceneAudioSystem.h b/Engine/Source/Runtime/MovieSceneTracks/Private/Systems/MovieSceneAudioSystem.h
--- a/Engine/Source/Runtime/MovieSceneTracks/Private/Systems/MovieSceneAudioSystem.h
+++ b/Engine/Source/Runtime/MovieSceneTracks/Private/Systems/MovieSceneAudioSystem.h
@@ -4,6 +4,7 @@
 
 #include "CoreTypes.h"
 #include "EntitySystem/MovieSceneEntityIDs.h"
+#include "EntitySystem/MovieSceneSequenceInstanceHandle.h"
 #include "MovieSceneTracksComponentTypes.h"
 #include "UObject/ObjectMacros.h"
 #include "Channels/MovieSceneAudioTriggerChannel.h"
@@ -65,6 +66,7 @@ class UMovieSceneAudioSystem : public UMovieSceneEntitySystem
 
 public:
 
+	using FInstanceHandle = UE::MovieScene::FInstanceHandle;
 	using FMovieSceneEntityID = UE::MovieScene::FMovieSceneEntityID;
 	using FAudioComponentEvaluationData = UE::MovieScene::FAudioComponentEvaluationData;
 	using FAudioComponentInputEvaluationData = UE::MovieScene::FAudioComponentInputEvaluationData;
@@ -80,24 +82,24 @@ public:
 	/**
 	 * Get the evaluation data for the given actor and section. Pass a null actor key for root (world) audio.
 	 */
-	FAudioComponentEvaluationData* GetAudioComponentEvaluationData(FObjectKey ActorKey, FObjectKey SectionKey);
+	FAudioComponentEvaluationData* GetAudioComponentEvaluationData(FInstanceHandle InstanceHandle, FObjectKey ActorKey, FObjectKey SectionKey);
 
 	/**
 	 * Adds an audio component to the given bound sequencer object.
 	 * WARNING: Only to be called on the game thread.
 	 */
-	FAudioComponentEvaluationData* AddBoundObjectAudioComponent(UMovieSceneAudioSection* Section, UObject* PrincipalObject);
+	FAudioComponentEvaluationData* AddBoundObjectAudioComponent(FInstanceHandle InstanceHandle, UMovieSceneAudioSection* Section, UObject* PrincipalObject);
 
 	/**
 	 * Adds an audio component to the world, for playing root audio tracks.
 	 * WARNING: Only to be called on the game thread.
 	 */
-	FAudioComponentEvaluationData* AddRootAudioComponent(UMovieSceneAudioSection* Section, UWorld* World);
+	FAudioComponentEvaluationData* AddRootAudioComponent(FInstanceHandle InstanceHandle, UMovieSceneAudioSection* Section, UWorld* World);
 
 	/**
 	 * Stop the audio on the audio component associated with the given audio section.
 	 */
-	void StopSound(FObjectKey ActorKey, FObjectKey SectionKey);
+	void StopSound(FInstanceHandle InstanceHandle, FObjectKey ActorKey, FObjectKey SectionKey);
 
 	/**
 	 * Reset shared accumulation data required every evaluation frame
@@ -106,13 +108,15 @@ public:
 
 private:
 
+	using FInstanceObjectKey = TTuple<FInstanceHandle, FObjectKey>;
+
 	/** Map of all created audio components */
-	using FAudioComponentBySectionKey = TMap<FObjectKey, FAudioComponentEvaluationData>;
+	using FAudioComponentBySectionKey = TMap<FInstanceObjectKey, FAudioComponentEvaluationData>;
 	using FAudioComponentsByActorKey = TMap<FObjectKey, FAudioComponentBySectionKey>;
 	FAudioComponentsByActorKey AudioComponentsByActorKey;
 
 	/** Map of audio input values, rebuilt every frame */
-	using FAudioInputsBySectionKey = TMap<FObjectKey, FAudioComponentInputEvaluationData>;
+	using FAudioInputsBySectionKey = TMap<FInstanceObjectKey, FAudioComponentInputEvaluationData>;
 	FAudioInputsBySectionKey AudioInputsBySectionKey;
 
 	/** Pre-animated state */
