Sequencer: fix FlushSingleEvaluationPhase for custom task scheduling

diff --git a/Engine/Source/Runtime/MovieScene/Private/EntitySystem/MovieSceneEntitySystemRunner.cpp b/Engine/Source/Runtime/MovieScene/Private/EntitySystem/MovieSceneEntitySystemRunner.cpp
--- a/Engine/Source/Runtime/MovieScene/Private/EntitySystem/MovieSceneEntitySystemRunner.cpp
+++ b/Engine/Source/Runtime/MovieScene/Private/EntitySystem/MovieSceneEntitySystemRunner.cpp
@@ -1335,6 +1335,11 @@ bool FMovieSceneEntitySystemRunner::FlushSingleEvaluationPhase()
 
 	Linker->EntityManager.LockDown();
 
+	if (FEntitySystemScheduler::IsCustomSchedulingEnabled())
+	{
+		Linker->SystemGraph.ScheduleTasks(&Linker->EntityManager);
+	}
+
 	FGraphEventArray AllTasks;
 	Linker->SystemGraph.ExecutePhase(ESystemPhase::Evaluation, Linker, AllTasks);
 
diff --git a/Engine/Source/Runtime/MovieSceneTracks/Private/Systems/MovieSceneMotionVectorSimulationSystem.cpp b/Engine/Source/Runtime/MovieSceneTracks/Private/Systems/MovieSceneMotionVectorSimulationSystem.cpp
--- a/Engine/Source/Runtime/MovieSceneTracks/Private/Systems/MovieSceneMotionVectorSimulationSystem.cpp
+++ b/Engine/Source/Runtime/MovieSceneTracks/Private/Systems/MovieSceneMotionVectorSimulationSystem.cpp
@@ -102,6 +102,10 @@ void UMovieSceneMotionVectorSimulationSystem::ComputeSimulatedMotion()
 		EvalTime->Disable();
 	}
 
+	// Simulate a structure change to force the persistent tasks to be re-scheduled without
+	// the systems we just disabled.
+	Linker->EntityManager.MimicStructureChanged();
+
 	FEntityComponentFilter* GlobalFilter     = &Linker->EntityManager.ModifyGlobalIterationFilter();
 	FEntityComponentFilter  GlobalFilterCopy = *GlobalFilter;
 
@@ -192,6 +196,9 @@ void UMovieSceneMotionVectorSimulationSystem::ComputeSimulatedMotion()
 	}
 	Enable();
 
+	// Force persistent tasks to be re-created again (see previous comment).
+	Linker->EntityManager.MimicStructureChanged();
+
 	bSimulateTransformsRequested = false;
 }
 
