Sequencer- Fix binding refactor issues with naming spawnables.
* We were no longer making unique names for spawnables when creating new ones
* We weren't properly synchronizing actor labels with the binding names.

diff --git a/Engine/Source/Editor/Sequencer/Private/SequencerUtilities.cpp b/Engine/Source/Editor/Sequencer/Private/SequencerUtilities.cpp
--- a/Engine/Source/Editor/Sequencer/Private/SequencerUtilities.cpp
+++ b/Engine/Source/Editor/Sequencer/Private/SequencerUtilities.cpp
@@ -2384,6 +2384,7 @@ FGuid TryCreateCustomBinding(TSharedRef<ISequencer> Sequencer, UObject* CustomBi
 	{
 		FString DesiredBindingName = NewCustomBinding->GetDesiredBindingName();
 		FString CurrentName = DesiredBindingName.IsEmpty() ? InParams.BindingNameOverride : DesiredBindingName;
+		CurrentName = MovieSceneHelpers::MakeUniqueBindingName(OwnerMovieScene, CurrentName);
 
 		FMovieScenePossessable* NewPossessable = nullptr;
 		FGuid NewID;
diff --git a/Engine/Source/Runtime/MovieScene/Private/Bindings/MovieSceneReplaceableBinding.cpp b/Engine/Source/Runtime/MovieScene/Private/Bindings/MovieSceneReplaceableBinding.cpp
--- a/Engine/Source/Runtime/MovieScene/Private/Bindings/MovieSceneReplaceableBinding.cpp
+++ b/Engine/Source/Runtime/MovieScene/Private/Bindings/MovieSceneReplaceableBinding.cpp
@@ -163,13 +163,10 @@ UMovieSceneCustomBinding* UMovieSceneReplaceableBindingBase::CreateNewCustomBind
 
 	UMovieSceneReplaceableBindingBase* NewCustomBinding = nullptr;
 
-	FString BindingNameToSet = MovieSceneHelpers::MakeUniqueBindingName(&OwnerMovieScene, FName::NameToDisplayString(SourceObject->GetName(), false));
-
 	const FName TemplateName = MakeUniqueObjectName(&OwnerMovieScene, UObject::StaticClass(), SourceObject->GetFName());
 	const FName InstancedBindingName = MakeUniqueObjectName(&OwnerMovieScene, UObject::StaticClass(), *FString(TemplateName.ToString() + TEXT("_CustomBinding")));
 
 	NewCustomBinding = NewObject<UMovieSceneReplaceableBindingBase>(&OwnerMovieScene, GetClass(), InstancedBindingName, RF_Transactional);
-	NewCustomBinding->BindingName = BindingNameToSet;
 #if WITH_EDITORONLY_DATA
 	NewCustomBinding->PreviewSpawnable = NewCustomBinding->CreateInnerSpawnable(SourceObject, OwnerMovieScene);
 #endif
diff --git a/Engine/Source/Runtime/MovieScene/Public/Bindings/MovieSceneReplaceableBinding.h b/Engine/Source/Runtime/MovieScene/Public/Bindings/MovieSceneReplaceableBinding.h
--- a/Engine/Source/Runtime/MovieScene/Public/Bindings/MovieSceneReplaceableBinding.h
+++ b/Engine/Source/Runtime/MovieScene/Public/Bindings/MovieSceneReplaceableBinding.h
@@ -74,16 +74,12 @@ protected:
 	TObjectPtr<UMovieSceneSpawnableBindingBase> PreviewSpawnable = nullptr;
 #endif
 
-	UPROPERTY()
-	FString BindingName;
-
 protected:
 
 	/* UMovieSceneCustomBinding overrides*/
 	bool WillSpawnObject(TSharedRef<const UE::MovieScene::FSharedPlaybackState> SharedPlaybackState) const override final;
 	FMovieSceneBindingResolveResult ResolveBinding(const FMovieSceneBindingResolveParams& ResolveParams, int32 BindingIndex, TSharedRef<const UE::MovieScene::FSharedPlaybackState> SharedPlaybackState) const override final;
 	const UMovieSceneSpawnableBindingBase* AsSpawnable(TSharedRef<const UE::MovieScene::FSharedPlaybackState> SharedPlaybackState) const override final;
-	FString GetDesiredBindingName() const override { return BindingName; }
 
 	/* 
 	*  Note that we choose to implement CreateCustomBinding here rather than in subclasses.
diff --git a/Engine/Source/Runtime/MovieSceneTracks/Private/Bindings/MovieSceneSpawnableActorBinding.cpp b/Engine/Source/Runtime/MovieSceneTracks/Private/Bindings/MovieSceneSpawnableActorBinding.cpp
--- a/Engine/Source/Runtime/MovieSceneTracks/Private/Bindings/MovieSceneSpawnableActorBinding.cpp
+++ b/Engine/Source/Runtime/MovieSceneTracks/Private/Bindings/MovieSceneSpawnableActorBinding.cpp
@@ -128,7 +128,12 @@ UObject* UMovieSceneSpawnableActorBindingBase::SpawnObjectInternal(UWorld* World
 	if (WorldContext->WorldType == EWorldType::Editor)
 	{
 		FString BindingName = GetDesiredBindingName();
-		SpawnedActor->SetActorLabel(!BindingName.IsEmpty() ? BindingName : SpawnName.ToString());
+		FString ActorLabel = !BindingName.IsEmpty() ? BindingName : SpawnName.ToString();
+		if (FMovieScenePossessable* Possessable = MovieScene.FindPossessable(BindingId))
+		{
+			ActorLabel = Possessable->GetName();
+		}
+		SpawnedActor->SetActorLabel(ActorLabel);
 	}
 #endif
 
@@ -431,8 +436,6 @@ UMovieSceneCustomBinding* UMovieSceneSpawnableActorBinding::CreateNewCustomBindi
 
 	UMovieSceneSpawnableActorBinding* NewCustomBinding = nullptr;
 
-	FString BindingNameToSet = MovieSceneHelpers::MakeUniqueBindingName(&OwnerMovieScene, FName::NameToDisplayString(SourceObject->GetName(), false));
-
 	const FName TemplateName = MakeUniqueObjectName(&OwnerMovieScene, UObject::StaticClass(), SourceObject->GetFName());
 	const FName InstancedBindingName = MakeUniqueObjectName(&OwnerMovieScene, UObject::StaticClass(), *FString(TemplateName.ToString() + TEXT("_CustomBinding")));
 
@@ -467,9 +470,6 @@ UMovieSceneCustomBinding* UMovieSceneSpawnableActorBinding::CreateNewCustomBindi
 #endif
 			NewCustomBinding->SetObjectTemplate(SpawnedActor);
 
-#if WITH_EDITOR
-			BindingNameToSet = Actor->GetActorLabel();
-#endif
 			if (!bWasTransactional)
 			{
 				Actor->ClearFlags(RF_Transactional);
@@ -526,11 +526,6 @@ UMovieSceneCustomBinding* UMovieSceneSpawnableActorBinding::CreateNewCustomBindi
 		}
 	}
 
-	if (NewCustomBinding)
-	{
-		NewCustomBinding->BindingName = BindingNameToSet;
-	}
-
 	return NewCustomBinding;
 }
 
diff --git a/Engine/Source/Runtime/MovieSceneTracks/Public/Bindings/MovieSceneSpawnableActorBinding.h b/Engine/Source/Runtime/MovieSceneTracks/Public/Bindings/MovieSceneSpawnableActorBinding.h
--- a/Engine/Source/Runtime/MovieSceneTracks/Public/Bindings/MovieSceneSpawnableActorBinding.h
+++ b/Engine/Source/Runtime/MovieSceneTracks/Public/Bindings/MovieSceneSpawnableActorBinding.h
@@ -131,7 +131,6 @@ public:
 	/* MovieSceneCustomBinding overrides*/
 	bool SupportsBindingCreationFromObject(const UObject* SourceObject) const override;
 	UMovieSceneCustomBinding* CreateNewCustomBinding(UObject* SourceObject, UMovieScene& OwnerMovieScene) override;
-	FString GetDesiredBindingName() const override { return BindingName; }
 
 #if WITH_EDITOR
 	bool SupportsConversionFromBinding(const FMovieSceneBindingReference& BindingReference, const UObject* SourceObject) const override;
@@ -152,10 +151,7 @@ protected:
 private:
 
 	UPROPERTY()
-	TObjectPtr<AActor> ActorTemplate;	
-
-	UPROPERTY()
-	FString BindingName;
+		TObjectPtr<AActor> ActorTemplate;
 
 
 };
