Sequencer- Fix issue with new custom spawnables where the object template could be saved as transient mistakenly, causing the spawnables to be invalid on reload of editor.

diff --git a/Engine/Source/Runtime/MovieSceneTracks/Private/Bindings/MovieSceneSpawnableActorBinding.cpp b/Engine/Source/Runtime/MovieSceneTracks/Private/Bindings/MovieSceneSpawnableActorBinding.cpp
--- a/Engine/Source/Runtime/MovieSceneTracks/Private/Bindings/MovieSceneSpawnableActorBinding.cpp
+++ b/Engine/Source/Runtime/MovieSceneTracks/Private/Bindings/MovieSceneSpawnableActorBinding.cpp
@@ -463,7 +463,7 @@ UMovieSceneCustomBinding* UMovieSceneSpawnableActorBinding::CreateNewCustomBindi
 
 		if (NewCustomBinding)
 		{
-			AActor* SpawnedActor = Cast<AActor>(StaticDuplicateObject(Actor, NewCustomBinding, TemplateName, RF_AllFlags));
+			AActor* SpawnedActor = Cast<AActor>(StaticDuplicateObject(Actor, NewCustomBinding, TemplateName, RF_AllFlags & ~RF_Transient));
 			SpawnedActor->DetachFromActor(FDetachmentTransformRules(EDetachmentRule::KeepRelative, false));
 #if WITH_EDITORONLY_DATA
 			SpawnedActor->bIsEditorPreviewActor = false;
