Sequencer: Unified all hbias blending logic into the HBias system

Previously this was not possible due to the grouping logic being scattered around various different systems. However, since the introduction of the grouping system we now have a unform way of reasoning about how entities are grouped together, which in turn allows us to set up hbias and blend targets uniformly.
The new method works by first collating all hbias levels for each group to determine whether hbias blending is being used for each, marking entities as Ignored as necessary, and finally adding/assigning and removing blend targets as necessary.

This change also introduces a way for grouping policies to initialize grouping keys in batches by implementing a InitializeGroupKeys method with the following signature:

		void InitializeGroupKeys(
			PolicyHandlerType& Handler,
			FEntityGroupBuilder* Builder,
			FEntityAllocationIteratorItem Item,
			FReadEntityIDs EntityIDs,
			TWrite<FEntityGroupID> GroupIDs,
			TRead<RequiredComponentTypes>... RequiredComponents)

allowing material parameters to be grouped and blended by parameter name and element index - previously they were all grouped and blended only by material element index.

diff --git a/Engine/Source/Runtime/MovieScene/Private/EntitySystem/BuiltInComponentTypes.cpp b/Engine/Source/Runtime/MovieScene/Private/EntitySystem/BuiltInComponentTypes.cpp
--- a/Engine/Source/Runtime/MovieScene/Private/EntitySystem/BuiltInComponentTypes.cpp
+++ b/Engine/Source/Runtime/MovieScene/Private/EntitySystem/BuiltInComponentTypes.cpp
@@ -407,7 +407,6 @@ FBuiltInComponentTypes::FBuiltInComponentTypes()
 	Tags.Finished                = ComponentRegistry->NewTag(TEXT("Finished Evaluating"));
 	Tags.Ignored                 = ComponentRegistry->NewTag(TEXT("Ignored"));
 	Tags.DontOptimizeConstants   = ComponentRegistry->NewTag(TEXT("Don't Optimize Constants"));
-	Tags.RemoveHierarchicalBlendTarget = ComponentRegistry->NewTag(TEXT("Remove Hierarchical Blend Target"));
 	Tags.FixedTime               = ComponentRegistry->NewTag(TEXT("Fixed Time"));
 	Tags.PreRoll                 = ComponentRegistry->NewTag(TEXT("Pre Roll"));
 	Tags.SectionPreRoll          = ComponentRegistry->NewTag(TEXT("Section Pre Roll"));
diff --git a/Engine/Source/Runtime/MovieScene/Private/EntitySystem/MovieSceneEntityGroupingSystem.cpp b/Engine/Source/Runtime/MovieScene/Private/EntitySystem/MovieSceneEntityGroupingSystem.cpp
--- a/Engine/Source/Runtime/MovieScene/Private/EntitySystem/MovieSceneEntityGroupingSystem.cpp
+++ b/Engine/Source/Runtime/MovieScene/Private/EntitySystem/MovieSceneEntityGroupingSystem.cpp
@@ -52,6 +52,17 @@ bool FEntityGroupBuilder::RemoveEntityFromGroup(const FMovieSceneEntityID& InEnt
 	return false;
 }
 
+int32 FEntityGroupBuilder::AllocateGroupIndex()
+{
+	return Owner->AllocateGroupIndex();
+}
+
+void FEntityGroupBuilder::FreeGroupIndex(int32 InGroupIndex)
+{
+	return Owner->FreeGroupIndex(InGroupIndex);
+}
+
+
 struct FAddGroupMutation : IMovieSceneEntityMutation
 {
 	UMovieSceneEntityGroupingSystem* System;
@@ -108,7 +119,6 @@ struct FUpdateGroupsTask
 
 	UMovieSceneEntityGroupingSystem* System;
 	FBuiltInComponentTypes* BuiltInComponents;
-	TBitArray<> ModifiedGroups;
 	bool bFreeGroupIDs = true;
 
 	FUpdateGroupsTask(UMovieSceneEntityGroupingSystem* InSystem, bool bInFreeGroupIDs = true)
@@ -123,9 +133,15 @@ struct FUpdateGroupsTask
 		BuiltInComponents = FBuiltInComponentTypes::Get();
 
 		// Run the handlers' pre-task callback.
-		for (const UMovieSceneEntityGroupingSystem::FEntityGroupingHandlerInfo& HandlerInfo : System->GroupHandlers)
+		
+		for (int32 HandlerIndex = 0; HandlerIndex < System->GroupHandlers.Num(); ++HandlerIndex)
 		{
-			HandlerInfo.Handler->PreTask();
+			const UMovieSceneEntityGroupingSystem::FEntityGroupingHandlerInfo& HandlerInfo = System->GroupHandlers[HandlerIndex];
+			if (HandlerInfo.Handler.IsValid())
+			{
+				FEntityGroupBuilder Builder(System, FEntityGroupingPolicyKey{ HandlerIndex });
+				HandlerInfo.Handler->PreTask(&Builder);
+			}
 		}
 	}
 
@@ -167,9 +183,14 @@ struct FUpdateGroupsTask
 	void PostTask()
 	{
 		// Run the handlers' post-task callback.
-		for (const UMovieSceneEntityGroupingSystem::FEntityGroupingHandlerInfo& HandlerInfo : System->GroupHandlers)
+		for (int32 HandlerIndex = 0; HandlerIndex < System->GroupHandlers.Num(); ++HandlerIndex)
 		{
-			HandlerInfo.Handler->PostTask(bFreeGroupIDs);
+			const UMovieSceneEntityGroupingSystem::FEntityGroupingHandlerInfo& HandlerInfo = System->GroupHandlers[HandlerIndex];
+			if (HandlerInfo.Handler.IsValid())
+			{
+				FEntityGroupBuilder Builder(System, FEntityGroupingPolicyKey{ HandlerIndex });
+				HandlerInfo.Handler->PostTask(bFreeGroupIDs, &Builder);
+			}
 		}
 	}
 };
@@ -209,6 +230,29 @@ bool UMovieSceneEntityGroupingSystem::IsRelevantImpl(UMovieSceneEntitySystemLink
 	return !GroupHandlers.IsEmpty();
 }
 
+int32 UMovieSceneEntityGroupingSystem::AllocateGroupIndex()
+{
+	// Find an existing free index
+	int32 NewGroupIndex = AllocatedGroupIndices.Find(false);
+	if (NewGroupIndex == INDEX_NONE)
+	{
+		NewGroupIndex = AllocatedGroupIndices.Add(true);
+	}
+	else
+	{
+		AllocatedGroupIndices[NewGroupIndex] = true;
+	}
+	return NewGroupIndex;
+}
+
+void UMovieSceneEntityGroupingSystem::FreeGroupIndex(int32 InGroupIndex)
+{
+	if (ensure(AllocatedGroupIndices.IsValidIndex(InGroupIndex) && AllocatedGroupIndices[InGroupIndex]))
+	{
+		AllocatedGroupIndices[InGroupIndex] = false;
+	}
+}
+
 void UMovieSceneEntityGroupingSystem::OnRun(FSystemTaskPrerequisites& InPrerequisites, FSystemSubsequentTasks& Subsequents)
 {
 	using namespace UE::MovieScene;
diff --git a/Engine/Source/Runtime/MovieScene/Public/EntitySystem/BuiltInComponentTypes.h b/Engine/Source/Runtime/MovieScene/Public/EntitySystem/BuiltInComponentTypes.h
--- a/Engine/Source/Runtime/MovieScene/Public/EntitySystem/BuiltInComponentTypes.h
+++ b/Engine/Source/Runtime/MovieScene/Public/EntitySystem/BuiltInComponentTypes.h
@@ -767,6 +767,7 @@ public:
 
 		FComponentTypeID DontOptimizeConstants;
 
+		UE_DEPRECATED(5.5, "This tag is no longer used. Blend targets are entirely managed by UMovieSceneHierarchicalBiasSystem now.")
 		FComponentTypeID RemoveHierarchicalBlendTarget;
 
 		FComponentTypeID OldStyleSpawnable;
diff --git a/Engine/Source/Runtime/MovieScene/Public/EntitySystem/MovieSceneEntityGroupingSystem.h b/Engine/Source/Runtime/MovieScene/Public/EntitySystem/MovieSceneEntityGroupingSystem.h
--- a/Engine/Source/Runtime/MovieScene/Public/EntitySystem/MovieSceneEntityGroupingSystem.h
+++ b/Engine/Source/Runtime/MovieScene/Public/EntitySystem/MovieSceneEntityGroupingSystem.h
@@ -25,6 +25,13 @@ namespace UE::MovieScene
 struct FAddGroupMutation;
 struct FUpdateGroupsTask;
 
+/** Concept that checks whether a grouping policy supports batch operations based on the presence of a InitializeGroupKeys function */
+struct CBatchGroupingPolicy
+{
+	template <typename T>
+	auto Requires(T& In) -> decltype(&T::InitializeGroupKeys);
+};
+
 /**
  * Utility class used by the grouping system's policies (see below) to manage groups.
  */
@@ -41,6 +48,9 @@ struct MOVIESCENE_API FEntityGroupBuilder
 	/** Remove the entity from the given group. The entity must already have the group ID component. */
 	bool RemoveEntityFromGroup(const FMovieSceneEntityID& InEntity, const FEntityGroupID& InPreviousGroupID);
 
+	int32 AllocateGroupIndex();
+	void FreeGroupIndex(int32 InGroupIndex);
+
 private:
 	UMovieSceneEntityGroupingSystem* Owner;
 	FEntityGroupingPolicyKey PolicyKey;
@@ -52,9 +62,9 @@ private:
 struct IEntityGroupingHandler
 {
 	virtual ~IEntityGroupingHandler() {}
-	virtual void PreTask() {}
+	virtual void PreTask(FEntityGroupBuilder* Builder) {}
 	virtual void ProcessAllocation(FEntityAllocationIteratorItem Item, FReadEntityIDs EntityIDs, TWrite<FEntityGroupID> GroupIDs, FEntityGroupBuilder* Builder) = 0;
-	virtual void PostTask(bool bFreeGroupIDs) {}
+	virtual void PostTask(bool bFreeGroupIDs, FEntityGroupBuilder* Builder) {}
 
 #if WITH_EDITOR
 	virtual void OnObjectsReplaced(const TMap<UObject*, UObject*>& ReplacementMap) = 0;
@@ -79,80 +89,74 @@ struct TEntityGroupingHandler
 	}
 };
 
-template<typename GroupingPolicy, int ...ComponentIndices, typename ...ComponentTypes>
-struct TEntityGroupingHandlerImpl<GroupingPolicy, TIntegerSequence<int, ComponentIndices...>, ComponentTypes...> : IEntityGroupingHandler
-{
-	using GroupKeyType = typename GroupingPolicy::GroupKeyType;
-
-	/** The grouping policy */
-	GroupingPolicy Policy;
 
-	/** The components that are required for making up a group key */
-	TTuple<TComponentTypeID<ComponentTypes>...> Components;
 
-	/** The group keys that we know about, mapped to their corresponding group index */
-	TMap<GroupKeyType, int32> GroupKeyToIndex;
-
-	/** The list of group indices in use */
-	TBitArray<> AllocatedGroupIndices;
-	/** The transient list of groups freed this frame */
-	TBitArray<> FreedGroupIndices;
-
-	TEntityGroupingHandlerImpl(GroupingPolicy&& InPolicy, TComponentTypeID<ComponentTypes>... InComponents)
-		: Policy(MoveTemp(InPolicy))
-		, Components(InComponents...)
-	{
-	}
+template<typename GroupingPolicy>
+struct TEntityGroupingHandlerBase
+{
+	using GroupKeyType = typename GroupingPolicy::GroupKeyType;
 
-	int32 GetOrAllocateGroupIndex(typename TCallTraits<GroupKeyType>::ParamType InGroupKey)
+	int32 GetOrAllocateGroupIndex(typename TCallTraits<GroupKeyType>::ParamType InGroupKey, FEntityGroupBuilder* Builder)
 	{
 		int32& GroupIndex = GroupKeyToIndex.FindOrAdd(InGroupKey, INDEX_NONE);
 		if (GroupIndex == INDEX_NONE)
 		{
 			// This group key isn't known to us... let's allocate a new group index for it.
 			// Try to find an available index first. Otherwise use a new high index.
-			int32 NewGroupIndex = AllocatedGroupIndices.Find(false);
-			if (NewGroupIndex == INDEX_NONE)
-			{
-				NewGroupIndex = AllocatedGroupIndices.Add(true);
-			}
-			else
-			{
-				AllocatedGroupIndices[NewGroupIndex] = true;
-			}
-			GroupIndex = NewGroupIndex;
+			GroupIndex = Builder->AllocateGroupIndex();
 		}
 		else
 		{
 			// We know this group key, so we'll return the group index we already have
 			// associated with it. We just need to "revive" it in case it was scheduled
 			// for being freed.
-			if (FreedGroupIndices.IsValidIndex(GroupIndex))
+			if (EmptyGroupIndices.IsValidIndex(GroupIndex))
 			{
-				FreedGroupIndices[GroupIndex] = false;
+				EmptyGroupIndices[GroupIndex] = false;
 			}
 		}
 		return GroupIndex;
 	}
 
-	void FreeGroupIndex(int32 InGroupIndex)
+protected:
+
+	/** The group keys that we know about, mapped to their corresponding group index */
+	TMap<GroupKeyType, int32> GroupKeyToIndex;
+
+	/** The transient list of groups freed this frame */
+	TBitArray<> EmptyGroupIndices;
+};
+
+
+
+
+template<typename GroupingPolicy, int ...ComponentIndices, typename ...ComponentTypes>
+struct TEntityGroupingHandlerImpl<GroupingPolicy, TIntegerSequence<int, ComponentIndices...>, ComponentTypes...> : TEntityGroupingHandlerBase<GroupingPolicy>, IEntityGroupingHandler
+{
+	using GroupKeyType = typename GroupingPolicy::GroupKeyType;
+
+	/** The grouping policy */
+	GroupingPolicy Policy;
+
+	/** The components that are required for making up a group key */
+	TTuple<TComponentTypeID<ComponentTypes>...> Components;
+
+	TEntityGroupingHandlerImpl(GroupingPolicy&& InPolicy, TComponentTypeID<ComponentTypes>... InComponents)
+		: Policy(MoveTemp(InPolicy))
+		, Components(InComponents...)
 	{
-		if (ensure(AllocatedGroupIndices.IsValidIndex(InGroupIndex) && AllocatedGroupIndices[InGroupIndex]))
-		{
-			AllocatedGroupIndices[InGroupIndex] = false;
-		}
 	}
 
 	/** Callback on the grouping policy invoked before all grouping happens this frame */
 	static void PreTaskImpl(void*, ...) {}
-	template <typename T> static void PreTaskImpl(T* InPolicy, decltype(&T::PreTask)* = 0)
+	template <typename T> static void PreTaskImpl(T* InPolicy, FEntityGroupBuilder* Builder, decltype(&T::PreTask)* = 0)
 	{
-		InPolicy->PreTask();
+		InPolicy->PreTask(Builder);
 	}
 
-	virtual void PreTask() override
+	virtual void PreTask(FEntityGroupBuilder* Builder) override
 	{
-		PreTaskImpl(&Policy);
+		PreTaskImpl(&Policy, Builder);
 	}
 
 	/**
@@ -195,36 +199,44 @@ struct TEntityGroupingHandlerImpl<GroupingPolicy, TIntegerSequence<int, Componen
 	
 	void VisitLinkedEntities(FEntityAllocationIteratorItem Item, FReadEntityIDs EntityIDs, TWrite<FEntityGroupID> GroupIDs, FEntityGroupBuilder* Builder)
 	{
-		const FEntityAllocation* Allocation = Item.GetAllocation();
-		const int32 Num = Allocation->Num();
+		if constexpr (TModels_V<CBatchGroupingPolicy, GroupingPolicy>)
+		{
+			const FEntityAllocation* Allocation = Item.GetAllocation();
+			Policy.InitializeGroupKeys(*this, Builder, Item, EntityIDs, GroupIDs, Allocation->ReadComponents(Components.template Get<ComponentIndices>())...);
+		}
+		else
+		{
+			const FEntityAllocation* Allocation = Item.GetAllocation();
+			const int32 Num = Allocation->Num();
 
-		const FEntityGroupID InvalidGroupID = Builder->MakeInvalidGroupID();
+			const FEntityGroupID InvalidGroupID = Builder->MakeInvalidGroupID();
 
-		TTuple<TComponentReader<ComponentTypes>...> ComponentReaders(
-				Allocation->ReadComponents(Components.template Get<ComponentIndices>())...);
+			TTuple<TComponentReader<ComponentTypes>...> ComponentReaders(
+					Allocation->ReadComponents(Components.template Get<ComponentIndices>())...);
 
-		for (int32 Index = 0; Index < Num; ++Index)
-		{
-			GroupKeyType GroupKey;
-			const bool bValidGroupKey = Policy.GetGroupKey(ComponentReaders.template Get<ComponentIndices>()[Index]..., GroupKey);
+			for (int32 Index = 0; Index < Num; ++Index)
+			{
+				GroupKeyType GroupKey;
+				const bool bValidGroupKey = Policy.GetGroupKey(ComponentReaders.template Get<ComponentIndices>()[Index]..., GroupKey);
 
-			const FMovieSceneEntityID EntityID(EntityIDs[Index]);
-			FEntityGroupID& GroupID(GroupIDs[Index]);
+				const FMovieSceneEntityID EntityID(EntityIDs[Index]);
+				FEntityGroupID& GroupID(GroupIDs[Index]);
 
-			if (bValidGroupKey)
-			{
-				// Find or create the appropriate group and put the entity in it.
-				int32 NewGroupIndex = GetOrAllocateGroupIndex(GroupKey);
-				FEntityGroupID NewGroupID = Builder->MakeGroupID(NewGroupIndex);
-				Builder->AddEntityToGroup(EntityID, NewGroupID);
-				GroupID = NewGroupID;
-			}
-			else
-			{
-				// This entity doesn't belong to any group.
-				// Let's assign an invalid group ID that nonetheless has a valid policy key
-				// pointing to this grouping.
-				GroupID = InvalidGroupID;
+				if (bValidGroupKey)
+				{
+					// Find or create the appropriate group and put the entity in it.
+					int32 NewGroupIndex = this->GetOrAllocateGroupIndex(GroupKey, Builder);
+					FEntityGroupID NewGroupID = Builder->MakeGroupID(NewGroupIndex);
+					Builder->AddEntityToGroup(EntityID, NewGroupID);
+					GroupID = NewGroupID;
+				}
+				else
+				{
+					// This entity doesn't belong to any group.
+					// Let's assign an invalid group ID that nonetheless has a valid policy key
+					// pointing to this grouping.
+					GroupID = InvalidGroupID;
+				}
 			}
 		}
 	}
@@ -256,8 +268,8 @@ struct TEntityGroupingHandlerImpl<GroupingPolicy, TIntegerSequence<int, Componen
 					//	  the above situation, and a brand new group that just happens to re-use the
 					//	  recently freed index.
 					//
-					FreedGroupIndices.PadToNum(GroupID.GroupIndex + 1, false);
-					FreedGroupIndices[GroupID.GroupIndex] = true;
+					this->EmptyGroupIndices.PadToNum(GroupID.GroupIndex + 1, false);
+					this->EmptyGroupIndices[GroupID.GroupIndex] = true;
 				}
 				// Leave the GroupID on the entity so that downstream systems can use it to track
 				// that this entity is leaving its group, but flag it so we don't re-free it.
@@ -275,38 +287,39 @@ struct TEntityGroupingHandlerImpl<GroupingPolicy, TIntegerSequence<int, Componen
 	}
 
 	static void PostTaskImpl(void*, ...){}
-	template <typename T> static void PostTaskImpl(T* InPolicy, bool bInFreeGroupIDs, decltype(&T::PostTask)* = 0)
+	template <typename T> static void PostTaskImpl(T* InPolicy, FEntityGroupBuilder* Builder, bool bInFreeGroupIDs, decltype(&T::PostTask)* = 0)
 	{
-		InPolicy->PostTask(bInFreeGroupIDs);
+		InPolicy->PostTask(Builder, bInFreeGroupIDs);
 	}
 	
-	virtual void PostTask(bool bInFreeGroupIDs) override
+	virtual void PostTask(bool bInFreeGroupIDs, FEntityGroupBuilder* Builder) override
 	{
-		if (bInFreeGroupIDs && FreedGroupIndices.Find(true) != INDEX_NONE)
+		if (bInFreeGroupIDs && this->EmptyGroupIndices.Find(true) != INDEX_NONE)
 		{
 			// Build a reverse lookup map to figure out which group keys we don't need anymore
 			// based on the group indices we have freed.
 			TMap<int32, GroupKeyType> GroupIndexToKey;
-			for (const TPair<GroupKeyType, int32>& Pair : GroupKeyToIndex)
+			for (const TPair<GroupKeyType, int32>& Pair : this->GroupKeyToIndex)
 			{
 				GroupIndexToKey.Add(Pair.Value, Pair.Key);
 			}
 
 			// Free the indices we don't use anymore, and free the group key lookup entry too.
-			for (TConstSetBitIterator<> It(FreedGroupIndices); It; ++It)
+			for (TConstSetBitIterator<> It(this->EmptyGroupIndices); It; ++It)
 			{
 				const GroupKeyType* GroupKey = GroupIndexToKey.Find(It.GetIndex());
 				if (ensure(GroupKey))
 				{
-					const int NumRemoved = GroupKeyToIndex.Remove(*GroupKey);
+					const int NumRemoved = this->GroupKeyToIndex.Remove(*GroupKey);
 					ensure(NumRemoved == 1);
 				}
-				FreeGroupIndex(It.GetIndex());
+
+				Builder->FreeGroupIndex(It.GetIndex());
 			}
-			FreedGroupIndices.Reset();
+			this->EmptyGroupIndices.Reset();
 		}
 
-		PostTaskImpl(&Policy);
+		PostTaskImpl(&Policy, Builder, bInFreeGroupIDs);
 	}
 
 #if WITH_EDITOR
@@ -314,7 +327,7 @@ struct TEntityGroupingHandlerImpl<GroupingPolicy, TIntegerSequence<int, Componen
 	{
 		// Get a list of keys that contain replaced objects.
 		TMap<GroupKeyType, GroupKeyType> ReplacedKeys;
-		for (const TPair<GroupKeyType, int32>& Pair : GroupKeyToIndex)
+		for (const TPair<GroupKeyType, int32>& Pair : this->GroupKeyToIndex)
 		{
 			GroupKeyType NewKey = Pair.Key;
 			if (Policy.OnObjectsReplaced(NewKey, ReplacementMap))
@@ -326,10 +339,10 @@ struct TEntityGroupingHandlerImpl<GroupingPolicy, TIntegerSequence<int, Componen
 		for (const TPair<GroupKeyType, GroupKeyType>& Pair : ReplacedKeys)
 		{
 			int32 GroupIndex;
-			const bool bRemoved = GroupKeyToIndex.RemoveAndCopyValue(Pair.Key, GroupIndex);
+			const bool bRemoved = this->GroupKeyToIndex.RemoveAndCopyValue(Pair.Key, GroupIndex);
 			if (ensure(bRemoved))
 			{
-				GroupKeyToIndex.Add(Pair.Value, GroupIndex);
+				this->GroupKeyToIndex.Add(Pair.Value, GroupIndex);
 			}
 		}
 	}
@@ -466,6 +479,27 @@ public:
 	 */
 	TArrayView<const UE::MovieScene::FMovieSceneEntityID> GetGroup(const FEntityGroupID& InGroupID) const;
 
+	/**
+	 * Allocate a new group index used to uniquely identify a collection of entities that animate the same target.
+	 * Group indices are globally unique within this system, regardless of the 'type' of the target.
+	 */
+	int32 AllocateGroupIndex();
+
+
+	/**
+	 * Free a previously allocated group index. Must not be in use or present on any components to be freed.
+	 */
+	void FreeGroupIndex(int32 Index);
+
+
+	/**
+	 * Return the maximum number of groups currently allocated
+	 */
+	int32 NumGroups() const
+	{
+		return AllocatedGroupIndices.Num();
+	}
+
 private:
 
 	virtual bool IsRelevantImpl(UMovieSceneEntitySystemLinker* InLinker) const override;
@@ -474,12 +508,18 @@ private:
 	virtual void OnUnlink() override;
 	virtual void OnCleanTaggedGarbage() override;
 
+	void ProcessModifiedGroups();
+
 #if WITH_EDITOR
 	void OnObjectsReplaced(const TMap<UObject*, UObject*>& ReplacementMap);
 #endif
 
 private:
 
+
+	/** The list of group indices in use */
+	TBitArray<> AllocatedGroupIndices;
+
 	struct FEntityGroupInfo
 	{
 		TArray<UE::MovieScene::FMovieSceneEntityID> Entities;
diff --git a/Engine/Source/Runtime/MovieSceneTracks/Private/Systems/MovieSceneCustomPrimitiveDataSystem.cpp b/Engine/Source/Runtime/MovieSceneTracks/Private/Systems/MovieSceneCustomPrimitiveDataSystem.cpp
--- a/Engine/Source/Runtime/MovieSceneTracks/Private/Systems/MovieSceneCustomPrimitiveDataSystem.cpp
+++ b/Engine/Source/Runtime/MovieSceneTracks/Private/Systems/MovieSceneCustomPrimitiveDataSystem.cpp
@@ -13,6 +13,7 @@
 #include "Evaluation/PreAnimatedState/MovieScenePreAnimatedStorageID.inl"
 #include "MovieSceneTracksComponentTypes.h"
 #include "Systems/DoubleChannelEvaluatorSystem.h"
+#include "EntitySystem/MovieSceneEntityGroupingSystem.h"
 #include "EntitySystem/MovieSceneEntitySystemRunner.h"
 #include "SceneTypes.h"
 #include "Components/PrimitiveComponent.h"
@@ -21,6 +22,54 @@
 
 namespace UE::MovieScene
 {
+
+	struct FCustomPrimitiveDataGroupingPolicy
+	{
+		using GroupKeyType = TTuple<FObjectKey, FName>;
+
+		void InitializeGroupKeys(
+			TEntityGroupingHandlerBase<FCustomPrimitiveDataGroupingPolicy>& Handler,
+			FEntityGroupBuilder* Builder,
+			FEntityAllocationIteratorItem Item,
+			FReadEntityIDs EntityIDs,
+			TWrite<FEntityGroupID> GroupIDs,
+			TRead<UObject*> BoundObjects,
+			TRead<FName> Names)
+		{
+			FMovieSceneTracksComponentTypes* TracksComponents = FMovieSceneTracksComponentTypes::Get();
+
+			const FComponentMask& AllocationType = Item.GetAllocationType();
+			if (!AllocationType.Contains(TracksComponents->Tags.CustomPrimitiveData))
+			{
+				return;
+			}
+
+			const FEntityAllocation* Allocation = Item.GetAllocation();
+			const int32              Num        = Allocation->Num();
+
+			TComponentReader<FName>    ParameterNames = Allocation->ReadComponents(TracksComponents->ScalarParameterName);
+			for (int32 Index = 0; Index < Num; ++Index)
+			{
+				GroupKeyType Key = MakeTuple(BoundObjects[Index], ParameterNames[Index]);
+
+				const int32    NewGroupIndex = Handler.GetOrAllocateGroupIndex(Key, Builder);
+				FEntityGroupID NewGroupID    = Builder->MakeGroupID(NewGroupIndex);
+
+				Builder->AddEntityToGroup(EntityIDs[Index], NewGroupID);
+				// Write out the group ID component
+				GroupIDs[Index] = NewGroupID;
+			}
+		}
+
+	#if WITH_EDITOR
+		bool OnObjectsReplaced(GroupKeyType& InOutKey, const TMap<UObject*, UObject*>& ReplacementMap)
+		{
+			return false;
+		}
+	#endif
+	};
+
+
 	void CollectGarbageForOutput(FAnimatedCustomPrimitiveDataInfo* Output)
 	{
 		// This should only happen during garbage collection
@@ -208,44 +257,12 @@ namespace UE::MovieScene
 
 				const FComponentTypeID BlenderTypeTag = System->DoubleBlenderSystem->GetBlenderTypeTag();
 
-				struct FBlendInfo
-				{
-					int16 HBias = TNumericLimits<int16>::Min();
-					bool bBlendHierarchicalBias = true;
-				};
-				FBlendInfo IgnoredBlendInfo;
-				FBlendInfo BlendInfo;
 
-				FHierarchicalBlendTarget BlendTarget;
 
-				for (FMovieSceneEntityID Input : Inputs)
-				{
-					FBlendInfo& BlendInfoToUpdate = Linker->EntityManager.HasComponent(Input, BuiltInComponents->Tags.Ignored) ? IgnoredBlendInfo : BlendInfo;
 
-					TOptionalComponentReader<int16> HBiasComponent = Linker->EntityManager.ReadComponent(Input, BuiltInComponents->HierarchicalBias);
-					const int16 HBias = HBiasComponent ? *HBiasComponent : 0;
 
-					BlendTarget.Add(HBias);
 
-					if (HBias > BlendInfoToUpdate.HBias)
-					{
-						BlendInfoToUpdate.HBias = HBias;
-						BlendInfoToUpdate.bBlendHierarchicalBias = Linker->EntityManager.HasComponent(Input, BuiltInComponents->Tags.BlendHierarchicalBias);
-					}
-					else if (HBias == BlendInfoToUpdate.HBias && !BlendInfoToUpdate.bBlendHierarchicalBias)
-					{
-						BlendInfoToUpdate.bBlendHierarchicalBias = Linker->EntityManager.HasComponent(Input, BuiltInComponents->Tags.BlendHierarchicalBias);
-					}
-				}
 
-				if (BlendInfo.HBias == TNumericLimits<int16>::Min())
-				{
-					BlendInfo = IgnoredBlendInfo;
-				}
-				else if (IgnoredBlendInfo.HBias != TNumericLimits<int16>::Min())
-				{
-					BlendInfo.bBlendHierarchicalBias |= IgnoredBlendInfo.bBlendHierarchicalBias;
-				}
 
 				for (FMovieSceneEntityID Input : Inputs)
 				{
@@ -259,21 +276,6 @@ namespace UE::MovieScene
 						Linker->EntityManager.WriteComponentChecked(Input, BuiltInComponents->BlendChannelInput, Output->BlendChannelID);
 					}
 
-					if (BlendInfo.bBlendHierarchicalBias)
-					{
-						if (!Linker->EntityManager.HasComponent(Input, BuiltInComponents->HierarchicalBlendTarget))
-						{
-							Linker->EntityManager.AddComponent(Input, BuiltInComponents->HierarchicalBlendTarget, BlendTarget);
-						}
-						else
-						{
-							Linker->EntityManager.WriteComponentChecked(Input, BuiltInComponents->HierarchicalBlendTarget, BlendTarget);
-						}
-					}
-					else
-					{
-						Linker->EntityManager.RemoveComponent(Input, BuiltInComponents->HierarchicalBlendTarget);
-					}
 
 					// Ensure we have the blender type tag on the inputs.
 					Linker->EntityManager.AddComponent(Input, BlenderTypeTag);
@@ -389,7 +391,7 @@ UMovieSceneCustomPrimitiveDataSystem::UMovieSceneCustomPrimitiveDataSystem(const
 	if (HasAnyFlags(RF_ClassDefaultObject))
 	{
 		DefineComponentConsumer(GetClass(), BuiltInComponents->BoundObject);
-		DefineComponentProducer(GetClass(), BuiltInComponents->HierarchicalBlendTarget);
+		DefineComponentConsumer(GetClass(), BuiltInComponents->HierarchicalBlendTarget);
 
 		DefineImplicitPrerequisite(UFloatChannelEvaluatorSystem::StaticClass(), GetClass());
 		DefineImplicitPrerequisite(UDoubleChannelEvaluatorSystem::StaticClass(), GetClass());
@@ -403,7 +405,6 @@ UMovieSceneCustomPrimitiveDataSystem::UMovieSceneCustomPrimitiveDataSystem(const
 
 		DefineImplicitPrerequisite(UMovieSceneHierarchicalEasingInstantiatorSystem::StaticClass(), GetClass());
 		DefineImplicitPrerequisite(UMovieScenePiecewiseDoubleBlenderSystem::StaticClass(), GetClass());
-		DefineImplicitPrerequisite(GetClass(), UMovieSceneHierarchicalBiasSystem::StaticClass());
 		DefineImplicitPrerequisite(GetClass(), UMovieSceneInitialValueSystem::StaticClass());
 	}
 }
@@ -412,9 +413,15 @@ void UMovieSceneCustomPrimitiveDataSystem::OnLink()
 {
 	using namespace UE::MovieScene;
 
+	FBuiltInComponentTypes*          BuiltInComponents = FBuiltInComponentTypes::Get();
+	FMovieSceneTracksComponentTypes* TracksComponents  = FMovieSceneTracksComponentTypes::Get();
+
 	ScalarParameterTracker.Initialize(this);
 
 	ScalarParameterStorage = Linker->PreAnimatedState.GetOrCreateStorage<FPreAnimatedCustomPrimitiveDataEntryStorage>();
+
+	UMovieSceneEntityGroupingSystem* GroupingSystem = Linker->LinkSystem<UMovieSceneEntityGroupingSystem>();
+	GroupingKey = GroupingSystem->AddGrouping(FCustomPrimitiveDataGroupingPolicy(), BuiltInComponents->BoundObject, TracksComponents->ScalarParameterName);
 }
 
 void UMovieSceneCustomPrimitiveDataSystem::OnUnlink()
@@ -425,6 +432,13 @@ void UMovieSceneCustomPrimitiveDataSystem::OnUnlink()
 	DoubleBlenderSystem = nullptr;
 
 	ScalarParameterTracker.Destroy(TOverlappingCustomPrimitiveDataHandler<FCustomPrimitiveDataMixin>(this));
+
+	UMovieSceneEntityGroupingSystem* GroupingSystem = Linker->FindSystem<UMovieSceneEntityGroupingSystem>();
+	if (ensure(GroupingSystem))
+	{
+		GroupingSystem->RemoveGrouping(GroupingKey);
+	}
+	GroupingKey = FEntityGroupingPolicyKey();
 }
 
 void UMovieSceneCustomPrimitiveDataSystem::OnRun(FSystemTaskPrerequisites& InPrerequisites, FSystemSubsequentTasks& Subsequents)
diff --git a/Engine/Source/Runtime/MovieSceneTracks/Private/Systems/MovieSceneHierarchicalBiasSystem.cpp b/Engine/Source/Runtime/MovieSceneTracks/Private/Systems/MovieSceneHierarchicalBiasSystem.cpp
--- a/Engine/Source/Runtime/MovieSceneTracks/Private/Systems/MovieSceneHierarchicalBiasSystem.cpp
+++ b/Engine/Source/Runtime/MovieSceneTracks/Private/Systems/MovieSceneHierarchicalBiasSystem.cpp
@@ -3,136 +3,167 @@
 #include "Systems/MovieSceneHierarchicalBiasSystem.h"
 
 #include "EntitySystem/BuiltInComponentTypes.h"
+#include "EntitySystem/MovieSceneEntityMutations.h"
 #include "EntitySystem/MovieSceneEntitySystemTask.h"
 #include "EntitySystem/MovieSceneEntitySystemLinker.h"
+#include "EntitySystem/MovieSceneEntityGroupingSystem.h"
 #include "EntitySystem/MovieScenePreAnimatedStateSystem.h"
 
 #include UE_INLINE_GENERATED_CPP_BY_NAME(MovieSceneHierarchicalBiasSystem)
 
 
-namespace UE
+namespace UE::MovieScene
 {
-namespace MovieScene
-{
-
-struct FEntityGroupSequenceKey
-{
-	FRootInstanceHandle RootInstance;
-	FEntityGroupID GroupID;
-	friend uint32 GetTypeHash(const FEntityGroupSequenceKey& In)
+	/** Temporary struct used for collating hbias meta data for each group */
+	struct FHBiasMetaData
 	{
-		return HashCombine(GetTypeHash(In.RootInstance), GetTypeHash(In.GroupID));
-	}
-	friend bool operator==(const FEntityGroupSequenceKey& A, const FEntityGroupSequenceKey& B)
-	{
-		return A.RootInstance == B.RootInstance && A.GroupID == B.GroupID;
-	}
-};
+		UE::MovieScene::FHierarchicalBlendTarget BlendTarget;
+		int16 HBias;
+		uint8 bBlendHierarchicalBias : 1;
+		uint8 bInUse : 1;
 
-struct FHierarchicalBiasTask
-{
-	explicit FHierarchicalBiasTask(UMovieSceneEntitySystemLinker* InLinker)
-		: Linker(InLinker)
-	{}
-
-	void InitializeGroup(FRootInstanceHandle RootInstanceHandle, FEntityGroupID GroupID)
-	{
-		MaxBiasByGroup.FindOrAdd(FEntityGroupSequenceKey{ RootInstanceHandle, GroupID }, MIN_int16);
-	}
+		FHBiasMetaData()
+		{
+			bBlendHierarchicalBias = false;
+			bInUse = false;
+			HBias = TNumericLimits<int16>::Lowest();
+		}
+	};
 
-	bool HasAnyWork() const
+	/** Mutation that adds or removes the Ignored tag for entities */
+	struct FToggleIgnoredMutation : IMovieSceneConditionalEntityMutation
 	{
-		return MaxBiasByGroup.Num() != 0;
-	}
+		TArrayView<const FHBiasMetaData> HBiasMetaData;
 
-	void ForEachAllocation(FEntityAllocationIteratorItem Iterator, TRead<FMovieSceneEntityID> EntityIDs, TRead<FRootInstanceHandle> RootInstanceHandles, TRead<FEntityGroupID> GroupIDs, TReadOptional<int16> OptHBiases)
-	{
-		const int32 Num = Iterator.GetAllocation()->Num();
-		const FComponentMask& AllocationType = Iterator.GetAllocationType();
-		const bool bIgnoreBias = AllocationType.Contains(FBuiltInComponentTypes::Get()->Tags.IgnoreHierarchicalBias)
-			|| AllocationType.Contains(FBuiltInComponentTypes::Get()->HierarchicalBlendTarget);
+		FToggleIgnoredMutation(TArrayView<const FHBiasMetaData> InHBiasMetaData)
+			: HBiasMetaData(InHBiasMetaData)
+		{}
 
-		if (bIgnoreBias)
+		void MarkAllocation(FEntityAllocation* Allocation, TBitArray<>& OutEntitiesToMutate) const override
 		{
+			FBuiltInComponentTypes* BuiltInComponents = FBuiltInComponentTypes::Get();
+
+			const bool bCurrentlyIgnored = Allocation->HasComponent(BuiltInComponents->Tags.Ignored);
+
+			TComponentReader<FEntityGroupID> GroupComponents = Allocation->ReadComponents(BuiltInComponents->Group);
+			TOptionalComponentReader<int16>  HBiasComponents = Allocation->TryReadComponents(BuiltInComponents->HierarchicalBias);
+
+			const int32 Num = Allocation->Num();
 			for (int32 Index = 0; Index < Num; ++Index)
 			{
-				FEntityGroupSequenceKey Key{ RootInstanceHandles[Index], GroupIDs[Index] };
-				ActiveContributorsByGroup.Add(Key, EntityIDs[Index]);
+				const int16 HBias = HBiasComponents ? HBiasComponents[Index] : 0;
+
+				if (!EnumHasAnyFlags(GroupComponents[Index].Flags, EEntityGroupFlags::RemovedFromGroup))
+				{
+					const FHBiasMetaData& MetaData = HBiasMetaData[GroupComponents[Index].GroupIndex];
+					const bool bShouldBeIgnored = (!MetaData.bBlendHierarchicalBias && MetaData.HBias > HBias);
+				
+					if (bShouldBeIgnored != bCurrentlyIgnored)
+					{
+						OutEntitiesToMutate.PadToNum(Index + 1, false);
+						OutEntitiesToMutate[Index] = true;
+					}
+				}
 			}
 		}
-		else if (OptHBiases)
+
+		void CreateMutation(FEntityManager* EntityManager, FComponentMask* InOutEntityComponentTypes) const override
 		{
-			for (int32 Index = 0; Index < Num; ++Index)
+			FBuiltInComponentTypes* BuiltInComponents = FBuiltInComponentTypes::Get();
+
+			if (InOutEntityComponentTypes->Contains(BuiltInComponents->Tags.Ignored))
 			{
-				VisitGroup(EntityIDs[Index], RootInstanceHandles[Index], GroupIDs[Index], OptHBiases[Index]);
+				InOutEntityComponentTypes->Remove(BuiltInComponents->Tags.Ignored);
 			}
-		}
-		else
-		{
-			for (int32 Index = 0; Index < Num; ++Index)
+			else
 			{
-				VisitGroup(EntityIDs[Index], RootInstanceHandles[Index], GroupIDs[Index], 0);
+				InOutEntityComponentTypes->Set(BuiltInComponents->Tags.Ignored);
 			}
+
+			InOutEntityComponentTypes->Set(BuiltInComponents->Tags.NeedsLink);
 		}
-	}
+	};
 
-	void PostTask()
+	/** Mutation that adds, removes or assigns the HierarchicalBlendTarget components for entities */
+	struct FBlendTargetMutation : IMovieSceneConditionalEntityMutation
 	{
-		FBuiltInComponentTypes* Components = FBuiltInComponentTypes::Get();
+		TArrayView<const FHBiasMetaData> HBiasMetaData;
+		FEntityAllocationWriteContext WriteContext;
 
-		for (auto It = ActiveContributorsByGroup.CreateIterator(); It; ++It)
-		{
-			Linker->EntityManager.RemoveComponent(It.Value(), Components->Tags.Ignored);
-		}
+		FBlendTargetMutation(TArrayView<const FHBiasMetaData> InHBiasMetaData, FEntityAllocationWriteContext InWriteContext)
+			: HBiasMetaData(InHBiasMetaData)
+			, WriteContext(InWriteContext)
+		{}
 
-		for (auto It = InactiveContributorsByGroup.CreateIterator(); It; ++It)
+		void MarkAllocation(FEntityAllocation* Allocation, TBitArray<>& OutEntitiesToMutate) const override
 		{
-			Linker->EntityManager.AddComponent(It.Value(), Components->Tags.Ignored);
-		}
-	}
+			FBuiltInComponentTypes* BuiltInComponents = FBuiltInComponentTypes::Get();
 
-private:
+			TOptionalComponentWriter<FHierarchicalBlendTarget> OutBlendTargets = Allocation->TryWriteComponents(BuiltInComponents->HierarchicalBlendTarget, WriteContext);
+			TComponentReader<FEntityGroupID>                   GroupComponents = Allocation->ReadComponents(BuiltInComponents->Group);
 
-	void VisitGroup(FMovieSceneEntityID EntityID, FRootInstanceHandle RootInstanceHandle, FEntityGroupID GroupID, int16 HBias)
-	{
-		FEntityGroupSequenceKey Key{ RootInstanceHandle, GroupID };
+			const bool bIsIgnored = Allocation->HasComponent(BuiltInComponents->Tags.Ignored);
+			const bool bHasBlendTarget = !!OutBlendTargets;
 
-		// If this group hasn't changed at all (ie InitializeGroup was not called for it) do nothing
-		if (int16* ExistingBias = MaxBiasByGroup.Find(Key))
-		{
-			if (HBias > *ExistingBias)
+			const int32 Num = Allocation->Num();
+			for (int32 Index = 0; Index < Num; ++Index)
 			{
-				for (auto It = ActiveContributorsByGroup.CreateKeyIterator(Key); It; ++It)
+				if (!EnumHasAnyFlags(GroupComponents[Index].Flags, EEntityGroupFlags::RemovedFromGroup))
 				{
-					InactiveContributorsByGroup.Add(Key, It.Value());
-					It.RemoveCurrent();
+					const int32  GroupIndex = GroupComponents[Index].GroupIndex;
+					const FHBiasMetaData& MetaData = HBiasMetaData[GroupIndex];
+
+					const bool bNeedsBlendTarget = !bIsIgnored && MetaData.bBlendHierarchicalBias;
+					if (bNeedsBlendTarget != bHasBlendTarget)
+					{
+						OutEntitiesToMutate.PadToNum(Index + 1, false);
+						OutEntitiesToMutate[Index] = true;
+					}
+					else if (OutBlendTargets)
+					{
+						OutBlendTargets[Index] = HBiasMetaData[GroupIndex].BlendTarget;
+					}
 				}
-
-				*ExistingBias = HBias;
-				ActiveContributorsByGroup.Add(Key, EntityID);
 			}
-			else if (HBias == *ExistingBias)
+		}
+
+		void CreateMutation(FEntityManager* EntityManager, FComponentMask* InOutEntityComponentTypes) const override
+		{
+			FBuiltInComponentTypes* BuiltInComponents = FBuiltInComponentTypes::Get();
+
+			if (InOutEntityComponentTypes->Contains(BuiltInComponents->HierarchicalBlendTarget))
 			{
-				ActiveContributorsByGroup.Add(Key, EntityID);
+				InOutEntityComponentTypes->Remove(BuiltInComponents->HierarchicalBlendTarget);
 			}
 			else
 			{
-				InactiveContributorsByGroup.Add(Key, EntityID);
+				InOutEntityComponentTypes->Set(BuiltInComponents->HierarchicalBlendTarget);
 			}
+
+			InOutEntityComponentTypes->Set(BuiltInComponents->Tags.NeedsLink);
 		}
-	}
 
-	TMap<FEntityGroupSequenceKey, int16> MaxBiasByGroup;
+		void InitializeEntities(const FEntityRange& EntityRange, const FComponentMask& AllocationType) const override
+		{
+			FBuiltInComponentTypes* BuiltInComponents = FBuiltInComponentTypes::Get();
 
-	TMultiMap<FEntityGroupSequenceKey, FMovieSceneEntityID> InactiveContributorsByGroup;
+			TOptionalComponentWriter<FHierarchicalBlendTarget> BlendTargets = EntityRange.Allocation->TryWriteComponents(BuiltInComponents->HierarchicalBlendTarget, FEntityAllocationWriteContext::NewAllocation());
+			if (BlendTargets)
+			{
+				TComponentReader<FEntityGroupID> GroupComponents = EntityRange.Allocation->ReadComponents(BuiltInComponents->Group);
 
-	TMultiMap<FEntityGroupSequenceKey, FMovieSceneEntityID> ActiveContributorsByGroup;
+				for (int32 Index = 0; Index < EntityRange.Num; ++Index)
+				{
+					const FHBiasMetaData& MetaData = HBiasMetaData[GroupComponents[EntityRange.ComponentStartOffset + Index].GroupIndex];
 
-	UMovieSceneEntitySystemLinker* Linker;
-};
+					BlendTargets[Index] = MetaData.BlendTarget;
+				}
+			}
+		}
+
+	};
 
-} // namespace MovieScene
-} // namespace UE
+} // namespace UE::MovieScene
 
 
 UMovieSceneHierarchicalBiasSystem::UMovieSceneHierarchicalBiasSystem(const FObjectInitializer& ObjInit)
@@ -140,15 +171,16 @@ UMovieSceneHierarchicalBiasSystem::UMovieSceneHierarchicalBiasSystem(const FObje
 {
 	using namespace UE::MovieScene;
 
+	GroupingSystem = nullptr;
 	SystemCategories = EEntitySystemCategory::Core;
 
 	if (HasAnyFlags(RF_ClassDefaultObject))
 	{
-		DefineComponentConsumer(GetClass(), FBuiltInComponentTypes::Get()->Group);
+		FBuiltInComponentTypes* BuiltInComponents = FBuiltInComponentTypes::Get();
+		DefineComponentConsumer(GetClass(), BuiltInComponents->Group);
 
-		// Don't flag things with the Ignore tag (due to hierarchical biases) until all systems have
-		// had a chance to take them into account for pre-animated state.
-		DefineImplicitPrerequisite(UMovieSceneCachePreAnimatedStateSystem::StaticClass(), GetClass());
+		DefineComponentProducer(GetClass(), BuiltInComponents->Tags.Ignored);
+		DefineComponentProducer(GetClass(), BuiltInComponents->HierarchicalBlendTarget);
 	}
 }
 
@@ -160,31 +192,84 @@ bool UMovieSceneHierarchicalBiasSystem::IsRelevantImpl(UMovieSceneEntitySystemLi
 	return InLinker->EntityManager.ContainsAllComponents({ Components->Group, Components->HierarchicalBias });
 }
 
+void UMovieSceneHierarchicalBiasSystem::OnLink()
+{
+	GroupingSystem = Linker->LinkSystem<UMovieSceneEntityGroupingSystem>();
+	Linker->SystemGraph.AddReference(this, GroupingSystem.Get());
+}
+
 void UMovieSceneHierarchicalBiasSystem::OnRun(FSystemTaskPrerequisites& InPrerequisites, FSystemSubsequentTasks& Subsequents)
 {
 	using namespace UE::MovieScene;
 
-	FBuiltInComponentTypes* Components = FBuiltInComponentTypes::Get();
+	FBuiltInComponentTypes* BuiltInComponents = FBuiltInComponentTypes::Get();
 
-	FHierarchicalBiasTask Task(Linker);
+	TArray<FHBiasMetaData> HBiasMetaData;
+	HBiasMetaData.SetNum(GroupingSystem->NumGroups());
 
-	// First, add all the groups that have changed to the map
-	FEntityTaskBuilder()
-	.Read(Components->RootInstanceHandle)
-	.Read(Components->Group)
-	.FilterAny({ Components->Tags.NeedsLink, Components->Tags.NeedsUnlink })
-	.Iterate_PerEntity(&Linker->EntityManager, [&Task](FRootInstanceHandle RootInstanceHandle, FEntityGroupID GroupID)
-			{ Task.InitializeGroup(RootInstanceHandle, GroupID); });
-
-	if (Task.HasAnyWork())
+	auto GatherHierarchicalMetaData = [&HBiasMetaData, BuiltInComponents](FEntityAllocationIteratorItem Item, const FEntityGroupID* GroupIDs, const int16* OptionalHBias)
 	{
-		FEntityTaskBuilder()
-		.ReadEntityIDs()
-		.Read(Components->RootInstanceHandle)
-		.Read(Components->Group)
-		.ReadOptional(Components->HierarchicalBias)
-		.FilterNone({ Components->Tags.NeedsUnlink })
-		.RunInline_PerAllocation(&Linker->EntityManager, Task);
-	}
+		const FComponentMask& AllocationType = Item.GetAllocationType();
+
+		const int32 Num = Item.GetAllocation()->Num();
+
+		const bool bBlendHBias = AllocationType.Contains(BuiltInComponents->Tags.BlendHierarchicalBias);
+
+		if (AllocationType.Contains(BuiltInComponents->Tags.IgnoreHierarchicalBias))
+		{
+			for (int32 Index = 0; Index < Num; ++Index)
+			{
+				if (!EnumHasAnyFlags(GroupIDs[Index].Flags, EEntityGroupFlags::RemovedFromGroup))
+				{
+					const int32 GroupIndex = GroupIDs[Index].GroupIndex;
+
+					FHBiasMetaData& MetaData = HBiasMetaData[GroupIndex];
+
+					MetaData.bInUse = true;
+					MetaData.bBlendHierarchicalBias |= bBlendHBias;
+				}
+			}
+		}
+		else
+		{
+			for (int32 Index = 0; Index < Num; ++Index)
+			{
+				if (!EnumHasAnyFlags(GroupIDs[Index].Flags, EEntityGroupFlags::RemovedFromGroup))
+				{
+					const int32 GroupIndex = GroupIDs[Index].GroupIndex;
+					const int16 ThisHBias  = OptionalHBias ? OptionalHBias[Index] : 0;
+
+					FHBiasMetaData& MetaData = HBiasMetaData[GroupIndex];
+
+					if (ThisHBias > MetaData.HBias)
+					{
+						MetaData.HBias = ThisHBias;
+					}
+
+					MetaData.bInUse = true;
+					MetaData.bBlendHierarchicalBias |= bBlendHBias;
+					MetaData.BlendTarget.Add(ThisHBias);
+				}
+			}
+		}
+	};
+
+	// --------------------------------------------------------------------------
+	// Step 1: Gather hbias meta-data for each group
+	FEntityTaskBuilder()
+	.Read(BuiltInComponents->Group)
+	.ReadOptional(BuiltInComponents->HierarchicalBias)
+	.Iterate_PerAllocation(&Linker->EntityManager, GatherHierarchicalMetaData);
+
+	// --------------------------------------------------------------------------
+	// Step 2: Toggle non-blended entities that are part of lower-hbias
+	FToggleIgnoredMutation ToggleIgnoredMutation(HBiasMetaData);
+	Linker->EntityManager.MutateConditional(FEntityComponentFilter().All({ BuiltInComponents->Group }), ToggleIgnoredMutation);
+
+	// --------------------------------------------------------------------------
+	// Step 3: Update blend targets on blended entities
+	FBlendTargetMutation ToggleBlendTargetMutation(HBiasMetaData, FEntityAllocationWriteContext(Linker->EntityManager));
+	Linker->EntityManager.MutateConditional(
+		FEntityComponentFilter().Any({ BuiltInComponents->Group, BuiltInComponents->HierarchicalBlendTarget }), ToggleBlendTargetMutation);
 }
 
diff --git a/Engine/Source/Runtime/MovieSceneTracks/Private/Systems/MovieSceneMaterialParameterSystem.cpp b/Engine/Source/Runtime/MovieSceneTracks/Private/Systems/MovieSceneMaterialParameterSystem.cpp
--- a/Engine/Source/Runtime/MovieSceneTracks/Private/Systems/MovieSceneMaterialParameterSystem.cpp
+++ b/Engine/Source/Runtime/MovieSceneTracks/Private/Systems/MovieSceneMaterialParameterSystem.cpp
@@ -390,44 +390,12 @@ struct TOverlappingMaterialParameterHandler : Mixin
 
 			const FComponentTypeID BlenderTypeTag = System->DoubleBlenderSystem->GetBlenderTypeTag();
 
-			struct FBlendInfo
-			{
-				int16 HBias = TNumericLimits<int16>::Min();
-				bool bBlendHierarchicalBias = true;
-			};
-			FBlendInfo IgnoredBlendInfo;
-			FBlendInfo BlendInfo;
 
-			FHierarchicalBlendTarget BlendTarget;
 
-			for (FMovieSceneEntityID Input : Inputs)
-			{
-				FBlendInfo& BlendInfoToUpdate = Linker->EntityManager.HasComponent(Input, BuiltInComponents->Tags.Ignored) ? IgnoredBlendInfo : BlendInfo;
 
-				TOptionalComponentReader<int16> HBiasComponent = Linker->EntityManager.ReadComponent(Input, BuiltInComponents->HierarchicalBias);
-				const int16 HBias = HBiasComponent ? *HBiasComponent : 0;
 
-				BlendTarget.Add(HBias);
 
-				if (HBias > BlendInfoToUpdate.HBias)
-				{
-					BlendInfoToUpdate.HBias = HBias;
-					BlendInfoToUpdate.bBlendHierarchicalBias = Linker->EntityManager.HasComponent(Input, BuiltInComponents->Tags.BlendHierarchicalBias);
-				}
-				else if (HBias == BlendInfoToUpdate.HBias && !BlendInfoToUpdate.bBlendHierarchicalBias)
-				{
-					BlendInfoToUpdate.bBlendHierarchicalBias = Linker->EntityManager.HasComponent(Input, BuiltInComponents->Tags.BlendHierarchicalBias);
-				}
-			}
 
-			if (BlendInfo.HBias == TNumericLimits<int16>::Min())
-			{
-				BlendInfo = IgnoredBlendInfo;
-			}
-			else if (IgnoredBlendInfo.HBias != TNumericLimits<int16>::Min())
-			{
-				BlendInfo.bBlendHierarchicalBias |= IgnoredBlendInfo.bBlendHierarchicalBias;
-			}
 
 			for (FMovieSceneEntityID Input : Inputs)
 			{
@@ -441,21 +409,6 @@ struct TOverlappingMaterialParameterHandler : Mixin
 					Linker->EntityManager.WriteComponentChecked(Input, BuiltInComponents->BlendChannelInput, Output->BlendChannelID);
 				}
 
-				if (BlendInfo.bBlendHierarchicalBias)
-				{
-					if (!Linker->EntityManager.HasComponent(Input, BuiltInComponents->HierarchicalBlendTarget))
-					{
-						Linker->EntityManager.AddComponent(Input, BuiltInComponents->HierarchicalBlendTarget, BlendTarget);
-					}
-					else
-					{
-						Linker->EntityManager.WriteComponentChecked(Input, BuiltInComponents->HierarchicalBlendTarget, BlendTarget);
-					}
-				}
-				else if (Linker->EntityManager.HasComponent(Input, BuiltInComponents->HierarchicalBlendTarget))
-				{
-					Linker->EntityManager.AddComponent(Input, BuiltInComponents->Tags.RemoveHierarchicalBlendTarget);
-				}
 
 				// Ensure we have the blender type tag on the inputs.
 				Linker->EntityManager.AddComponent(Input, BlenderTypeTag);
@@ -466,11 +419,6 @@ struct TOverlappingMaterialParameterHandler : Mixin
 			Linker->EntityManager.RemoveComponent(Inputs[0], BuiltInComponents->BlendChannelInput);
 
 			Mixin::InitializeSoleInput(Linker, BoundMaterial, ParameterInfo, Inputs[0], Output);
-
-			if (Linker->EntityManager.HasComponent(Inputs[0], BuiltInComponents->HierarchicalBlendTarget))
-			{
-				Linker->EntityManager.AddComponent(Inputs[0], BuiltInComponents->Tags.RemoveHierarchicalBlendTarget);
-			}
 		}
 
 		Output->NumContributors = NumContributors;
@@ -511,7 +459,7 @@ UMovieSceneMaterialParameterInstantiatorSystem::UMovieSceneMaterialParameterInst
 	{
 		DefineComponentConsumer(GetClass(), TracksComponents->BoundMaterial);
 
-		DefineComponentProducer(GetClass(), BuiltInComponents->HierarchicalBlendTarget);
+		DefineComponentConsumer(GetClass(), BuiltInComponents->HierarchicalBlendTarget);
 
 		DefineImplicitPrerequisite(UMovieSceneHierarchicalEasingInstantiatorSystem::StaticClass(), GetClass());
 		DefineImplicitPrerequisite(GetClass(), UMovieSceneInitialValueSystem::StaticClass());
diff --git a/Engine/Source/Runtime/MovieSceneTracks/Private/Systems/MovieScenePropertyInstantiator.cpp b/Engine/Source/Runtime/MovieSceneTracks/Private/Systems/MovieScenePropertyInstantiator.cpp
--- a/Engine/Source/Runtime/MovieSceneTracks/Private/Systems/MovieScenePropertyInstantiator.cpp
+++ b/Engine/Source/Runtime/MovieSceneTracks/Private/Systems/MovieScenePropertyInstantiator.cpp
@@ -49,28 +49,10 @@ struct FPropertyInstantiatorGroupingPolicy
 
 } // namespace UE::MovieScene
 
-void UMovieScenePropertyInstantiatorSystem::FHierarchicalMetaData::CombineWith(const FHierarchicalMetaData& Other)
-{
-	NumContributors        += Other.NumContributors;
-	bWantsRestoreState     |= Other.bWantsRestoreState;
-	bSupportsFastPath      &= Other.bSupportsFastPath;
-	bNeedsInitialValue     |= Other.bNeedsInitialValue;
-	bBlendHierarchicalBias |= Other.bBlendHierarchicalBias;
-}
-
-void UMovieScenePropertyInstantiatorSystem::FHierarchicalMetaData::ResetTracking()
-{
-	NumContributors = 0;
-	bWantsRestoreState = false;
-	bSupportsFastPath = true;
-	bNeedsInitialValue = false;
-	bBlendHierarchicalBias = false;
-	bInUse = false;
-}
 
 UMovieScenePropertyInstantiatorSystem::FContributorKey UMovieScenePropertyInstantiatorSystem::FPropertyParameters::MakeContributorKey() const
 {
-	return PropertyInfo->HierarchicalMetaData.bBlendHierarchicalBias ? FContributorKey(PropertyInfoIndex) : FContributorKey(PropertyInfoIndex, PropertyInfo->HierarchicalMetaData.HBias);
+	return FContributorKey(PropertyInfoIndex);
 }
 
 UMovieScenePropertyInstantiatorSystem::UMovieScenePropertyInstantiatorSystem(const FObjectInitializer& ObjInit)
@@ -89,9 +71,10 @@ UMovieScenePropertyInstantiatorSystem::UMovieScenePropertyInstantiatorSystem(con
 	{
 		DefineComponentConsumer(GetClass(), BuiltInComponents->BoundObject);
 		DefineComponentConsumer(GetClass(), BuiltInComponents->Group);
+		DefineComponentConsumer(GetClass(), BuiltInComponents->HierarchicalBlendTarget);
+		DefineComponentConsumer(GetClass(), BuiltInComponents->Tags.Ignored);
 
 		DefineComponentProducer(GetClass(), BuiltInComponents->BlendChannelInput);
-		DefineComponentProducer(GetClass(), BuiltInComponents->HierarchicalBlendTarget);
 		DefineComponentProducer(GetClass(), BuiltInComponents->SymbolicTags.CreatesEntities);
 	}
 }
@@ -320,7 +303,7 @@ void UMovieScenePropertyInstantiatorSystem::DiscoverNewProperties(TBitArray<>&Ou
 
 	TArrayView<const FPropertyDefinition> Properties = this->BuiltInComponents->PropertyRegistry.GetProperties();
 
-	auto VisitNewProperties = [this, Properties, &OutInvalidatedProperties](FEntityAllocationIteratorItem AllocationItem, const FMovieSceneEntityID* EntityIDs, UObject* const * ObjectPtrs, const FMovieScenePropertyBinding* PropertyPtrs, const FEntityGroupID* GroupIDs, const int16* HierarchicalBiases)
+	auto VisitNewProperties = [this, Properties, &OutInvalidatedProperties](FEntityAllocationIteratorItem AllocationItem, const FMovieSceneEntityID* EntityIDs, UObject* const * ObjectPtrs, const FMovieScenePropertyBinding* PropertyPtrs, const FEntityGroupID* GroupIDs)
 	{
 		const FEntityAllocation* Allocation     = AllocationItem.GetAllocation();
 		const FComponentMask&    AllocationType = AllocationItem.GetAllocationType();
@@ -335,28 +318,38 @@ void UMovieScenePropertyInstantiatorSystem::DiscoverNewProperties(TBitArray<>&Ou
 
 		FCustomAccessorView CustomAccessors = PropertyDefinition.CustomPropertyRegistration ? PropertyDefinition.CustomPropertyRegistration->GetAccessors() : FCustomAccessorView();
 
-		// Figure out the hbias we should use for this allocation.
-		// If the allocation is tagged to ignore hbias, we null out the hbias components and use the ANY_HBIAS symbolic value
-		int16 DefaultHierarchicalBias = 0;
-		if (AllocationType.Contains(this->BuiltInComponents->Tags.IgnoreHierarchicalBias))
+		const bool bIgnored = AllocationType.Contains(BuiltInComponents->Tags.Ignored);
+		if (bIgnored)
 		{
-			HierarchicalBiases = nullptr;
-			DefaultHierarchicalBias = FContributorKey::ANY_HBIAS;
-		}
+			// Ignored properties should no longer contribute
+			for (int32 Index = 0; Index < Allocation->Num(); ++Index)
+			{
+				const int32 PropertyIndex = GroupIDs[Index].GroupIndex;
+				if (PropertyIndex != INDEX_NONE)
+				{
+					OutInvalidatedProperties.PadToNum(PropertyIndex + 1, false);
+					OutInvalidatedProperties[PropertyIndex] = true;
 
-		for (int32 Index = 0; Index < Allocation->Num(); ++Index)
+					this->Contributors.Remove(PropertyIndex, EntityIDs[Index]);
+				}
+			}
+		}
+		else
 		{
-			const bool bResolved = this->ResolveProperty(CustomAccessors, ObjectPtrs[Index], PropertyPtrs[Index], GroupIDs[Index], PropertyDefinitionIndex);
-			if (bResolved)
+			for (int32 Index = 0; Index < Allocation->Num(); ++Index)
 			{
-				const int32 PropertyIndex = GroupIDs[Index].GroupIndex;
+				const bool bResolved = this->ResolveProperty(CustomAccessors, ObjectPtrs[Index], PropertyPtrs[Index], GroupIDs[Index], PropertyDefinitionIndex);
+				if (bResolved)
+				{
+					const int32 PropertyIndex = GroupIDs[Index].GroupIndex;
 
-				FContributorKey Key { PropertyIndex, HierarchicalBiases ? HierarchicalBiases[Index] : DefaultHierarchicalBias };
-				this->Contributors.Add(Key, EntityIDs[Index]);
-				this->NewContributors.Add(Key, EntityIDs[Index]);
+					FContributorKey Key { PropertyIndex };
+					this->Contributors.Add(Key, EntityIDs[Index]);
+					this->NewContributors.Add(Key, EntityIDs[Index]);
 
-				OutInvalidatedProperties.PadToNum(PropertyIndex + 1, false);
-				OutInvalidatedProperties[PropertyIndex] = true;
+					OutInvalidatedProperties.PadToNum(PropertyIndex + 1, false);
+					OutInvalidatedProperties[PropertyIndex] = true;
+				}
 			}
 		}
 	};
@@ -366,7 +359,6 @@ void UMovieScenePropertyInstantiatorSystem::DiscoverNewProperties(TBitArray<>&Ou
 	.Read(BuiltInComponents->BoundObject)
 	.Read(BuiltInComponents->PropertyBinding)
 	.Read(BuiltInComponents->Group)
-	.ReadOptional(BuiltInComponents->HierarchicalBias)
 	.FilterNone({ BuiltInComponents->BlendChannelOutput })
 	.FilterAll({ BuiltInComponents->Tags.NeedsLink })
 	.Iterate_PerAllocation(&Linker->EntityManager, VisitNewProperties);
@@ -386,7 +378,6 @@ void UMovieScenePropertyInstantiatorSystem::DiscoverExpiredProperties(TBitArray<
 
 			this->Contributors.Remove(PropertyIndex, EntityID);
 		}
-
 	};
 
 	FEntityTaskBuilder()
@@ -558,7 +549,7 @@ void UMovieScenePropertyInstantiatorSystem::ProcessInvalidatedProperties(const T
 			StaleProperties[PropertyIndex] = true;
 		}
 		// Does it support fast path?
-		else if (Params.PropertyInfo->HierarchicalMetaData.bSupportsFastPath)
+		else if (Params.PropertyInfo->bSupportsFastPath)
 		{
 			InitializeFastPath(Params);
 		}
@@ -628,112 +619,72 @@ void UMovieScenePropertyInstantiatorSystem::UpdatePropertyInfo(const FPropertyPa
 
 	// Channel masks for all entities, entities within the active hbias, and within the 'ignored hbias' buckets
 	// Set bits denote channels that are not animated by entities in these contexts
-	FChannelMask BlendedBiasEmptyChannels(true, Params.PropertyDefinition->CompositeSize);
 	FChannelMask ActiveBiasEmptyChannels(true, Params.PropertyDefinition->CompositeSize);
-	FChannelMask IgnoredBiasEmptyChannels(true, Params.PropertyDefinition->CompositeSize);
 
 	// Key that visits any contributor to this property regardless of hbias
 	FContributorKey AnyContributor(Params.PropertyInfoIndex);
 
-	FHierarchicalMetaData IgnoredHBiasEntry;
-	FHierarchicalMetaData ActiveBiasEntry;
-
-	ActiveBiasEntry.HBias = TNumericLimits<int16>::Lowest();
+	int32 NumContributors = 0;
+	int16 HBias = 0;
+	bool bSupportsFastPath = true;
+	bool bWantsRestoreState = false;
+	bool bNeedsInitialValue = false;
+	bool bBlendHierarchicalBias = false;
 
 	// Iterate all contributors for this property to re-generate the meta-data
 	for (auto ContributorIt = Contributors.CreateConstKeyIterator(AnyContributor); ContributorIt; ++ContributorIt)
 	{
 		FMovieSceneEntityID Contributor = ContributorIt.Value();
-		const int16         ThisHBias   = ContributorIt.Key().HBias;
-		const bool          bIgnoreBias = ThisHBias == FContributorKey::ANY_HBIAS;
-
-		FHierarchicalMetaData* MetaDataToUpdate = nullptr;
-		if (bIgnoreBias)
-		{
-			MetaDataToUpdate = &IgnoredHBiasEntry;
-		}
-		else 
-		{
-			ActiveBiasEntry.BlendTarget.Add(ThisHBias);
-
-			if (ThisHBias >= ActiveBiasEntry.HBias)
-			{
-				MetaDataToUpdate = &ActiveBiasEntry;
 
-				if (ThisHBias > ActiveBiasEntry.HBias)
-				{
-					// We found a greater bias than any we've encountered this far
-					// Reset the empty channel list for the active hbias
-					ActiveBiasEmptyChannels = FChannelMask(true, Params.PropertyDefinition->CompositeSize);
-					ActiveBiasEntry.HBias = ThisHBias;
+		const FComponentMask& Type = Linker->EntityManager.GetEntityType(Contributor);
 
-					MetaDataToUpdate = &ActiveBiasEntry;
-				}
-			}
+		if (Type.Contains(BuiltInComponents->HierarchicalBias))
+		{
+			HBias = Linker->EntityManager.ReadComponentChecked(Contributor, BuiltInComponents->HierarchicalBias);
 		}
 
-		const bool bContributorIsActive = MetaDataToUpdate != nullptr;
-
 		// Update the various empty channel masks
 		for (int32 CompositeIndex = 0; CompositeIndex < Params.PropertyDefinition->CompositeSize; ++CompositeIndex)
 		{
-			const bool bCheckChannel =
-				BlendedBiasEmptyChannels[CompositeIndex] == true ||
-				(bIgnoreBias && IgnoredBiasEmptyChannels[CompositeIndex] == true) ||
-				(bContributorIsActive && ActiveBiasEmptyChannels[CompositeIndex] == true);
+			const bool bCheckChannel = ActiveBiasEmptyChannels[CompositeIndex] == true;
 
 			if (bCheckChannel)
 			{
 				FComponentTypeID ThisChannel = Composites[CompositeIndex].ComponentTypeID;
-				if (ThisChannel && Linker->EntityManager.HasComponent(Contributor, ThisChannel))
+				if (ThisChannel && Type.Contains(ThisChannel))
 				{
-					BlendedBiasEmptyChannels[CompositeIndex] = false;
-					if (bIgnoreBias)
-					{
-						IgnoredBiasEmptyChannels[CompositeIndex] = false;
-					}
-					if (bContributorIsActive)
-					{
-						ActiveBiasEmptyChannels[CompositeIndex] = false;
-					}
+					ActiveBiasEmptyChannels[CompositeIndex] = false;
 				}
 			}
 		}
 
-		if (!bContributorIsActive)
-		{
-			continue;
-		}
-
-		MetaDataToUpdate->bInUse = true;
-		++MetaDataToUpdate->NumContributors;
+		++NumContributors;
 
 		// Update whether this meta-data entry wants restore state
-		if (!MetaDataToUpdate->bWantsRestoreState && Linker->EntityManager.HasComponent(Contributor, BuiltInComponents->Tags.RestoreState))
+		if (!bWantsRestoreState && Type.Contains(BuiltInComponents->Tags.RestoreState))
 		{
-			MetaDataToUpdate->bWantsRestoreState = true;
+			bWantsRestoreState = true;
 		}
 
 		// Update whether this meta-data entry needs an initial value or not
-		if (!MetaDataToUpdate->bNeedsInitialValue && Linker->EntityManager.HasComponent(Contributor, BuiltInComponents->Tags.AlwaysCacheInitialValue))
+		if (!bNeedsInitialValue && Type.Contains(BuiltInComponents->Tags.AlwaysCacheInitialValue))
 		{
-			MetaDataToUpdate->bNeedsInitialValue = true;
+			bNeedsInitialValue = true;
 		}
 
-		const FComponentMask& Type = Linker->EntityManager.GetEntityType(Contributor);
 
-		if (!MetaDataToUpdate->bBlendHierarchicalBias && Type.Contains(BuiltInComponents->Tags.BlendHierarchicalBias))
+		if (!bBlendHierarchicalBias && Type.Contains(BuiltInComponents->Tags.BlendHierarchicalBias))
 		{
-			MetaDataToUpdate->bBlendHierarchicalBias = true;
-			MetaDataToUpdate->bSupportsFastPath = false;
+			bBlendHierarchicalBias = true;
+			bSupportsFastPath = false;
 		}
 
 		// Update whether this property supports fast path
-		if (MetaDataToUpdate->bSupportsFastPath)
+		if (bSupportsFastPath)
 		{
-			if (MetaDataToUpdate->NumContributors > 1)
+			if (NumContributors > 1)
 			{
-				MetaDataToUpdate->bSupportsFastPath = false;
+				bSupportsFastPath = false;
 			}
 			else
 			{
@@ -742,35 +693,15 @@ void UMovieScenePropertyInstantiatorSystem::UpdatePropertyInfo(const FPropertyPa
 						Type.Contains(BuiltInComponents->Tags.AdditiveFromBaseBlend) || 
 						Type.Contains(BuiltInComponents->WeightAndEasingResult))
 				{
-					MetaDataToUpdate->bSupportsFastPath = false;
+					bSupportsFastPath = false;
 				}
 			}
 		}
 	}
 
-	// Add any ignored hbias entries to the first tracked meta data that is in use
-	if (IgnoredHBiasEntry.bInUse)
-	{
-		// Combine using a bitwise & since channels are only empty if they are empty in both
-		ActiveBiasEmptyChannels.CombineWithBitwiseAND(IgnoredBiasEmptyChannels, EBitwiseOperatorFlags::MaintainSize);
-		if (ActiveBiasEntry.bInUse)
-		{
-			ActiveBiasEntry.CombineWith(IgnoredHBiasEntry);
-		}
-		else
-		{
-			ActiveBiasEntry = MoveTemp(IgnoredHBiasEntry);
-		}
-	}
-
-	// -----------------------------------
-	// NOW UNSAFE TO USE IgnoredHBiasEntry
-
 	// Reset the restore state status of the property if we still have contributors
 	// We do not do this if there are no contributors to ensure that stale properties are restored correctly
-	Params.PropertyInfo->EmptyChannels = ActiveBiasEntry.bBlendHierarchicalBias
-		? BlendedBiasEmptyChannels
-		: ActiveBiasEmptyChannels;
+	Params.PropertyInfo->EmptyChannels = ActiveBiasEmptyChannels;
 
 	const bool bWasPartial = Params.PropertyInfo->bIsPartiallyAnimated;
 	const bool bIsPartial  = Params.PropertyInfo->EmptyChannels.Find(true) != INDEX_NONE;
@@ -781,91 +712,80 @@ void UMovieScenePropertyInstantiatorSystem::UpdatePropertyInfo(const FPropertyPa
 		PropertyStats[StatIndex].NumPartialProperties += bIsPartial ? 1 : -1;
 	}
 
-	Params.PropertyInfo->bIsPartiallyAnimated = bIsPartial;
-	Params.PropertyInfo->bMaxHBiasHasChanged  = Params.PropertyInfo->HierarchicalMetaData.HBias != ActiveBiasEntry.HBias;
-	Params.PropertyInfo->HierarchicalMetaData = MoveTemp(ActiveBiasEntry);
+	Params.PropertyInfo->bIsPartiallyAnimated   = bIsPartial;
+	Params.PropertyInfo->bMaxHBiasHasChanged    = Params.PropertyInfo->HBias != HBias;
+	Params.PropertyInfo->HBias                  = HBias;
+	Params.PropertyInfo->bSupportsFastPath      = bSupportsFastPath;
+	Params.PropertyInfo->bWantsRestoreState     = bWantsRestoreState;
+	Params.PropertyInfo->bNeedsInitialValue     = bNeedsInitialValue;
 }
 
 void UMovieScenePropertyInstantiatorSystem::InitializeFastPath(const FPropertyParameters& Params)
 {
 	using namespace UE::MovieScene;
 
-	FTypelessMutation IgnoredContributorMutation;
-	IgnoredContributorMutation.RemoveMask = CleanFastPathMask;
-	IgnoredContributorMutation.RemoveMask.Set(BuiltInComponents->BlendChannelInput);
-	IgnoredContributorMutation.RemoveMask.Set(Params.PropertyDefinition->InitialValueType);
-	IgnoredContributorMutation.RemoveMask.Set(BuiltInComponents->Tags.HasAssignedInitialValue);
-	IgnoredContributorMutation.AddMask.Set(BuiltInComponents->Tags.Ignored);
-
 	// Find the sole contributor with the specific property info and hbias
-	const int16 ActiveHBias = Params.PropertyInfo->HierarchicalMetaData.HBias;
+	const int16 ActiveHBias = Params.PropertyInfo->HBias;
 	FContributorKey AnyContributor(Params.PropertyInfoIndex);
 	for (auto ContributorIt = Contributors.CreateConstKeyIterator(AnyContributor); ContributorIt; ++ContributorIt)
 	{
 		FMovieSceneEntityID Contributor = ContributorIt.Value();
 
-		if (ContributorIt.Key().HBias == ActiveHBias)
+		FTypelessMutation SoleContributorMutation;
+		SoleContributorMutation.RemoveMask.SetAll({ BuiltInComponents->BlendChannelInput, BuiltInComponents->HierarchicalBlendTarget });
+
+		if (Params.PropertyInfo->bNeedsInitialValue)
 		{
-			FTypelessMutation SoleContributorMutation;
-			SoleContributorMutation.RemoveMask.SetAll({ BuiltInComponents->Tags.Ignored, BuiltInComponents->BlendChannelInput, BuiltInComponents->HierarchicalBlendTarget });
+			SoleContributorMutation.AddMask.Set(Params.PropertyDefinition->InitialValueType);
+		}
 
-			if (Params.PropertyInfo->HierarchicalMetaData.bNeedsInitialValue)
-			{
-				SoleContributorMutation.AddMask.Set(Params.PropertyDefinition->InitialValueType);
-			}
+		if (Params.PropertyDefinition->MetaDataTypes.Num() > 0)
+		{
+			InitializePropertyMetaDataTasks.PadToNum(Params.PropertyInfo->PropertyDefinitionIndex+1, false);
+			InitializePropertyMetaDataTasks[Params.PropertyInfo->PropertyDefinitionIndex] = true;
 
-			if (Params.PropertyDefinition->MetaDataTypes.Num() > 0)
+			for (FComponentTypeID Component : Params.PropertyDefinition->MetaDataTypes)
 			{
-				InitializePropertyMetaDataTasks.PadToNum(Params.PropertyInfo->PropertyDefinitionIndex+1, false);
-				InitializePropertyMetaDataTasks[Params.PropertyInfo->PropertyDefinitionIndex] = true;
-
-				for (FComponentTypeID Component : Params.PropertyDefinition->MetaDataTypes)
-				{
-					SoleContributorMutation.AddMask.Set(Component);
-				}
+				SoleContributorMutation.AddMask.Set(Component);
 			}
+		}
 
-			// Ensure the sole contributor is set up to apply the property as a final output
-			switch (Params.PropertyInfo->Property.GetIndex())
-			{
-			case 0:
-				FEntityBuilder()
-				.Add(BuiltInComponents->FastPropertyOffset, Params.PropertyInfo->Property.template Get<uint16>())
-				.MutateExisting(&Linker->EntityManager, Contributor, SoleContributorMutation);
-				break;
-			case 1:
-				FEntityBuilder()
-				.Add(BuiltInComponents->CustomPropertyIndex, Params.PropertyInfo->Property.template Get<FCustomPropertyIndex>())
-				.MutateExisting(&Linker->EntityManager, Contributor, SoleContributorMutation);
-				break;
-			case 2:
-				FEntityBuilder()
-				.Add(BuiltInComponents->SlowProperty, Params.PropertyInfo->Property.template Get<FSlowPropertyPtr>())
-				.MutateExisting(&Linker->EntityManager, Contributor, SoleContributorMutation);
-				break;
-			}
+		// Ensure the sole contributor is set up to apply the property as a final output
+		switch (Params.PropertyInfo->Property.GetIndex())
+		{
+		case 0:
+			FEntityBuilder()
+			.Add(BuiltInComponents->FastPropertyOffset, Params.PropertyInfo->Property.template Get<uint16>())
+			.MutateExisting(&Linker->EntityManager, Contributor, SoleContributorMutation);
+			break;
+		case 1:
+			FEntityBuilder()
+			.Add(BuiltInComponents->CustomPropertyIndex, Params.PropertyInfo->Property.template Get<FCustomPropertyIndex>())
+			.MutateExisting(&Linker->EntityManager, Contributor, SoleContributorMutation);
+			break;
+		case 2:
+			FEntityBuilder()
+			.Add(BuiltInComponents->SlowProperty, Params.PropertyInfo->Property.template Get<FSlowPropertyPtr>())
+			.MutateExisting(&Linker->EntityManager, Contributor, SoleContributorMutation);
+			break;
+		}
 
-			// Copy initial values and meta-data back off our old blend output
-			FMovieSceneEntityID OldOutput = Params.PropertyInfo->PreviousFastPathID ? Params.PropertyInfo->PreviousFastPathID : Params.PropertyInfo->FinalBlendOutputID;
-			if (OldOutput)
+		// Copy initial values and meta-data back off our old blend output
+		FMovieSceneEntityID OldOutput = Params.PropertyInfo->PreviousFastPathID ? Params.PropertyInfo->PreviousFastPathID : Params.PropertyInfo->FinalBlendOutputID;
+		if (OldOutput)
+		{
+			FComponentMask CopyMask;
+			CopyMask.Set(Params.PropertyDefinition->InitialValueType);
+			CopyMask.Set(BuiltInComponents->Tags.HasAssignedInitialValue);
+			for (FComponentTypeID Component : Params.PropertyDefinition->MetaDataTypes)
 			{
-				FComponentMask CopyMask;
-				CopyMask.Set(Params.PropertyDefinition->InitialValueType);
-				CopyMask.Set(BuiltInComponents->Tags.HasAssignedInitialValue);
-				for (FComponentTypeID Component : Params.PropertyDefinition->MetaDataTypes)
-				{
-					CopyMask.Set(Component);
-				}
-
-				Linker->EntityManager.CopyComponents(OldOutput, Contributor, CopyMask);
+				CopyMask.Set(Component);
 			}
 
-			Params.PropertyInfo->PreviousFastPathID = Contributor;
-		}
-		else
-		{
-			Linker->EntityManager.ChangeEntityType(Contributor, IgnoredContributorMutation.MutateType(Linker->EntityManager.GetEntityType(Contributor)));
+			Linker->EntityManager.CopyComponents(OldOutput, Contributor, CopyMask);
 		}
+
+		Params.PropertyInfo->PreviousFastPathID = Contributor;
 	}
 
 	// If this was previously blended, destroy the blend output
@@ -999,20 +919,12 @@ void UMovieScenePropertyInstantiatorSystem::InitializeBlendPath(const FPropertyP
 			? Contributors.CreateConstKeyIterator(ContributorKey)
 			: NewContributors.CreateConstKeyIterator(ContributorKey);
 
-		FTypelessMutation Mutation;
-		if (!Params.PropertyInfo->HierarchicalMetaData.bBlendHierarchicalBias)
-		{
-			// Make sure that the hierarchical blend target component does not exist if it no longer has one
-			Mutation.AddMask.Set(BuiltInComponents->Tags.RemoveHierarchicalBlendTarget);
-		}
-
 		for (; ContributorIt; ++ContributorIt)
 		{
 			FEntityBuilder()
 			.Add(BuiltInComponents->BlendChannelInput, BlendChannel)
 			.AddTag(SetupResult.CurrentInfo.BlenderTypeTag)
-			.AddConditional(BuiltInComponents->HierarchicalBlendTarget, Params.PropertyInfo->HierarchicalMetaData.BlendTarget, Params.PropertyInfo->HierarchicalMetaData.bBlendHierarchicalBias)
-			.MutateExisting(&Linker->EntityManager, ContributorIt.Value(), Mutation);
+			.MutateExisting(&Linker->EntityManager, ContributorIt.Value());
 		}
 
 		check(!Linker->EntityManager.HasComponent(Params.PropertyInfo->FinalBlendOutputID, BuiltInComponents->BlendChannelInput));
@@ -1033,11 +945,6 @@ void UMovieScenePropertyInstantiatorSystem::InitializeBlendPath(const FPropertyP
 	InputMutation.RemoveMask = CleanFastPathMask;
 	InputMutation.RemoveMask.Set(Params.PropertyDefinition->InitialValueType);
 	InputMutation.RemoveMask.Set(BuiltInComponents->Tags.HasAssignedInitialValue);
-	if (!Params.PropertyInfo->HierarchicalMetaData.bBlendHierarchicalBias)
-	{
-		// Make sure that the hierarchical blend target component does not exist if it no longer has one
-		InputMutation.AddMask.Set(BuiltInComponents->Tags.RemoveHierarchicalBlendTarget);
-	}
 	for (FComponentTypeID Component : Params.PropertyDefinition->MetaDataTypes)
 	{
 		InputMutation.RemoveMask.Set(Component);
@@ -1110,7 +1017,7 @@ void UMovieScenePropertyInstantiatorSystem::InitializeBlendPath(const FPropertyP
 		.Add(BuiltInComponents->BlendChannelOutput,      NewBlendChannel)
 		.Add(BuiltInComponents->PropertyBinding,         Params.PropertyInfo->PropertyBinding)
 		.Add(BuiltInComponents->BoundObject,             Params.PropertyInfo->BoundObject)
-		.AddTagConditional(BuiltInComponents->Tags.RestoreState, Params.PropertyInfo->HierarchicalMetaData.bWantsRestoreState)
+		.AddTagConditional(BuiltInComponents->Tags.RestoreState, Params.PropertyInfo->bWantsRestoreState)
 		.AddTag(SetupResult.CurrentInfo.BlenderTypeTag)
 		.AddTag(BuiltInComponents->Tags.NeedsLink)
 		.AddMutualComponents();
@@ -1169,7 +1076,6 @@ void UMovieScenePropertyInstantiatorSystem::InitializeBlendPath(const FPropertyP
 
 		FEntityBuilder()
 		.Add(BuiltInComponents->BlendChannelInput, NewBlendChannel)
-		.AddConditional(BuiltInComponents->HierarchicalBlendTarget, Params.PropertyInfo->HierarchicalMetaData.BlendTarget, Params.PropertyInfo->HierarchicalMetaData.bBlendHierarchicalBias)
 		.AddTag(SetupResult.CurrentInfo.BlenderTypeTag)
 		.MutateExisting(&Linker->EntityManager, Contributor, InputMutation);
 	}
@@ -1278,7 +1184,7 @@ void UMovieScenePropertyInstantiatorSystem::FPropertyParameters::MakeOutputCompo
 	OutComponentType.Set(PropertyDefinition->PropertyType);
 
 	// Set the restore state tag appropriately
-	if (PropertyInfo->HierarchicalMetaData.bWantsRestoreState)
+	if (PropertyInfo->bWantsRestoreState)
 	{
 		OutComponentType.Set(FBuiltInComponentTypes::Get()->Tags.RestoreState);
 	}
diff --git a/Engine/Source/Runtime/MovieSceneTracks/Private/Systems/WeightAndEasingEvaluatorSystem.cpp b/Engine/Source/Runtime/MovieSceneTracks/Private/Systems/WeightAndEasingEvaluatorSystem.cpp
--- a/Engine/Source/Runtime/MovieSceneTracks/Private/Systems/WeightAndEasingEvaluatorSystem.cpp
+++ b/Engine/Source/Runtime/MovieSceneTracks/Private/Systems/WeightAndEasingEvaluatorSystem.cpp
@@ -428,13 +428,12 @@ struct FConsumerEasingChannelMutation : FEasingChannelMutationBase
 		TOptionalComponentWriter<uint16>                   ExistingEasingChannels = Allocation->TryWriteComponents(BuiltInComponents->HierarchicalEasingChannel, FEntityAllocationWriteContext::NewAllocation());
 		TOptionalComponentReader<FHierarchicalBlendTarget> BlendTargets           = Allocation->TryReadComponents(BuiltInComponents->HierarchicalBlendTarget);
 
-		const bool bRemoveBlendTarget = Allocation->HasComponent(BuiltInComponents->Tags.RemoveHierarchicalBlendTarget);
 		const bool bHasExistingEasing = ExistingEasingChannels.IsValid();
 
 		// Loop through each entity and check whether it has easing.
 		// If so, either assign the new easing channel or mark it for mutation.
 		// If not, mark it for removal
-		if (BlendTargets && !bRemoveBlendTarget)
+		if (BlendTargets)
 		{
 			check(CachedHierarchicalBlendTargetChannels);
 			TOptionalComponentReader<int16> HierarchicalBias = Allocation->TryReadComponents(BuiltInComponents->HierarchicalBias);
@@ -874,6 +873,9 @@ void UMovieSceneHierarchicalEasingInstantiatorSystem::OnRun(FSystemTaskPrerequis
 
 		ConsumerMutation.RemoveStaleComponents(Linker);
 	}
+
+
+	FinalizeBlendTargets();
 }
 
 void UMovieSceneHierarchicalEasingInstantiatorSystem::FinalizeBlendTargets()
@@ -1021,12 +1023,6 @@ void UMovieSceneHierarchicalEasingInstantiatorSystem::FinalizeBlendTargets()
 		ConsumerMutation.RemoveStaleComponents(Linker);
 	}
 
-	if (Linker->EntityManager.ContainsComponent(BuiltInComponents->Tags.RemoveHierarchicalBlendTarget))
-	{
-		FRemoveSingleMutation RemoveTag(BuiltInComponents->Tags.RemoveHierarchicalBlendTarget);
-		Linker->EntityManager.MutateAll(FEntityComponentFilter().All({ BuiltInComponents->Tags.RemoveHierarchicalBlendTarget }), RemoveTag);
-	}
-
 
 #if UE_MOVIESCENE_EXPENSIVE_CONSISTENCY_CHECKS
 	FHierarchicalEasingChannelBuffer& Buffer = EvaluatorSystem->GetComputationBuffer();
@@ -1049,32 +1045,6 @@ void UMovieSceneHierarchicalEasingInstantiatorSystem::FinalizeBlendTargets()
 UMovieSceneHierarchicalEasingFinalizationSystem::UMovieSceneHierarchicalEasingFinalizationSystem(const FObjectInitializer& ObjInit)
 	: Super(ObjInit)
 {
-	using namespace UE::MovieScene;
-
-	SystemCategories = EEntitySystemCategory::Core;
-	RelevantComponent = FBuiltInComponentTypes::Get()->HierarchicalBlendTarget;
-
-	if (HasAnyFlags(RF_ClassDefaultObject))
-	{
-		// Has to come after property instantiation property instantiation so that initial values get created correctly
-		DefineImplicitPrerequisite(UMovieScenePropertyInstantiatorSystem::StaticClass(), GetClass());
-		DefineImplicitPrerequisite(UMovieSceneMaterialParameterInstantiatorSystem::StaticClass(), GetClass());
-		DefineImplicitPrerequisite(UMovieSceneHierarchicalEasingInstantiatorSystem::StaticClass(), GetClass());
-
-		DefineComponentConsumer(GetClass(), FBuiltInComponentTypes::Get()->HierarchicalBlendTarget);
-		DefineComponentProducer(GetClass(), FBuiltInComponentTypes::Get()->WeightAndEasingResult);
-	}
-}
-
-void UMovieSceneHierarchicalEasingFinalizationSystem::OnLink()
-{
-	InstantiatorSystem = Linker->LinkSystem<UMovieSceneHierarchicalEasingInstantiatorSystem>();
-	Linker->SystemGraph.AddReference(this, InstantiatorSystem);
-}
-
-void UMovieSceneHierarchicalEasingFinalizationSystem::OnRun(FSystemTaskPrerequisites& InPrerequisites, FSystemSubsequentTasks& Subsequents)
-{
-	InstantiatorSystem->FinalizeBlendTargets();
 }
 
 UWeightAndEasingEvaluatorSystem::UWeightAndEasingEvaluatorSystem(const FObjectInitializer& ObjInit)
diff --git a/Engine/Source/Runtime/MovieSceneTracks/Public/Systems/MovieSceneCustomPrimitiveDataSystem.h b/Engine/Source/Runtime/MovieSceneTracks/Public/Systems/MovieSceneCustomPrimitiveDataSystem.h
--- a/Engine/Source/Runtime/MovieSceneTracks/Public/Systems/MovieSceneCustomPrimitiveDataSystem.h
+++ b/Engine/Source/Runtime/MovieSceneTracks/Public/Systems/MovieSceneCustomPrimitiveDataSystem.h
@@ -72,6 +72,9 @@ private:
 
 	/** Holds pre-animated values for custom primitive data entries */
 	TSharedPtr<UE::MovieScene::FPreAnimatedCustomPrimitiveDataEntryStorage> ScalarParameterStorage;
+
+	UE::MovieScene::FEntityGroupingPolicyKey GroupingKey;
+
 public:
 
 	UPROPERTY()
diff --git a/Engine/Source/Runtime/MovieSceneTracks/Public/Systems/MovieSceneHierarchicalBiasSystem.h b/Engine/Source/Runtime/MovieSceneTracks/Public/Systems/MovieSceneHierarchicalBiasSystem.h
--- a/Engine/Source/Runtime/MovieSceneTracks/Public/Systems/MovieSceneHierarchicalBiasSystem.h
+++ b/Engine/Source/Runtime/MovieSceneTracks/Public/Systems/MovieSceneHierarchicalBiasSystem.h
@@ -9,8 +9,9 @@
 
 #include "MovieSceneHierarchicalBiasSystem.generated.h"
 
-class UMovieSceneEntitySystemLinker;
 class UObject;
+class UMovieSceneEntitySystemLinker;
+class UMovieSceneEntityGroupingSystem;
 
 UCLASS()
 class UMovieSceneHierarchicalBiasSystem : public UMovieSceneEntityInstantiatorSystem
@@ -24,5 +25,9 @@ public:
 private:
 
 	virtual bool IsRelevantImpl(UMovieSceneEntitySystemLinker* InLinker) const override;
+	virtual void OnLink() override;
 	virtual void OnRun(FSystemTaskPrerequisites& InPrerequisites, FSystemSubsequentTasks& Subsequents) override;
+
+	UPROPERTY()
+	TObjectPtr<UMovieSceneEntityGroupingSystem> GroupingSystem;
 };
diff --git a/Engine/Source/Runtime/MovieSceneTracks/Public/Systems/MovieSceneMaterialSystem.h b/Engine/Source/Runtime/MovieSceneTracks/Public/Systems/MovieSceneMaterialSystem.h
--- a/Engine/Source/Runtime/MovieSceneTracks/Public/Systems/MovieSceneMaterialSystem.h
+++ b/Engine/Source/Runtime/MovieSceneTracks/Public/Systems/MovieSceneMaterialSystem.h
@@ -105,6 +105,82 @@ public:
 
 protected:
 
+	struct FMaterialGroupingPolicy
+	{
+		using GroupKeyType = TTuple<RequiredComponents..., FMaterialParameterInfo>;
+
+		void InitializeGroupKeys(
+			TEntityGroupingHandlerBase<FMaterialGroupingPolicy>& Handler,
+			FEntityGroupBuilder* Builder,
+			FEntityAllocationIteratorItem Item,
+			FReadEntityIDs EntityIDs,
+			TWrite<FEntityGroupID> GroupIDs,
+			TRead<RequiredComponents>... Components)
+		{
+			const FComponentMask& AllocationType = Item.GetAllocationType();
+
+			FMovieSceneTracksComponentTypes* TracksComponents = FMovieSceneTracksComponentTypes::Get();
+
+			TComponentTypeID<FName> ParameterName;
+			if (AllocationType.Contains(TracksComponents->ScalarParameterName))
+			{
+				ParameterName = TracksComponents->ScalarParameterName;
+			}
+			else if (AllocationType.Contains(TracksComponents->VectorParameterName))
+			{
+				ParameterName = TracksComponents->VectorParameterName;
+			}
+			else if (AllocationType.Contains(TracksComponents->ColorParameterName))
+			{
+				ParameterName = TracksComponents->ColorParameterName;
+			}
+
+			TComponentTypeID<FMaterialParameterInfo> ParameterInfo;
+			if (AllocationType.Contains(TracksComponents->ScalarMaterialParameterInfo))
+			{
+				ParameterInfo = TracksComponents->ScalarMaterialParameterInfo;
+			}
+			else if (AllocationType.Contains(TracksComponents->VectorMaterialParameterInfo))
+			{
+				ParameterInfo = TracksComponents->VectorMaterialParameterInfo;
+			}
+			else if (AllocationType.Contains(TracksComponents->ColorMaterialParameterInfo))
+			{
+				ParameterInfo = TracksComponents->ColorMaterialParameterInfo;
+			}
+
+
+			const FEntityAllocation* Allocation = Item.GetAllocation();
+			const int32              Num        = Allocation->Num();
+
+			TReadOptional<FName>                  ParameterNames = ParameterName ? Allocation->TryReadComponents(ParameterName) : TReadOptional<FName>();
+			TReadOptional<FMaterialParameterInfo> ParameterInfos = ParameterInfo ? Allocation->TryReadComponents(ParameterInfo) : TReadOptional<FMaterialParameterInfo>();
+
+			for (int32 Index = 0; Index < Num; ++Index)
+			{
+				GroupKeyType Key = ParameterInfos
+					? MakeTuple(Components[Index]..., ParameterInfos[Index])
+					: ParameterNames 
+						? MakeTuple(Components[Index]..., FMaterialParameterInfo(ParameterNames[Index]))
+						: MakeTuple(Components[Index]..., FMaterialParameterInfo());
+
+				const int32    NewGroupIndex = Handler.GetOrAllocateGroupIndex(Key, Builder);
+				FEntityGroupID NewGroupID    = Builder->MakeGroupID(NewGroupIndex);
+
+				Builder->AddEntityToGroup(EntityIDs[Index], NewGroupID);
+				// Write out the group ID component
+				GroupIDs[Index] = NewGroupID;
+			}
+		}
+
+	#if WITH_EDITOR
+		bool OnObjectsReplaced(GroupKeyType& InOutKey, const TMap<UObject*, UObject*>& ReplacementMap)
+		{
+			return false;
+		}
+	#endif
+	};
+
 	UE::MovieScene::FEntityGroupingPolicyKey GroupingKey;
 
 	FEntityComponentFilter MaterialSwitcherFilter;
@@ -157,8 +233,13 @@ struct TInitializeBoundMaterials
 
 		if (!ExistingMaterial)
 		{
-			OutDynamicMaterial = FObjectComponent::Null();
-			return true;
+			// If the object was not previously explicitly null, make it so
+			if (OutDynamicMaterial != FObjectComponent::Null())
+			{
+				OutDynamicMaterial = FObjectComponent::Null();
+				return true;
+			}
+			return false;
 		}
 
 		if (UMaterialInstanceDynamic* MID = Cast<UMaterialInstanceDynamic>(ExistingMaterial))
@@ -276,7 +357,7 @@ void TMovieSceneMaterialSystem<AccessorType, RequiredComponents...>::OnLink(UMov
 
 	// Define a grouping for these materials. This will make hierarchical bias work.
 	UMovieSceneEntityGroupingSystem* GroupingSystem = Linker->LinkSystem<UMovieSceneEntityGroupingSystem>();
-	GroupingKey = GroupingSystem->AddGrouping<RequiredComponents...>(InRequiredComponents...);
+	GroupingKey = GroupingSystem->AddGrouping(FMaterialGroupingPolicy(), InRequiredComponents...);
 
 	MaterialSwitcherFilter.Reset();
 	MaterialSwitcherFilter.All({ InRequiredComponents..., BuiltInComponents->ObjectResult });
diff --git a/Engine/Source/Runtime/MovieSceneTracks/Public/Systems/MovieScenePropertyInstantiator.h b/Engine/Source/Runtime/MovieSceneTracks/Public/Systems/MovieScenePropertyInstantiator.h
--- a/Engine/Source/Runtime/MovieSceneTracks/Public/Systems/MovieScenePropertyInstantiator.h
+++ b/Engine/Source/Runtime/MovieSceneTracks/Public/Systems/MovieScenePropertyInstantiator.h
@@ -112,19 +112,11 @@ private:
 
 	struct FContributorKey
 	{
-		static constexpr int16 ANY_HBIAS = TNumericLimits<int16>::Max();
-
 		FContributorKey(int32 InPropertyIndex)
 			: PropertyIndex(InPropertyIndex)
-			, HBias(ANY_HBIAS)
-		{}
-		FContributorKey(int32 InPropertyIndex, int16 InHBias)
-			: PropertyIndex(InPropertyIndex)
-			, HBias(InHBias)
 		{}
 
 		int32 PropertyIndex;
-		int16 HBias;
 
 		friend uint32 GetTypeHash(FContributorKey In)
 		{
@@ -133,45 +125,19 @@ private:
 
 		friend bool operator==(FContributorKey A, FContributorKey B)
 		{
-			return A.PropertyIndex == B.PropertyIndex && (
-				A.HBias == B.HBias || A.HBias == ANY_HBIAS || B.HBias == ANY_HBIAS
-			);
+			return A.PropertyIndex == B.PropertyIndex;
 		}
 	};
 
-	struct FHierarchicalMetaData
-	{
-		FHierarchicalMetaData()
-			: NumContributors(0)
-			, HBias(0)
-		{
-			bWantsRestoreState = false;
-			bSupportsFastPath = true;
-			bNeedsInitialValue = false;
-			bBlendHierarchicalBias = false;
-			bInUse = false;
-		}
-
-		UE::MovieScene::FHierarchicalBlendTarget BlendTarget;
-		int32 NumContributors = 0;
-		int16 HBias = 0;
-		uint8 bWantsRestoreState : 1;
-		uint8 bSupportsFastPath : 1;
-		uint8 bNeedsInitialValue : 1;
-		uint8 bBlendHierarchicalBias : 1;
-		uint8 bInUse : 1;
-
-		void CombineWith(const FHierarchicalMetaData& Other);
-
-		void ResetTracking();
-	};
-
 	struct FObjectPropertyInfo
 	{
 		FObjectPropertyInfo(UE::MovieScene::FResolvedProperty&& InProperty)
 			: Property(MoveTemp(InProperty))
 			, BoundObject(nullptr)
 			, BlendChannel(FMovieSceneBlendChannelID::INVALID_BLEND_CHANNEL)
+			, bWantsRestoreState( false)
+			, bSupportsFastPath( true)
+			, bNeedsInitialValue( false)
 			, bMaxHBiasHasChanged(false)
 			, bIsPartiallyAnimated(false)
 		{}
@@ -195,7 +161,10 @@ private:
 		int32 PropertyDefinitionIndex;
 		/** Index of a float-based property if this property has been set for float-to-double conversion */
 		TOptional<int32> ConvertedFromPropertyDefinitionIndex;
-		FHierarchicalMetaData HierarchicalMetaData;
+		int16 HBias = 0;
+		uint8 bWantsRestoreState : 1;
+		uint8 bSupportsFastPath : 1;
+		uint8 bNeedsInitialValue : 1;
 		uint8 bMaxHBiasHasChanged : 1;
 		uint8 bIsPartiallyAnimated : 1;
 	};
diff --git a/Engine/Source/Runtime/MovieSceneTracks/Public/Systems/WeightAndEasingEvaluatorSystem.h b/Engine/Source/Runtime/MovieSceneTracks/Public/Systems/WeightAndEasingEvaluatorSystem.h
--- a/Engine/Source/Runtime/MovieSceneTracks/Public/Systems/WeightAndEasingEvaluatorSystem.h
+++ b/Engine/Source/Runtime/MovieSceneTracks/Public/Systems/WeightAndEasingEvaluatorSystem.h
@@ -151,20 +151,12 @@ private:
  * 
  */
 UCLASS()
-class UMovieSceneHierarchicalEasingFinalizationSystem : public UMovieSceneEntityInstantiatorSystem
+class UE_DEPRECATED(5.5, "This system is no longer functional") UMovieSceneHierarchicalEasingFinalizationSystem : public UMovieSceneEntityInstantiatorSystem
 {
 public:
 	GENERATED_BODY()
 
 	UMovieSceneHierarchicalEasingFinalizationSystem(const FObjectInitializer& ObjInit);
-
-	virtual void OnLink() override;
-	virtual void OnRun(FSystemTaskPrerequisites& InPrerequisites, FSystemSubsequentTasks& Subsequents) override;
-
-private:
-
-	UPROPERTY()
-	TObjectPtr<UMovieSceneHierarchicalEasingInstantiatorSystem> InstantiatorSystem;
 };
 
 /**
