Sequencer: Back out unnecessary change that removes spaces and makes a unique object name because it makes the spawnable no longer match the spawnable name in the level sequence.

#rnx

diff --git a/Engine/Source/Runtime/MovieSceneTracks/Private/Bindings/MovieSceneSpawnableActorBinding.cpp b/Engine/Source/Runtime/MovieSceneTracks/Private/Bindings/MovieSceneSpawnableActorBinding.cpp
--- a/Engine/Source/Runtime/MovieSceneTracks/Private/Bindings/MovieSceneSpawnableActorBinding.cpp
+++ b/Engine/Source/Runtime/MovieSceneTracks/Private/Bindings/MovieSceneSpawnableActorBinding.cpp
@@ -42,8 +42,7 @@ UObject* UMovieSceneSpawnableActorBindingBase::SpawnObjectInternal(UWorld* World
 	// Spawn the actor
 	FActorSpawnParameters SpawnInfo;
 	{
-		SpawnInfo.Name = *SpawnName.ToString().Replace(TEXT(" "), TEXT("_"));
-		SpawnInfo.Name = MakeUniqueObjectName(WorldContext->PersistentLevel.Get(), ActorClass, SpawnInfo.Name);
+		SpawnInfo.Name = SpawnName;
 		SpawnInfo.ObjectFlags = SpawnObjectFlags;
 		SpawnInfo.SpawnCollisionHandlingOverride = ESpawnActorCollisionHandlingMethod::AlwaysSpawn;
 		// allow pre-construction variables to be set.
