Constraints: fix constraint initializatin with MRQ

#rnx

diff --git a/Engine/Source/Runtime/MovieSceneTracks/Private/Systems/MovieSceneConstraintSystem.cpp b/Engine/Source/Runtime/MovieSceneTracks/Private/Systems/MovieSceneConstraintSystem.cpp
--- a/Engine/Source/Runtime/MovieSceneTracks/Private/Systems/MovieSceneConstraintSystem.cpp
+++ b/Engine/Source/Runtime/MovieSceneTracks/Private/Systems/MovieSceneConstraintSystem.cpp
@@ -192,7 +192,7 @@ struct FEvaluateConstraintChannels
 				UpdateHandle.Constraint = TransformConstraint;
 				UpdateHandle.TransformHandle = ComponentHandle;
 				System->DynamicOffsets.Add(UpdateHandle);
-				TransformConstraint->EnsurePrimaryDependency(BoundObject->GetWorld());
+				TransformConstraint->EnsurePrimaryDependency(World);
 			}
 		}
 		else
