Expose VolumetricFogStartDistance and VolumetricFogNearFadeInDistance to Sequencer.

- This requires exposing custom setters and getters since we need to mark the render state of the fog component dirty when these properties are set.
- Adding copy/paste boilerplate to expose these two properties similar to the existing VolumetricFogAlbedo property.
- Add EditAnywhere to VolumetricFogStartDistance to match the other properties on the component.  Interp isn't needed since we have a custom setter function.
- Add blueprint callable SetVolumetricFogStartDistance and SetVolumetricFogNearFadeInDistance to the component.
- In MovieSceneTracksCustomAccessors.cpp bind the setters and getters to the properties.

#test Adding a volumetric fog component to the scene, adding a level sequence and animating VolumetricFogStartDistance and VolumetricFogNearFadeInDistance

diff --git a/Engine/Source/Runtime/Engine/Classes/Components/ExponentialHeightFogComponent.h b/Engine/Source/Runtime/Engine/Classes/Components/ExponentialHeightFogComponent.h
--- a/Engine/Source/Runtime/Engine/Classes/Components/ExponentialHeightFogComponent.h
+++ b/Engine/Source/Runtime/Engine/Classes/Components/ExponentialHeightFogComponent.h
@@ -180,7 +180,7 @@ class UExponentialHeightFogComponent : public USceneComponent
 	/** 
 	 * Distance from the camera that the volumetric fog will start, in world units. 
 	 */
-	UPROPERTY(BlueprintReadOnly, interp, Category= VolumetricFog, meta=(DisplayName = "Start Distance", UIMin = "0", UIMax = "5000"))
+	UPROPERTY(EditAnywhere, BlueprintReadOnly, Category= VolumetricFog, meta=(DisplayName = "Start Distance", UIMin = "0", UIMax = "5000"))
 	float VolumetricFogStartDistance;
 
 	/** 
@@ -279,6 +279,12 @@ public:
 	UFUNCTION(BlueprintCallable, Category="Rendering|VolumetricFog")
 	ENGINE_API void SetVolumetricFogDistance(float NewValue);
 
+	UFUNCTION(BlueprintCallable, Category = "Rendering|VolumetricFog")
+	ENGINE_API void SetVolumetricFogStartDistance(float NewValue);
+
+	UFUNCTION(BlueprintCallable, Category = "Rendering|VolumetricFog")
+	ENGINE_API void SetVolumetricFogNearFadeInDistance(float NewValue);
+
 	UFUNCTION(BlueprintCallable, Category="Rendering|VolumetricFog")
 	ENGINE_API void SetSecondFogData(FExponentialHeightFogData NewValue);
 
diff --git a/Engine/Source/Runtime/Engine/Private/Components/HeightFogComponent.cpp b/Engine/Source/Runtime/Engine/Private/Components/HeightFogComponent.cpp
--- a/Engine/Source/Runtime/Engine/Private/Components/HeightFogComponent.cpp
+++ b/Engine/Source/Runtime/Engine/Private/Components/HeightFogComponent.cpp
@@ -341,6 +341,24 @@ void UExponentialHeightFogComponent::SetVolumetricFogDistance(float NewValue)
 	}
 }
 
+void UExponentialHeightFogComponent::SetVolumetricFogStartDistance(float NewValue)
+{
+	if (VolumetricFogStartDistance != NewValue)
+	{
+		VolumetricFogStartDistance = NewValue;
+		MarkRenderStateDirty();
+	}
+}
+
+void UExponentialHeightFogComponent::SetVolumetricFogNearFadeInDistance(float NewValue)
+{
+	if (VolumetricFogNearFadeInDistance != NewValue)
+	{
+		VolumetricFogNearFadeInDistance = NewValue;
+		MarkRenderStateDirty();
+	}
+}
+
 void UExponentialHeightFogComponent::SetSecondFogData(FExponentialHeightFogData NewValue)
 {
 	if(SecondFogData.FogDensity != NewValue.FogDensity ||
diff --git a/Engine/Source/Runtime/MovieSceneTracks/Private/MovieSceneTracksCustomAccessors.cpp b/Engine/Source/Runtime/MovieSceneTracks/Private/MovieSceneTracksCustomAccessors.cpp
--- a/Engine/Source/Runtime/MovieSceneTracks/Private/MovieSceneTracksCustomAccessors.cpp
+++ b/Engine/Source/Runtime/MovieSceneTracks/Private/MovieSceneTracksCustomAccessors.cpp
@@ -425,6 +425,30 @@ void SetExponentialHeightFogComponentVolumetricFogExtinctionScale(UObject* Objec
 	UExponentialHeightFogComponent* ExponentialHeightFogComponent = CastChecked<UExponentialHeightFogComponent>(Object);
 	ExponentialHeightFogComponent->SetVolumetricFogExtinctionScale(InVolumetricFogExtinctionScale);
 }
+
+float GetExponentialHeightFogComponentVolumetricFogStartDistance(const UObject* Object)
+{
+	const UExponentialHeightFogComponent* ExponentialHeightFogComponent = CastChecked<const UExponentialHeightFogComponent>(Object);
+	return ExponentialHeightFogComponent->VolumetricFogStartDistance;
+}
+
+void SetExponentialHeightFogComponentVolumetricFogStartDistance(UObject* Object, float InVolumetricFogStartDistance)
+{
+	UExponentialHeightFogComponent* ExponentialHeightFogComponent = CastChecked<UExponentialHeightFogComponent>(Object);
+	ExponentialHeightFogComponent->SetVolumetricFogStartDistance(InVolumetricFogStartDistance);
+}
+
+float GetExponentialHeightFogComponentVolumetricFogNearFadeInDistance(const UObject* Object)
+{
+	const UExponentialHeightFogComponent* ExponentialHeightFogComponent = CastChecked<const UExponentialHeightFogComponent>(Object);
+	return ExponentialHeightFogComponent->VolumetricFogNearFadeInDistance;
+}
+
+void SetExponentialHeightFogComponentVolumetricFogNearFadeInDistance(UObject* Object, float InVolumetricFogNearFadeInDistance)
+{
+	UExponentialHeightFogComponent* ExponentialHeightFogComponent = CastChecked<UExponentialHeightFogComponent>(Object);
+	ExponentialHeightFogComponent->SetVolumetricFogNearFadeInDistance(InVolumetricFogNearFadeInDistance);
+}
 	
 FIntermediateColor GetExponentialHeightFogComponentDirectionalInscatteringColor(const UObject* Object, EColorPropertyType InColorType)
 {
@@ -658,6 +682,12 @@ void InitializeMovieSceneTracksAccessors(FMovieSceneTracksComponentTypes* Tracks
 	TracksComponents->Accessors.Float.Add(
 			UExponentialHeightFogComponent::StaticClass(), GET_MEMBER_NAME_CHECKED(UExponentialHeightFogComponent, StartDistance),
 			GetExponentialHeightFogComponentStartDistance, SetExponentialHeightFogComponentStartDistance);
+	TracksComponents->Accessors.Float.Add(
+			UExponentialHeightFogComponent::StaticClass(), GET_MEMBER_NAME_CHECKED(UExponentialHeightFogComponent, VolumetricFogStartDistance),
+			GetExponentialHeightFogComponentVolumetricFogStartDistance, SetExponentialHeightFogComponentVolumetricFogStartDistance);
+	TracksComponents->Accessors.Float.Add(
+			UExponentialHeightFogComponent::StaticClass(), GET_MEMBER_NAME_CHECKED(UExponentialHeightFogComponent, VolumetricFogNearFadeInDistance),
+			GetExponentialHeightFogComponentVolumetricFogNearFadeInDistance, SetExponentialHeightFogComponentVolumetricFogNearFadeInDistance);
 	TracksComponents->Accessors.Float.Add(
 			UExponentialHeightFogComponent::StaticClass(), GET_MEMBER_NAME_CHECKED(UExponentialHeightFogComponent, VolumetricFogExtinctionScale),
 			GetExponentialHeightFogComponentVolumetricFogExtinctionScale, SetExponentialHeightFogComponentVolumetricFogExtinctionScale);
