Sequencer: Fixed initial values being cached for ignored entities

Ignored entities do not get processed by the property instantiator so they do not include any property meta-data components and thus cannot store initial values.
This is ok - if the entit ever becomes not-Ignored it will also tagged NeedsLink and will be processed by the instantiator.

diff --git a/Engine/Source/Runtime/MovieSceneTracks/Private/Systems/MovieSceneInitialValueSystem.cpp b/Engine/Source/Runtime/MovieSceneTracks/Private/Systems/MovieSceneInitialValueSystem.cpp
--- a/Engine/Source/Runtime/MovieSceneTracks/Private/Systems/MovieSceneInitialValueSystem.cpp
+++ b/Engine/Source/Runtime/MovieSceneTracks/Private/Systems/MovieSceneInitialValueSystem.cpp
@@ -141,7 +141,7 @@ void UMovieSceneInitialValueSystem::OnRun(FSystemTaskPrerequisites& InPrerequisi
 			Filter.Any(Mutation.AnyInitialValue);
 			Filter.All({ BuiltInComponents->Tags.NeedsLink });
 			Filter.None({ BuiltInComponents->InitialValueIndex });
-			Filter.None({ BuiltInComponents->Tags.HasAssignedInitialValue });
+			Filter.None({ BuiltInComponents->Tags.HasAssignedInitialValue, BuiltInComponents->Tags.Ignored });
 
 			Linker->EntityManager.MutateAll(Filter, Mutation);
 		}
@@ -171,7 +171,7 @@ void UMovieSceneInitialValueSystem::OnRun(FSystemTaskPrerequisites& InPrerequisi
 		Filter.Any(Mutation.AnyInitialValue);
 		Filter.Any({ BuiltInComponents->BoundObject, BuiltInComponents->Interrogation.OutputKey });
 		Filter.All({ BuiltInComponents->Tags.NeedsLink });
-		Filter.None({ BuiltInComponents->Tags.HasAssignedInitialValue });
+		Filter.None({ BuiltInComponents->Tags.HasAssignedInitialValue, BuiltInComponents->Tags.Ignored });
 
 		Linker->EntityManager.MutateAll(Filter, Mutation);
 	}
