Sequencer: Fixed skeletal animations not honoring 'Evaluate Nearest Section'

diff --git a/Engine/Source/Runtime/MovieSceneTracks/Private/Systems/MovieSceneSkeletalAnimationSystem.cpp b/Engine/Source/Runtime/MovieSceneTracks/Private/Systems/MovieSceneSkeletalAnimationSystem.cpp
--- a/Engine/Source/Runtime/MovieSceneTracks/Private/Systems/MovieSceneSkeletalAnimationSystem.cpp
+++ b/Engine/Source/Runtime/MovieSceneTracks/Private/Systems/MovieSceneSkeletalAnimationSystem.cpp
@@ -335,6 +335,7 @@ struct FGatherSkeletalAnimations
 			TRead<FInstanceHandle> InstanceHandles,
 			TRead<UObject*> BoundObjects,
 			TRead<FMovieSceneSkeletalAnimationComponentData> SkeletalAnimations,
+			TReadOptional<FFrameTime> OptionalEvalTimes,
 			TReadOptional<double> WeightAndEasings) const
 	{
 		// Gather all the skeletal animations currently active in all sequences.
@@ -358,6 +359,8 @@ struct FGatherSkeletalAnimations
 			const FSequenceInstance& SequenceInstance = InstanceRegistry->GetInstance(InstanceHandle);
 			const FMovieSceneContext& Context = SequenceInstance.GetContext();
 
+			const FFrameTime EvalFrameTime = OptionalEvalTimes ? OptionalEvalTimes[Index] : Context.GetTime();
+
 			// Calculate the time at which to evaluate the animation
 			const UMovieSceneSkeletalAnimationSection* AnimSection = SkeletalAnimation.Section;
 			const FMovieSceneSkeletalAnimationParams& AnimParams = AnimSection->Params;
@@ -368,7 +371,7 @@ struct FGatherSkeletalAnimations
 			{
 				continue;
 			}
-			const float EvalTime = AnimParams.MapTimeToAnimation(AnimSection, Context.GetTime(), Context.GetFrameRate());
+			const float EvalTime = AnimParams.MapTimeToAnimation(AnimSection, EvalFrameTime, Context.GetFrameRate());
 			const float PreviousEvalTime = AnimParams.MapTimeToAnimation(AnimSection, Context.GetPreviousTime(), Context.GetFrameRate());
 
 			const FSequenceInstance& RootInstance = InstanceRegistry->GetInstance(RootInstanceHandle);
@@ -395,6 +398,7 @@ struct FGatherSkeletalAnimations
 			FActiveSkeletalAnimation Animation;
 			Animation.AnimSection        = AnimSection;
 			Animation.Context            = Context;
+			Animation.EvalFrameTime      = EvalFrameTime;
 			Animation.EntityID           = EntityID;
 			Animation.RootInstanceHandle = RootInstanceHandle;
 			Animation.FromEvalTime       = PreviousEvalTime;
@@ -662,7 +666,7 @@ private:
 			const UMovieSceneSkeletalAnimationSection* AnimSection = SkeletalAnimation.AnimSection;
 			const FMovieSceneSkeletalAnimationParams& AnimParams = AnimSection->Params;
 			UMovieSceneSkeletalAnimationSection::FRootMotionParams RootMotionParams;
-			AnimSection->GetRootMotion(SkeletalAnimation.Context.GetTime().RoundToFrame(), RootMotionParams);
+			AnimSection->GetRootMotion(SkeletalAnimation.EvalFrameTime.RoundToFrame(), RootMotionParams);
 			//set up root motion/bone transform delegates
 			if (AnimSection->Params.SwapRootBone != ESwapRootBone::SwapRootBone_None)
 			{
@@ -693,7 +697,7 @@ private:
 			SetAnimPositionParams.RootInstanceHandle = SkeletalAnimation.RootInstanceHandle;
 			SetAnimPositionParams.Section = AnimSection;
 			SetAnimPositionParams.SkeletalMeshComponent = SkeletalMeshComponent;
-			SetAnimPositionParams.CurrentTime = SkeletalAnimation.Context.GetTime();
+			SetAnimPositionParams.CurrentTime = SkeletalAnimation.EvalFrameTime;
 			SetAnimPositionParams.FromPosition = SkeletalAnimation.FromEvalTime;
 			SetAnimPositionParams.ToPosition = SkeletalAnimation.ToEvalTime;
 			SetAnimPositionParams.Weight = SkeletalAnimation.BlendWeight;
@@ -999,6 +1003,7 @@ void UMovieSceneSkeletalAnimationSystem::OnSchedulePersistentTasks(UE::MovieScen
 	.Read(BuiltInComponents->InstanceHandle)
 	.Read(BuiltInComponents->BoundObject)
 	.Read(TrackComponents->SkeletalAnimation)
+	.ReadOptional(BuiltInComponents->EvalTime)
 	.ReadOptional(BuiltInComponents->WeightAndEasingResult)
 	.FilterNone({ BuiltInComponents->Tags.Ignored })
 	.SetStat(GET_STATID(MovieSceneEval_GatherSkeletalAnimations))
@@ -1036,6 +1041,7 @@ void UMovieSceneSkeletalAnimationSystem::OnRun(FSystemTaskPrerequisites& InPrere
 	.Read(BuiltInComponents->InstanceHandle)
 	.Read(BuiltInComponents->BoundObject)
 	.Read(TrackComponents->SkeletalAnimation)
+	.ReadOptional(BuiltInComponents->EvalTime)
 	.ReadOptional(BuiltInComponents->WeightAndEasingResult)
 	.FilterNone({ BuiltInComponents->Tags.Ignored })
 	.SetStat(GatherStatId)
