Added missing file

diff --git a/Engine/Source/Runtime/MovieSceneTracks/Public/Systems/MovieSceneSkeletalAnimationSystem.h b/Engine/Source/Runtime/MovieSceneTracks/Public/Systems/MovieSceneSkeletalAnimationSystem.h
--- a/Engine/Source/Runtime/MovieSceneTracks/Public/Systems/MovieSceneSkeletalAnimationSystem.h
+++ b/Engine/Source/Runtime/MovieSceneTracks/Public/Systems/MovieSceneSkeletalAnimationSystem.h
@@ -21,6 +21,7 @@ struct FActiveSkeletalAnimation
 {
 	const UMovieSceneSkeletalAnimationSection* AnimSection;
 	FMovieSceneContext Context;
+	FFrameTime EvalFrameTime;
 	FMovieSceneEntityID EntityID;
 	FRootInstanceHandle RootInstanceHandle;
 	double BlendWeight;
