Sequencer: Fixed handling of bound materials becoming null before GC

diff --git a/Engine/Source/Runtime/MovieSceneTracks/Private/Systems/MovieSceneMaterialParameterSystem.cpp b/Engine/Source/Runtime/MovieSceneTracks/Private/Systems/MovieSceneMaterialParameterSystem.cpp
--- a/Engine/Source/Runtime/MovieSceneTracks/Private/Systems/MovieSceneMaterialParameterSystem.cpp
+++ b/Engine/Source/Runtime/MovieSceneTracks/Private/Systems/MovieSceneMaterialParameterSystem.cpp
@@ -366,6 +366,20 @@ struct TOverlappingMaterialParameterHandler : Mixin
 			return;
 		}
 
+		if (!BoundMaterial)
+		{
+			// Handle null bound materials - this is a rare case but must be handled explicitly because it is possible
+			//  a GC has already destroyed our OutputEntityID because it was considered garbage. In this instance we just
+			//  remove the invalid blend channel input components off all inputs.
+			DestroyOutput(BoundMaterial, ParameterInfo, Output, Aggregate);
+
+			for (FMovieSceneEntityID Input : Inputs)
+			{
+				Linker->EntityManager.RemoveComponent(Input, BuiltInComponents->BlendChannelInput);
+			}
+			return;
+		}
+
 		const bool bUseBlending = NumContributors > 1 || !Linker->EntityManager.HasComponent(Inputs[0], BuiltInComponents->Tags.AbsoluteBlend) || Linker->EntityManager.HasComponent(Inputs[0], BuiltInComponents->WeightAndEasingResult);
 		if (bUseBlending)
 		{
