Sequencer: Add missing audio fade evaluation

This is a regression from 5.2 when converting to ECS AudioSystem.

diff --git a/Engine/Source/Runtime/MovieSceneTracks/Private/Systems/MovieSceneAudioSystem.cpp b/Engine/Source/Runtime/MovieSceneTracks/Private/Systems/MovieSceneAudioSystem.cpp
--- a/Engine/Source/Runtime/MovieSceneTracks/Private/Systems/MovieSceneAudioSystem.cpp
+++ b/Engine/Source/Runtime/MovieSceneTracks/Private/Systems/MovieSceneAudioSystem.cpp
@@ -405,7 +405,7 @@ private:
 					AudioComponent->AttachToComponent(AttachComponent, FAttachmentTransformRules::KeepRelativeTransform, AttachKey.SocketName);
 				}
 
-				EvaluationData->VolumeMultiplier = VolumeMultiplier;
+				EvaluationData->VolumeMultiplier = VolumeMultiplier * AudioSection->EvaluateEasing(Context.GetTime());
 				EvaluationData->PitchMultiplier = PitchMultiplier;
 
 				EnsureAudioIsPlaying(nullptr, InstanceHandle, *AudioSection, *EvaluationData, Context, *Player);
