Sequencer: Update the channel proxy if it's null in GetKeyStruct

Fixes a crash where if you have the key details panel and then undo a modifcation to the keys, the channel proxy will be null and need to be recreated.

diff --git a/Engine/Source/Runtime/MovieSceneTracks/Private/Sections/MovieScene3DTransformSection.cpp b/Engine/Source/Runtime/MovieSceneTracks/Private/Sections/MovieScene3DTransformSection.cpp
--- a/Engine/Source/Runtime/MovieSceneTracks/Private/Sections/MovieScene3DTransformSection.cpp
+++ b/Engine/Source/Runtime/MovieSceneTracks/Private/Sections/MovieScene3DTransformSection.cpp
@@ -840,6 +840,11 @@ EMovieSceneChannelProxyType UMovieScene3DTransformSection::CacheChannelProxy()
 
 TSharedPtr<FStructOnScope> UMovieScene3DTransformSection::GetKeyStruct(TArrayView<const FKeyHandle> KeyHandles)
 {
+	if (!ChannelProxy)
+	{
+		GetChannelProxy();
+	}
+
 	TArrayView<FMovieSceneDoubleChannel* const> DoubleChannels = ChannelProxy->GetChannels<FMovieSceneDoubleChannel>();
 
 	TOptional<TTuple<FKeyHandle, FFrameNumber>> LocationKeys[3] = {
diff --git a/Engine/Source/Runtime/MovieSceneTracks/Private/Sections/MovieSceneColorSection.cpp b/Engine/Source/Runtime/MovieSceneTracks/Private/Sections/MovieSceneColorSection.cpp
--- a/Engine/Source/Runtime/MovieSceneTracks/Private/Sections/MovieSceneColorSection.cpp
+++ b/Engine/Source/Runtime/MovieSceneTracks/Private/Sections/MovieSceneColorSection.cpp
@@ -143,6 +143,11 @@ UMovieSceneColorSection::UMovieSceneColorSection(const FObjectInitializer& Objec
 
 TSharedPtr<FStructOnScope> UMovieSceneColorSection::GetKeyStruct(TArrayView<const FKeyHandle> KeyHandles)
 {
+	if (!ChannelProxy)
+	{
+		GetChannelProxy();
+	}
+	
 	TSharedRef<FStructOnScope> KeyStruct = MakeShareable(new FStructOnScope(FMovieSceneColorKeyStruct::StaticStruct()));
 	auto Struct = (FMovieSceneColorKeyStruct*)KeyStruct->GetStructMemory();
 
diff --git a/Engine/Source/Runtime/MovieSceneTracks/Private/Sections/MovieSceneVectorSection.cpp b/Engine/Source/Runtime/MovieSceneTracks/Private/Sections/MovieSceneVectorSection.cpp
--- a/Engine/Source/Runtime/MovieSceneTracks/Private/Sections/MovieSceneVectorSection.cpp
+++ b/Engine/Source/Runtime/MovieSceneTracks/Private/Sections/MovieSceneVectorSection.cpp
@@ -209,6 +209,11 @@ void UMovieSceneFloatVectorSection::RecreateChannelProxy()
 
 TSharedPtr<FStructOnScope> UMovieSceneFloatVectorSection::GetKeyStruct(TArrayView<const FKeyHandle> KeyHandles)
 {
+	if (!ChannelProxy)
+	{
+		GetChannelProxy();
+	}
+
 	TSharedPtr<FStructOnScope> KeyStruct;
 	if (ChannelsUsed == 2)
 	{
@@ -324,6 +329,11 @@ void UMovieSceneDoubleVectorSection::RecreateChannelProxy()
 
 TSharedPtr<FStructOnScope> UMovieSceneDoubleVectorSection::GetKeyStruct(TArrayView<const FKeyHandle> KeyHandles)
 {
+	if (!ChannelProxy)
+	{
+		GetChannelProxy();
+	}
+
 	TSharedPtr<FStructOnScope> KeyStruct;
 	if (ChannelsUsed == 2)
 	{
