Sequencer: Check for skel anim components being already post-evaluated before forcibly ticking for the current frame
This is a temporary fix while there are some playback APIs that cause inline evaluations of Sequencer. Ultimately we can remove this entire force-tick mechanism once Sequencer's playback API is entirely async.

Same as 33166070 but this occurs during restore state

diff --git a/Engine/Source/Runtime/MovieSceneTracks/Private/Systems/MovieSceneSkeletalAnimationSystem.cpp b/Engine/Source/Runtime/MovieSceneTracks/Private/Systems/MovieSceneSkeletalAnimationSystem.cpp
--- a/Engine/Source/Runtime/MovieSceneTracks/Private/Systems/MovieSceneSkeletalAnimationSystem.cpp
+++ b/Engine/Source/Runtime/MovieSceneTracks/Private/Systems/MovieSceneSkeletalAnimationSystem.cpp
@@ -135,12 +135,15 @@ struct FPreAnimatedSkeletalAnimationTraits : FBoundObjectPreAnimatedStateTraits
 		Component->SetUpdateClothInEditor(true);
 		Component->TickAnimation(0.f, false);
 
-		Component->RefreshBoneTransforms();
-		Component->RefreshFollowerComponents();
-		Component->UpdateComponentToWorld();
-		Component->FinalizeBoneTransform();
-		Component->MarkRenderTransformDirty();
-		Component->MarkRenderDynamicDataDirty();
+		if (!Component->IsPostEvaluatingAnimation())
+		{
+			Component->RefreshBoneTransforms();
+			Component->RefreshFollowerComponents();
+			Component->UpdateComponentToWorld();
+			Component->FinalizeBoneTransform();
+			Component->MarkRenderTransformDirty();
+			Component->MarkRenderDynamicDataDirty();
+		}
 
 		// Reset the mesh component update flag and animation mode to what they were before we animated the object
 		InOutCachedValue.SkeletalMeshRestoreState.RestoreState(Component);
