Enable this module for AutoRTFM

diff --git a/Engine/Source/Runtime/MovieSceneTracks/MovieSceneTracks.Build.cs b/Engine/Source/Runtime/MovieSceneTracks/MovieSceneTracks.Build.cs
--- a/Engine/Source/Runtime/MovieSceneTracks/MovieSceneTracks.Build.cs
+++ b/Engine/Source/Runtime/MovieSceneTracks/MovieSceneTracks.Build.cs
@@ -40,5 +40,7 @@ public class MovieSceneTracks : ModuleRules
 					"UnrealEd"
 				});
 		}
+
+		bAllowAutoRTFMInstrumentation = true;
 	}
 }
