Sequencer: remove IMovieScenePlayer usage from miscellaneous code

diff --git a/Engine/Source/Runtime/MovieScene/Private/EntitySystem/MovieSceneBoundObjectInstantiator.cpp b/Engine/Source/Runtime/MovieScene/Private/EntitySystem/MovieSceneBoundObjectInstantiator.cpp
--- a/Engine/Source/Runtime/MovieScene/Private/EntitySystem/MovieSceneBoundObjectInstantiator.cpp
+++ b/Engine/Source/Runtime/MovieScene/Private/EntitySystem/MovieSceneBoundObjectInstantiator.cpp
@@ -44,7 +44,9 @@ void UMovieSceneGenericBoundObjectInstantiator::OnRun(FSystemTaskPrerequisites&
 			EResolveError Error = EResolveError::UnresolvedBinding;
 
 			FSequenceInstance& SequenceInstance = InstanceRegistry->MutateInstance(InstanceHandle);
-			for (TWeakObjectPtr<> WeakObject : SequenceInstance.GetPlayer()->FindBoundObjects(ObjectBinding, SequenceInstance.GetSequenceID()))
+			TSharedRef<const FSharedPlaybackState> SharedPlaybackState = SequenceInstance.GetSharedPlaybackState();
+			TArrayView<TWeakObjectPtr<>> BoundObjects = SharedPlaybackState->FindBoundObjects(ObjectBinding, SequenceInstance.GetSequenceID());
+			for (TWeakObjectPtr<> WeakObject : BoundObjects)
 			{
 				if (UObject* Object = WeakObject.Get())
 				{
diff --git a/Engine/Source/Runtime/MovieScene/Private/EntitySystem/MovieSceneBoundSceneComponentInstantiator.cpp b/Engine/Source/Runtime/MovieScene/Private/EntitySystem/MovieSceneBoundSceneComponentInstantiator.cpp
--- a/Engine/Source/Runtime/MovieScene/Private/EntitySystem/MovieSceneBoundSceneComponentInstantiator.cpp
+++ b/Engine/Source/Runtime/MovieScene/Private/EntitySystem/MovieSceneBoundSceneComponentInstantiator.cpp
@@ -9,6 +9,7 @@
 #include "EntitySystem/MovieSceneInstanceRegistry.h"
 #include "EntitySystem/MovieSceneEntitySystemLinker.h"
 
+#include "Evaluation/MovieSceneEvaluationState.h"
 #include "MovieSceneCommonHelpers.h"
 #include "IMovieScenePlayer.h"
 
@@ -46,8 +47,8 @@ void UMovieSceneBoundSceneComponentInstantiator::OnRun(FSystemTaskPrerequisites&
 			EResolveError Result = EResolveError::UnresolvedBinding;
 
 			FSequenceInstance& SequenceInstance = InstanceRegistry->MutateInstance(InstanceHandle);
-
-			TArrayView<TWeakObjectPtr<>> BoundObjects = SequenceInstance.GetPlayer()->FindBoundObjects(ObjectBinding, SequenceInstance.GetSequenceID());
+			TSharedRef<const FSharedPlaybackState> SharedPlaybackState = SequenceInstance.GetSharedPlaybackState();
+			TArrayView<TWeakObjectPtr<>> BoundObjects = SharedPlaybackState->FindBoundObjects(ObjectBinding, SequenceInstance.GetSequenceID());
 			if (BoundObjects.Num() == 0)
 			{
 				UE_LOG(LogMovieSceneECS, Verbose, TEXT("FBoundSceneComponentBatch::ResolveObjects: No bound objects returned for FGuid: %s"), *ObjectBinding.ToString());
diff --git a/Engine/Source/Runtime/MovieScene/Private/MovieSceneSpawnRegister.cpp b/Engine/Source/Runtime/MovieScene/Private/MovieSceneSpawnRegister.cpp
--- a/Engine/Source/Runtime/MovieScene/Private/MovieSceneSpawnRegister.cpp
+++ b/Engine/Source/Runtime/MovieScene/Private/MovieSceneSpawnRegister.cpp
@@ -50,7 +50,6 @@ UObject* FMovieSceneSpawnRegister::SpawnObject(const FGuid& BindingId, UMovieSce
 		}
 
 		SpawnOwnership = Spawnable->GetSpawnOwnership();
-		IMovieScenePlayer* Player = UE::MovieScene::FPlayerIndexPlaybackCapability::GetPlayer(SharedPlaybackState);
 
 		// See if there is some dynamic binding logic to invoke, otherwise spawn the actor
 			// See if there is some dynamic binding logic to invoke, otherwise spawn the actor
diff --git a/Engine/Source/Runtime/MovieScene/Public/IMovieScenePlaybackClient.h b/Engine/Source/Runtime/MovieScene/Public/IMovieScenePlaybackClient.h
--- a/Engine/Source/Runtime/MovieScene/Public/IMovieScenePlaybackClient.h
+++ b/Engine/Source/Runtime/MovieScene/Public/IMovieScenePlaybackClient.h
@@ -26,7 +26,7 @@ class IMovieScenePlaybackClient
 public:
 	GENERATED_BODY()
 
-	static UE::MovieScene::TPlaybackCapabilityID<IMovieScenePlaybackClient> ID;
+	MOVIESCENE_API static UE::MovieScene::TPlaybackCapabilityID<IMovieScenePlaybackClient> ID;
 
 	/**
 	 * Locate bound objects that relate to the specified binding ID
diff --git a/Engine/Source/Runtime/MovieSceneTracks/Private/Systems/MovieSceneAudioSystem.cpp b/Engine/Source/Runtime/MovieSceneTracks/Private/Systems/MovieSceneAudioSystem.cpp
--- a/Engine/Source/Runtime/MovieSceneTracks/Private/Systems/MovieSceneAudioSystem.cpp
+++ b/Engine/Source/Runtime/MovieSceneTracks/Private/Systems/MovieSceneAudioSystem.cpp
@@ -310,7 +310,7 @@ private:
 			bool bWantsRestoreState) const
 	{
 		const FMovieSceneContext& Context = Instance.GetContext();
-		IMovieScenePlayer* Player = Instance.GetPlayer();
+		UObject* PlaybackContext = Instance.GetSharedPlaybackState()->GetPlaybackContext();
 
 		UMovieSceneAudioSection* AudioSection = AudioData.Section;
 		if (!ensureMsgf(AudioSection, TEXT("No valid audio section found in audio track component data!")))
@@ -340,8 +340,6 @@ private:
 		// Root audio track
 		if (BoundObject == nullptr)
 		{
-			UObject* PlaybackContext = Player->GetPlaybackContext();
-
 			const FMovieSceneActorReferenceData& AttachActorData = AudioSection->GetAttachActorData();
 
 			USceneComponent* AttachComponent = nullptr;
@@ -351,7 +349,7 @@ private:
 			if (AttachBindingID.IsValid())
 			{
 				// If the transform is set, otherwise use the bound actor's transform
-				for (TWeakObjectPtr<> WeakObject : AttachBindingID.ResolveBoundObjects(Instance.GetSequenceID(), *Player))
+				for (TWeakObjectPtr<> WeakObject : AttachBindingID.ResolveBoundObjects(Instance.GetSequenceID(), Instance.GetSharedPlaybackState()))
 				{
 					AActor* AttachActor = Cast<AActor>(WeakObject.Get());
 					if (AttachActor)
@@ -408,7 +406,7 @@ private:
 				EvaluationData->VolumeMultiplier = VolumeMultiplier * AudioSection->EvaluateEasing(Context.GetTime());
 				EvaluationData->PitchMultiplier = PitchMultiplier;
 
-				EnsureAudioIsPlaying(nullptr, InstanceHandle, *AudioSection, *EvaluationData, Context, *Player);
+				EnsureAudioIsPlaying(nullptr, InstanceHandle, *AudioSection, *EvaluationData, Context, PlaybackContext);
 			}
 		}
 
@@ -449,7 +447,7 @@ private:
 				EvaluationData->VolumeMultiplier = VolumeMultiplier;
 				EvaluationData->PitchMultiplier = PitchMultiplier;
 
-				EnsureAudioIsPlaying(BoundObject, InstanceHandle, *AudioSection, *EvaluationData, Context, *Player);
+				EnsureAudioIsPlaying(BoundObject, InstanceHandle, *AudioSection, *EvaluationData, Context, PlaybackContext);
 			}
 		}
 	}
@@ -460,7 +458,7 @@ private:
 			UMovieSceneAudioSection& AudioSection,
 			FAudioComponentEvaluationData& EvaluationData,
 			const FMovieSceneContext& Context, 
-			IMovieScenePlayer& Player) const
+			UObject* PlaybackContext) const
 	{
 		using FInstanceObjectKey = UMovieSceneAudioSystem::FInstanceObjectKey;
 		using FAudioInputsBySectionKey = UMovieSceneAudioSystem::FAudioInputsBySectionKey;
@@ -530,7 +528,6 @@ private:
 		const bool bSoundNeedsStateChange =  AudioComponent.Sound != Sound;
 		bool bSoundNeedsTimeSync = false;
 
-		UObject* PlaybackContext = Player.GetPlaybackContext();
 		UWorld* World = PlaybackContext ? PlaybackContext->GetWorld() : nullptr;
 
 		// Sync only if there is no time dilation because otherwise the system will constantly resync because audio 
diff --git a/Engine/Source/Runtime/MovieSceneTracks/Private/Systems/MovieSceneCameraShakeSystem.cpp b/Engine/Source/Runtime/MovieSceneTracks/Private/Systems/MovieSceneCameraShakeSystem.cpp
--- a/Engine/Source/Runtime/MovieSceneTracks/Private/Systems/MovieSceneCameraShakeSystem.cpp
+++ b/Engine/Source/Runtime/MovieSceneTracks/Private/Systems/MovieSceneCameraShakeSystem.cpp
@@ -215,7 +215,7 @@ FCameraShakePreviewer& FCameraShakePreviewerLinkerExtension::GetPreviewer(FInsta
 	UMovieSceneEntitySystemLinker* Linker = WeakLinker.Get();
 	FInstanceRegistry* InstanceRegistry = Linker->GetInstanceRegistry();
 	const FSequenceInstance& SequenceInstance = InstanceRegistry->GetInstance(InstanceHandle);
-	UObject* PlaybackContext = SequenceInstance.GetPlayer()->GetPlaybackContext();
+	UObject* PlaybackContext = SequenceInstance.GetSharedPlaybackState()->GetPlaybackContext();
 	UWorld* ContextWorld = PlaybackContext ? PlaybackContext->GetWorld() : nullptr;
 
 	FCameraShakePreviewer& NewPreviewer = Previewers.Emplace(InstanceHandle, ContextWorld);
@@ -276,7 +276,7 @@ void FCameraShakePreviewerLinkerExtension::OnLevelViewportClientListChanged()
 	for (TPair<FInstanceHandle, FCameraShakePreviewer> Pair : Previewers)
 	{
 		const FSequenceInstance& SequenceInstance = InstanceRegistry->GetInstance(Pair.Key);
-		UObject* PlaybackContext = SequenceInstance.GetPlayer()->GetPlaybackContext();
+		UObject* PlaybackContext = SequenceInstance.GetSharedPlaybackState()->GetPlaybackContext();
 		UWorld* ContextWorld = PlaybackContext ? PlaybackContext->GetWorld() : nullptr;
 
 		Pair.Value.RegisterViewModifiers(
@@ -375,7 +375,7 @@ void UMovieSceneCameraShakeInstantiatorSystem::OnRun(FSystemTaskPrerequisites& I
 	{
 		const FSequenceInstance& Instance = InstanceRegistry->GetInstance(InstanceHandle);
 		const FMovieSceneContext& Context = Instance.GetContext();
-		IMovieScenePlayer* Player = Instance.GetPlayer();
+		UObject* PlaybackContext = Instance.GetSharedPlaybackState()->GetPlaybackContext();
 
 		TSubclassOf<UCameraShakeBase> ShakeClass = ShakeData.SectionData.ShakeClass;
 		UCameraShakeSourceComponent* ShakeSourceComponent = Cast<UCameraShakeSourceComponent>(BoundObject);
@@ -471,7 +471,7 @@ void UMovieSceneCameraShakeInstantiatorSystem::OnRun(FSystemTaskPrerequisites& I
 			}
 			else
 			{
-				UObject* OuterObject = Player->GetPlaybackContext() ? Player->GetPlaybackContext() : GetTransientPackage();
+				UObject* OuterObject = PlaybackContext ? PlaybackContext : GetTransientPackage();
 				ShakeInstance = NewObject<UCameraShakeBase>(OuterObject, ShakeClass);
 			}
 
@@ -533,7 +533,7 @@ void UMovieSceneCameraShakeInstantiatorSystem::TriggerOneShotShakes()
 	for (TPair<FInstanceHandle, TArray<FTimedTrigger>>& Pair : TriggersByInstance)
 	{
 		const FSequenceInstance& Instance = InstanceRegistry->GetInstance(Pair.Key);
-		IMovieScenePlayer* Player = Instance.GetPlayer();
+		TSharedRef<const FSharedPlaybackState> SharedPlaybackState = Instance.GetSharedPlaybackState();
 
 		const FMovieSceneContext& Context = Instance.GetContext();
 		if (Context.GetDirection() != EPlayDirection::Forwards)
@@ -546,7 +546,8 @@ void UMovieSceneCameraShakeInstantiatorSystem::TriggerOneShotShakes()
 			TArray<UCameraShakeSourceComponent*> ShakeSourceComponents;
 			if (Trigger.ObjectBindingID.IsValid())
 			{
-				for (TWeakObjectPtr<> WeakBoundObject : Player->FindBoundObjects(Trigger.ObjectBindingID, Instance.GetSequenceID()))
+				TArrayView<TWeakObjectPtr<>> BoundObjects = SharedPlaybackState->FindBoundObjects(Trigger.ObjectBindingID, Instance.GetSequenceID());
+				for (TWeakObjectPtr<> WeakBoundObject : BoundObjects)
 				{
 					if (UObject* BoundObject = WeakBoundObject.Get())
 					{
diff --git a/Engine/Source/Runtime/MovieSceneTracks/Private/Systems/MovieSceneDataLayerSystem.cpp b/Engine/Source/Runtime/MovieSceneTracks/Private/Systems/MovieSceneDataLayerSystem.cpp
--- a/Engine/Source/Runtime/MovieSceneTracks/Private/Systems/MovieSceneDataLayerSystem.cpp
+++ b/Engine/Source/Runtime/MovieSceneTracks/Private/Systems/MovieSceneDataLayerSystem.cpp
@@ -577,7 +577,8 @@ void UMovieSceneDataLayerSystem::OnRun(FSystemTaskPrerequisites& InPrerequisites
 					FString SequenceList;
 					for (const FSequenceInstance& Instance : Linker->GetInstanceRegistry()->GetSparseInstances())
 					{
-						UMovieSceneSequence* Sequence = Instance.GetPlayer()->GetEvaluationTemplate().GetSequence(Instance.GetSequenceID());
+						TSharedRef<const FSharedPlaybackState> SharedPlaybackState = Instance.GetSharedPlaybackState();
+						UMovieSceneSequence* Sequence = SharedPlaybackState->GetSequence(Instance.GetSequenceID());
 
 						if (SequenceList.Len())
 						{
@@ -599,16 +600,12 @@ UDataLayerManager* UMovieSceneDataLayerSystem::GetDataLayerManager(UE::MovieScen
 	using namespace UE::MovieScene;
 
 	const FSequenceInstance& Instance = Linker->GetInstanceRegistry()->GetInstance(RootInstance);
-	IMovieScenePlayer* Player = Instance.GetPlayer();
-	if (!Player)
+	UObject* PlaybackContext = Instance.GetSharedPlaybackState()->GetPlaybackContext();
+	if (PlaybackContext)
 	{
-		return nullptr;
+		return UDataLayerManager::GetDataLayerManager(PlaybackContext);
 	}
-
-	UObject* PlayerUObject = Player->AsUObject();
-	UObject* PlaybackContext = PlayerUObject ? PlayerUObject : Player->GetPlaybackContext();
-
-	return UDataLayerManager::GetDataLayerManager(PlaybackContext);
+	return nullptr;
 }
 
 void UMovieSceneDataLayerSystem::UpdateDesiredStates()
diff --git a/Engine/Source/Runtime/MovieSceneTracks/Private/Systems/MovieSceneLevelVisibilitySystem.cpp b/Engine/Source/Runtime/MovieSceneTracks/Private/Systems/MovieSceneLevelVisibilitySystem.cpp
--- a/Engine/Source/Runtime/MovieSceneTracks/Private/Systems/MovieSceneLevelVisibilitySystem.cpp
+++ b/Engine/Source/Runtime/MovieSceneTracks/Private/Systems/MovieSceneLevelVisibilitySystem.cpp
@@ -1,15 +1,16 @@
 // Copyright Epic Games, Inc. All Rights Reserved.
 
 #include "Systems/MovieSceneLevelVisibilitySystem.h"
+
 #include "Engine/Level.h"
 #include "EntitySystem/BuiltInComponentTypes.h"
+#include "EntitySystem/MovieSceneEntitySystemLinker.h"
 #include "EntitySystem/MovieSceneEntitySystemTask.h"
 #include "EntitySystem/MovieSceneRootInstantiatorSystem.h"
+#include "EntitySystem/MovieSceneSharedPlaybackState.h"
 #include "EntitySystem/MovieSceneSpawnablesSystem.h"
-#include "EntitySystem/MovieSceneEntitySystemLinker.h"
-#include "MovieSceneTracksComponentTypes.h"
 #include "MovieSceneExecutionToken.h"
-#include "IMovieScenePlayer.h"
+#include "MovieSceneTracksComponentTypes.h"
 
 #include "Misc/PackageName.h"
 #include "Engine/World.h"
@@ -188,9 +189,10 @@ void FMovieSceneLevelStreamingSharedData::Flush(UMovieSceneEntitySystemLinker* L
 		return;
 	}
 
-	TArray<IMovieScenePlayer*> PlayerPtrsScratch;
+	TArray<TSharedRef<FSharedPlaybackState>> StatesScratch;
 
 	FInstanceRegistry* InstanceRegistry = Linker->GetInstanceRegistry();
+	FPreAnimatedStateExtension& PreAnimatedState = Linker->PreAnimatedState;
 
 	FLevelStreamingPreAnimatedTokenProducer TokenProducer;
 
@@ -236,12 +238,12 @@ void FMovieSceneLevelStreamingSharedData::Flush(UMovieSceneEntitySystemLinker* L
 					Pair.Value.bPreviousState = GetLevelVisibility(*Level);
 				}
 
-				PlayerPtrsScratch.Reset();
-				Pair.Value.GetPlayers(InstanceRegistry, PlayerPtrsScratch);
-				for (IMovieScenePlayer* Player : PlayerPtrsScratch)
+				StatesScratch.Reset();
+				Pair.Value.GetSharedPlaybackStates(InstanceRegistry, StatesScratch);
+				for (TSharedRef<FSharedPlaybackState> State : StatesScratch)
 				{
 					// Globally save preanimated state
-					Player->SavePreAnimatedState(*Level, TMovieSceneAnimTypeID<FMovieSceneLevelStreamingSharedData>(), TokenProducer);
+					State->GetPreAnimatedState().SavePreAnimatedState(*Level, TMovieSceneAnimTypeID<FMovieSceneLevelStreamingSharedData>(), TokenProducer);
 				}
 
 				SetLevelVisibility(*Level, bShouldBeVisible, &FlushStreamingType);
@@ -380,17 +382,14 @@ TOptional<ELevelVisibility> FMovieSceneLevelStreamingSharedData::FVisibilityData
 	}
 }
 
-void FMovieSceneLevelStreamingSharedData::FVisibilityData::GetPlayers(FInstanceRegistry* InstanceRegistry, TArray<IMovieScenePlayer*>& OutPlayers) const
+void FMovieSceneLevelStreamingSharedData::FVisibilityData::GetSharedPlaybackStates(FInstanceRegistry* InstanceRegistry, TArray<TSharedRef<FSharedPlaybackState>>& OutSharedPlaybackStates) const
 {
 	for (const FVisibilityRequest& Request : Requests)
 	{
 		if (InstanceRegistry->IsHandleValid(Request.Instance))
 		{
-			IMovieScenePlayer* Player = InstanceRegistry->GetInstance(Request.Instance).GetPlayer();
-			if (Player)
-			{
-				OutPlayers.AddUnique(Player);
-			}
+			TSharedRef<FSharedPlaybackState> State = InstanceRegistry->GetInstance(Request.Instance).GetSharedPlaybackState();
+			OutSharedPlaybackStates.AddUnique(State);
 		}
 	}
 }
diff --git a/Engine/Source/Runtime/MovieSceneTracks/Private/Systems/MovieSceneLevelVisibilitySystem.h b/Engine/Source/Runtime/MovieSceneTracks/Private/Systems/MovieSceneLevelVisibilitySystem.h
--- a/Engine/Source/Runtime/MovieSceneTracks/Private/Systems/MovieSceneLevelVisibilitySystem.h
+++ b/Engine/Source/Runtime/MovieSceneTracks/Private/Systems/MovieSceneLevelVisibilitySystem.h
@@ -18,6 +18,8 @@ namespace UE
 {
 namespace MovieScene
 {
+	struct FSharedPlaybackState;
+
 	struct FMovieSceneLevelStreamingSharedData
 	{
 		bool HasAnythingToDo() const;
@@ -43,8 +45,8 @@ namespace MovieScene
 			/** Returns whether or not this level name should be visible or not */
 			TOptional<ELevelVisibility> CalculateVisibility() const;
 
-			/** Retrieve all players that are animating this level's visibility */
-			void GetPlayers(FInstanceRegistry* InstanceRegistry, TArray<IMovieScenePlayer*>& OutPlayers) const;
+			/** Retrieve all shared playback states that are animating this level's visibility */
+			void GetSharedPlaybackStates(FInstanceRegistry* InstanceRegistry, TArray<TSharedRef<FSharedPlaybackState>>& OutSharedPlaybackStates) const;
 
 		private:
 			struct FVisibilityRequest
diff --git a/Engine/Source/Runtime/MovieSceneTracks/Private/Systems/MovieSceneMaterialParameterCollectionSystem.cpp b/Engine/Source/Runtime/MovieSceneTracks/Private/Systems/MovieSceneMaterialParameterCollectionSystem.cpp
--- a/Engine/Source/Runtime/MovieSceneTracks/Private/Systems/MovieSceneMaterialParameterCollectionSystem.cpp
+++ b/Engine/Source/Runtime/MovieSceneTracks/Private/Systems/MovieSceneMaterialParameterCollectionSystem.cpp
@@ -88,8 +88,8 @@ void UMovieSceneMaterialParameterCollectionSystem::OnRun(FSystemTaskPrerequisite
 				OutBoundMaterials[Index] = FObjectComponent::Null();
 
 				UMaterialParameterCollection* Collection = MPCs[Index].Get();
-				IMovieScenePlayer* Player = InstanceRegistry->GetInstance(InstanceHandles[Index]).GetPlayer();
-				UObject* WorldContextObject = Player->GetPlaybackContext();
+				TSharedRef<const FSharedPlaybackState> SharedPlaybackState = InstanceRegistry->GetInstance(InstanceHandles[Index]).GetSharedPlaybackState();
+				UObject* WorldContextObject = SharedPlaybackState->GetPlaybackContext();
 				UWorld* World = GEngine->GetWorldFromContextObject(WorldContextObject, EGetWorldErrorMode::LogAndReturnNull);
 				if (World && Collection)
 				{
@@ -108,7 +108,7 @@ void UMovieSceneMaterialParameterCollectionSystem::OnRun(FSystemTaskPrerequisite
 							{
 								if (!Instance->bLoggedMissingParameterWarning)
 								{
-									MissingParameters.FindOrAdd(MakeTuple(Instance, Player)).Add(Name.ToString());
+									MissingParameters.FindOrAdd(MakeTuple(Instance, SharedPlaybackState)).Add(Name.ToString());
 								}
 							}
 						}
@@ -119,7 +119,7 @@ void UMovieSceneMaterialParameterCollectionSystem::OnRun(FSystemTaskPrerequisite
 							{
 								if (!Instance->bLoggedMissingParameterWarning)
 								{
-									MissingParameters.FindOrAdd(MakeTuple(Instance, Player)).Add(Name.ToString());
+									MissingParameters.FindOrAdd(MakeTuple(Instance, SharedPlaybackState)).Add(Name.ToString());
 								}
 							}
 						}
@@ -133,13 +133,13 @@ void UMovieSceneMaterialParameterCollectionSystem::OnRun(FSystemTaskPrerequisite
 		{
 #if !UE_BUILD_SHIPPING && !UE_BUILD_TEST
 
-			for (TPair<TTuple<UMaterialParameterCollectionInstance*, IMovieScenePlayer*>, TArray<FString>>& Pair : MissingParameters)
+			for (TPair<FMissingParameterKey, TArray<FString>>& Pair : MissingParameters)
 			{
 				FFormatNamedArguments Arguments;
 				Arguments.Add(TEXT("ParamNames"), FText::FromString(FString::Join(Pair.Value, TEXT(", "))));
 				FMessageLog("PIE").Warning()
 					->AddToken(FTextToken::Create(NSLOCTEXT("MaterialParameterCollectionTrack", "InvalidParameterText", "Invalid parameter name or type applied in sequence")))
-					->AddToken(FUObjectToken::Create(Pair.Key.Get<1>()->GetEvaluationTemplate().GetSequence(MovieSceneSequenceID::Root)))
+					->AddToken(FUObjectToken::Create(Pair.Key.Get<1>()->GetSequence(MovieSceneSequenceID::Root)))
 					->AddToken(FTextToken::Create(NSLOCTEXT("MaterialParameterCollectionTrack", "OnText", "on")))
 					->AddToken(FUObjectToken::Create(Pair.Key.Get<0>()))
 					->AddToken(FTextToken::Create(FText::Format(NSLOCTEXT("MaterialParameterCollectionTrack", "InvalidParameterFormatText", "with the following invalid parameters: {ParamNames}."), Arguments)));
@@ -154,7 +154,9 @@ void UMovieSceneMaterialParameterCollectionSystem::OnRun(FSystemTaskPrerequisite
 		FBuiltInComponentTypes* BuiltInComponents;
 		FMovieSceneTracksComponentTypes* TracksComponents;
 
-		mutable TMap<TTuple<UMaterialParameterCollectionInstance*, IMovieScenePlayer*>, TArray<FString>> MissingParameters;
+		using FMissingParameterKey = TTuple<UMaterialParameterCollectionInstance*, TSharedRef<const FSharedPlaybackState>>;
+		using FMissingParametersMap = TMap<FMissingParameterKey, TArray<FString>>;
+		mutable FMissingParametersMap MissingParameters;
 	};
 
 
diff --git a/Engine/Source/Runtime/MovieSceneTracks/Private/Systems/MovieSceneSlomoSystem.cpp b/Engine/Source/Runtime/MovieSceneTracks/Private/Systems/MovieSceneSlomoSystem.cpp
--- a/Engine/Source/Runtime/MovieSceneTracks/Private/Systems/MovieSceneSlomoSystem.cpp
+++ b/Engine/Source/Runtime/MovieSceneTracks/Private/Systems/MovieSceneSlomoSystem.cpp
@@ -18,9 +18,8 @@ namespace UE::MovieScene
 
 struct FSlomoUtil
 {
-	static void ApplySlomo(IMovieScenePlayer& Player, double TimeDilation)
+	static void ApplySlomo(UObject* PlaybackContext, double TimeDilation)
 	{
-		UObject* PlaybackContext = Player.GetPlaybackContext();
 		UWorld* World = PlaybackContext ? PlaybackContext->GetWorld() : nullptr;
 
 		if (!World || (!GIsEditor && World->GetNetMode() == NM_Client) || TimeDilation <= 0.f)
@@ -42,9 +41,9 @@ struct FPreAnimatedSlomoState
 {
 	TOptional<double> TimeDilation;
 
-	static FPreAnimatedSlomoState SaveState(IMovieScenePlayer* Player)
+	static FPreAnimatedSlomoState SaveState(UObject* PlaybackContext)
 	{
-		if (AWorldSettings* WorldSettings = Player->GetPlaybackContext()->GetWorld()->GetWorldSettings())
+		if (AWorldSettings* WorldSettings = PlaybackContext->GetWorld()->GetWorldSettings())
 		{
 			return FPreAnimatedSlomoState{ WorldSettings->CinematicTimeDilation };
 		}
@@ -53,15 +52,15 @@ struct FPreAnimatedSlomoState
 
 	void RestoreState(const FMovieSceneAnimTypeID& Unused, const FRestoreStateParams& Params)
 	{
-		IMovieScenePlayer* Player = Params.GetTerminalPlayer();
-		if (!ensure(Player))
+		TSharedPtr<const FSharedPlaybackState> SharedPlaybackState = Params.GetTerminalPlaybackState();
+		if (!ensure(SharedPlaybackState))
 		{
 			return;
 		}
 
 		if (TimeDilation.IsSet())
 		{
-			FSlomoUtil::ApplySlomo(*Player, TimeDilation.GetValue());
+			FSlomoUtil::ApplySlomo(SharedPlaybackState->GetPlaybackContext(), TimeDilation.GetValue());
 		}
 	}
 };
@@ -90,10 +89,10 @@ struct FEvaluateSlomo
 		{
 			FRootInstanceHandle RootInstanceHandle = RootInstanceHandles[Index];
 			const FSequenceInstance& Instance = InstanceRegistry->GetInstance(RootInstanceHandle);
-			IMovieScenePlayer* Player = Instance.GetPlayer();
+			UObject* PlaybackContext = Instance.GetSharedPlaybackState()->GetPlaybackContext();
 
 			const double TimeDilation(TimeDilations[Index]);
-			FSlomoUtil::ApplySlomo(*Player, TimeDilation);
+			FSlomoUtil::ApplySlomo(PlaybackContext, TimeDilation);
 		}
 	}
 };
@@ -173,10 +172,10 @@ void UMovieSceneSlomoSystem::SavePreAnimatedState(const FPreAnimationParameters&
 		{
 			FRootInstanceHandle RootInstanceHandle = RootInstanceHandles[Index];
 			const FSequenceInstance& Instance = InstanceRegistry->GetInstance(RootInstanceHandle);
-			IMovieScenePlayer* Player = Instance.GetPlayer();
+			UObject* PlaybackContext = Instance.GetSharedPlaybackState()->GetPlaybackContext();
 
 			PreAnimatedStorage->BeginTrackingEntity(EntityIDs[Index], bWantsRestoreState, RootInstanceHandle, Key);
-			PreAnimatedStorage->CachePreAnimatedValue(Key, [Player](const FMovieSceneAnimTypeID&) { return FPreAnimatedSlomoState::SaveState(Player); });
+			PreAnimatedStorage->CachePreAnimatedValue(Key, [PlaybackContext](const FMovieSceneAnimTypeID&) { return FPreAnimatedSlomoState::SaveState(PlaybackContext); });
 		}
 	};
 
diff --git a/Engine/Source/Runtime/MovieSceneTracks/Private/Systems/MovieSceneTransformOriginSystem.cpp b/Engine/Source/Runtime/MovieSceneTracks/Private/Systems/MovieSceneTransformOriginSystem.cpp
--- a/Engine/Source/Runtime/MovieSceneTracks/Private/Systems/MovieSceneTransformOriginSystem.cpp
+++ b/Engine/Source/Runtime/MovieSceneTracks/Private/Systems/MovieSceneTransformOriginSystem.cpp
@@ -41,8 +41,9 @@ struct FGatherTransformOrigin
 			}
 
 			const FSequenceInstance& Instance = SparseInstances[Index];
+			TSharedRef<const FSharedPlaybackState> SharedPlaybackState = Instance.GetSharedPlaybackState();
 
-			const IMovieScenePlaybackClient*  Client       = Instance.GetPlayer()->GetPlaybackClient();
+			const IMovieScenePlaybackClient*  Client       = SharedPlaybackState->FindCapability<IMovieScenePlaybackClient>();
 			const UObject*                    InstanceData = Client ? Client->GetInstanceData() : nullptr;
 			const IMovieSceneTransformOrigin* RawInterface = Cast<const IMovieSceneTransformOrigin>(InstanceData);
 
@@ -183,7 +184,9 @@ bool UMovieSceneTransformOriginSystem::IsRelevantImpl(UMovieSceneEntitySystemLin
 
 	for (const FSequenceInstance& Instance : InLinker->GetInstanceRegistry()->GetSparseInstances())
 	{
-		const IMovieScenePlaybackClient*  Client       = Instance.GetPlayer()->GetPlaybackClient();
+		TSharedRef<const FSharedPlaybackState> SharedPlaybackState = Instance.GetSharedPlaybackState();
+
+		const IMovieScenePlaybackClient*  Client       = SharedPlaybackState->FindCapability<IMovieScenePlaybackClient>();
 		const UObject*                    InstanceData = Client ? Client->GetInstanceData() : nullptr;
 		const IMovieSceneTransformOrigin* RawInterface = Cast<const IMovieSceneTransformOrigin>(InstanceData);
 
@@ -257,8 +260,9 @@ void UMovieSceneTransformOriginSystem::OnRun(FSystemTaskPrerequisites& InPrerequ
 		}
 
 		const FSequenceInstance& Instance = SparseInstances[Index];
+		TSharedRef<const FSharedPlaybackState> SharedPlaybackState = Instance.GetSharedPlaybackState();
 
-		const IMovieScenePlaybackClient*  Client       = Instance.GetPlayer()->GetPlaybackClient();
+		const IMovieScenePlaybackClient*  Client       = SharedPlaybackState->FindCapability<IMovieScenePlaybackClient>();
 		const UObject*                    InstanceData = Client ? Client->GetInstanceData() : nullptr;
 		const IMovieSceneTransformOrigin* RawInterface = Cast<const IMovieSceneTransformOrigin>(InstanceData);
 
diff --git a/Engine/Source/Runtime/MovieSceneTracks/Private/Systems/WeightAndEasingEvaluatorSystem.cpp b/Engine/Source/Runtime/MovieSceneTracks/Private/Systems/WeightAndEasingEvaluatorSystem.cpp
--- a/Engine/Source/Runtime/MovieSceneTracks/Private/Systems/WeightAndEasingEvaluatorSystem.cpp
+++ b/Engine/Source/Runtime/MovieSceneTracks/Private/Systems/WeightAndEasingEvaluatorSystem.cpp
@@ -154,7 +154,7 @@ struct FRootSequenceData
 
 		if (!RootHierarchy)
 		{
-			RootHierarchy = InstanceRegistry->GetInstance(RootInstanceHandle).GetPlayer()->GetEvaluationTemplate().GetHierarchy();
+			RootHierarchy = InstanceRegistry->GetInstance(RootInstanceHandle).GetSharedPlaybackState()->GetHierarchy();
 		}
 
 		if (ensure(RootHierarchy))
@@ -210,7 +210,7 @@ struct FRootSequenceData
 
 		if (!RootHierarchy)
 		{
-			RootHierarchy = InstanceRegistry->GetInstance(RootInstanceHandle).GetPlayer()->GetEvaluationTemplate().GetHierarchy();
+			RootHierarchy = InstanceRegistry->GetInstance(RootInstanceHandle).GetSharedPlaybackState()->GetHierarchy();
 		}
 
 		if (!RootHierarchy)
@@ -1211,4 +1211,4 @@ void UWeightAndEasingEvaluatorSystem::OnRun(FSystemTaskPrerequisites& InPrerequi
 		.Dispatch_PerAllocation<FPropagateHierarchicalEasings>(&Linker->EntityManager, PropagatePrereqs, &Subsequents,
 			PreAllocatedComputationData.Channels);
 	}
-}
\ No newline at end of file
+}
diff --git a/Engine/Source/Runtime/MovieSceneTracks/Private/TrackInstances/MovieSceneCameraCutEditorHandler.cpp b/Engine/Source/Runtime/MovieSceneTracks/Private/TrackInstances/MovieSceneCameraCutEditorHandler.cpp
--- a/Engine/Source/Runtime/MovieSceneTracks/Private/TrackInstances/MovieSceneCameraCutEditorHandler.cpp
+++ b/Engine/Source/Runtime/MovieSceneTracks/Private/TrackInstances/MovieSceneCameraCutEditorHandler.cpp
@@ -90,8 +90,7 @@ void FCameraCutEditorHandler::CachePreAnimatedValue(
 		return;
 	}
 
-	IMovieScenePlayer* Player = SequenceInstance.GetPlayer();
-	UObject* PlaybackContext = Player->GetPlaybackContext();
+	UObject* PlaybackContext = SequenceInstance.GetSharedPlaybackState()->GetPlaybackContext();
 	UWorld* ContextWorld = PlaybackContext ? PlaybackContext->GetWorld() : nullptr;
 
 	// Only handle editor world/viewports.
@@ -129,8 +128,7 @@ void FCameraCutEditorHandler::ForcePreAnimatedValueOperation(
 
 	TGuardValue<FEntityManager*> EntityManagerForDebugging(GEntityManagerForDebuggingVisualizers, &Linker->EntityManager);
 
-	IMovieScenePlayer* Player = SequenceInstance.GetPlayer();
-	UObject* PlaybackContext = Player->GetPlaybackContext();
+	UObject* PlaybackContext = SequenceInstance.GetSharedPlaybackState()->GetPlaybackContext();
 	UWorld* ContextWorld = PlaybackContext ? PlaybackContext->GetWorld() : nullptr;
 
 	// Only handle editor world/viewports.
@@ -200,8 +198,7 @@ void FCameraCutEditorHandler::SetCameraCut(
 		return;
 	}
 
-	IMovieScenePlayer* Player = SequenceInstance.GetPlayer();
-	UObject* PlaybackContext = Player->GetPlaybackContext();
+	UObject* PlaybackContext = SequenceInstance.GetSharedPlaybackState()->GetPlaybackContext();
 	UWorld* ContextWorld = PlaybackContext ? PlaybackContext->GetWorld() : nullptr;
 
 	// Only handle editor world/viewports.
diff --git a/Engine/Source/Runtime/MovieSceneTracks/Private/TrackInstances/MovieSceneCameraCutGameHandler.cpp b/Engine/Source/Runtime/MovieSceneTracks/Private/TrackInstances/MovieSceneCameraCutGameHandler.cpp
--- a/Engine/Source/Runtime/MovieSceneTracks/Private/TrackInstances/MovieSceneCameraCutGameHandler.cpp
+++ b/Engine/Source/Runtime/MovieSceneTracks/Private/TrackInstances/MovieSceneCameraCutGameHandler.cpp
@@ -28,10 +28,9 @@ bool FPreAnimatedCameraCutTraits::ShouldHandleWorldCameraCuts(UWorld* World)
 }
 
 FPreAnimatedCameraCutState FPreAnimatedCameraCutTraits::CachePreAnimatedValue(
-		IMovieScenePlayer* Player, 
+		UObject* PlaybackContext,
 		uint8 InKey)
 {
-	UObject* PlaybackContext = Player->GetPlaybackContext();
 	UWorld* World = PlaybackContext ? PlaybackContext->GetWorld() : nullptr;
 	if (ShouldHandleWorldCameraCuts(World))
 	{
@@ -154,10 +153,10 @@ void FCameraCutGameHandler::CachePreAnimatedValue(
 {
 	TSharedPtr<FPreAnimatedCameraCutStorage> PreAnimatedStorage = Linker->PreAnimatedState.GetOrCreateStorage<FPreAnimatedCameraCutStorage>();
 
-	IMovieScenePlayer* Player = SequenceInstance.GetPlayer();
+	UObject* PlaybackContext = SequenceInstance.GetSharedPlaybackState()->GetPlaybackContext();
 	PreAnimatedStorage->CachePreAnimatedValue(
 			(uint8)0,  // Later this can be an index for split-screen player
-			[Player](uint8 InKey) { return FPreAnimatedCameraCutTraits::CachePreAnimatedValue(Player, InKey); },
+			[PlaybackContext](uint8 InKey) { return FPreAnimatedCameraCutTraits::CachePreAnimatedValue(PlaybackContext, InKey); },
 			EPreAnimatedCaptureSourceTracking::AlwaysCache);
 }
 
@@ -181,8 +180,7 @@ void FCameraCutGameHandler::SetCameraCut(
 		return;
 	}
 
-	IMovieScenePlayer* Player = SequenceInstance.GetPlayer();
-	UObject* PlaybackContext = Player->GetPlaybackContext();
+	UObject* PlaybackContext = SequenceInstance.GetSharedPlaybackState()->GetPlaybackContext();
 	UWorld* World = PlaybackContext ? PlaybackContext->GetWorld() : nullptr;
 
 	// Also bail out if we don't have a world running any sort of game.
diff --git a/Engine/Source/Runtime/MovieSceneTracks/Private/TrackInstances/MovieSceneCameraCutGameHandler.h b/Engine/Source/Runtime/MovieSceneTracks/Private/TrackInstances/MovieSceneCameraCutGameHandler.h
--- a/Engine/Source/Runtime/MovieSceneTracks/Private/TrackInstances/MovieSceneCameraCutGameHandler.h
+++ b/Engine/Source/Runtime/MovieSceneTracks/Private/TrackInstances/MovieSceneCameraCutGameHandler.h
@@ -34,7 +34,7 @@ struct FPreAnimatedCameraCutTraits : FPreAnimatedStateTraits
 	using StorageType = FPreAnimatedCameraCutState;
 
 	static bool ShouldHandleWorldCameraCuts(UWorld* World);
-	static StorageType CachePreAnimatedValue(IMovieScenePlayer* Player, uint8 InKey);
+	static StorageType CachePreAnimatedValue(UObject* PlaybackContext, uint8 InKey);
 
 	void RestorePreAnimatedValue(uint8 InKey, const StorageType& CachedValue, const FRestoreStateParams& Params);
 };
diff --git a/Engine/Source/Runtime/MovieSceneTracks/Private/TrackInstances/MovieSceneCameraCutTrackInstance.cpp b/Engine/Source/Runtime/MovieSceneTracks/Private/TrackInstances/MovieSceneCameraCutTrackInstance.cpp
--- a/Engine/Source/Runtime/MovieSceneTracks/Private/TrackInstances/MovieSceneCameraCutTrackInstance.cpp
+++ b/Engine/Source/Runtime/MovieSceneTracks/Private/TrackInstances/MovieSceneCameraCutTrackInstance.cpp
@@ -37,9 +37,9 @@ namespace MovieScene
 
 FCameraCutPlaybackCapabilityCompatibilityWrapper::FCameraCutPlaybackCapabilityCompatibilityWrapper(const FSequenceInstance& SequenceInstance)
 {
-	TSharedRef<FSharedPlaybackState> PlaybackState = SequenceInstance.GetSharedPlaybackState();
-	CameraCutCapability = PlaybackState->FindCapability<FCameraCutPlaybackCapability>();
-	Player = SequenceInstance.GetPlayer();
+	TSharedRef<FSharedPlaybackState> SharedPlaybackState = SequenceInstance.GetSharedPlaybackState();
+	CameraCutCapability = SharedPlaybackState->FindCapability<FCameraCutPlaybackCapability>();
+	Player = FPlayerIndexPlaybackCapability::GetPlayer(SharedPlaybackState);
 }
 
 bool FCameraCutPlaybackCapabilityCompatibilityWrapper::ShouldUpdateCameraCut()
