Sequencer: remove IMovieScenePlayer usage from fade system, add viewport settings playback capability

diff --git a/Engine/Source/Editor/Sequencer/Private/Sequencer.cpp b/Engine/Source/Editor/Sequencer/Private/Sequencer.cpp
--- a/Engine/Source/Editor/Sequencer/Private/Sequencer.cpp
+++ b/Engine/Source/Editor/Sequencer/Private/Sequencer.cpp
@@ -88,6 +88,7 @@
 #include "ScopedTransaction.h"
 #include "ISequencerTrackEditor.h"
 #include "MovieSceneToolHelpers.h"
+#include "Evaluation/ViewportSettingsPlaybackCapability.h"
 #include "Sections/MovieSceneBoolSection.h"
 #include "Sections/MovieScene3DTransformSection.h"
 #include "Sections/MovieSceneSubSection.h"
diff --git a/Engine/Source/Runtime/MovieScene/Public/IMovieScenePlayer.h b/Engine/Source/Runtime/MovieScene/Public/IMovieScenePlayer.h
--- a/Engine/Source/Runtime/MovieScene/Public/IMovieScenePlayer.h
+++ b/Engine/Source/Runtime/MovieScene/Public/IMovieScenePlayer.h
@@ -27,6 +27,7 @@ struct FMovieSceneRootEvaluationTemplateInstance;
 class FMovieSceneSequenceInstance;
 class IMovieScenePlayer;
 class IMovieSceneSequencePlayerObserver;
+struct EMovieSceneViewportParams;
 
 namespace UE::MovieScene
 {
@@ -51,31 +52,6 @@ namespace UE::MovieScene
 	};
 }
 
-struct EMovieSceneViewportParams
-{
-	EMovieSceneViewportParams()
-	{
-		FadeAmount = 0.f;
-		FadeColor = FLinearColor::Black;
-		bEnableColorScaling = false;
-	}
-
-	enum SetViewportParam
-	{
-		SVP_FadeAmount   = 0x00000001,
-		SVP_FadeColor    = 0x00000002,
-		SVP_ColorScaling = 0x00000004,
-		SVP_All          = SVP_FadeAmount | SVP_FadeColor | SVP_ColorScaling
-	};
-
-	SetViewportParam SetWhichViewportParam;
-
-	float FadeAmount;
-	FLinearColor FadeColor;
-	FVector ColorScale; 
-	bool bEnableColorScaling;
-};
-
 /** Camera cut parameters */
 struct FMovieSceneCameraCutParams
 {
@@ -136,14 +112,16 @@ public:
 	 *
 	 * @param ViewportParamMap A map from the viewport client to its settings
 	 */
-	virtual void SetViewportSettings(const TMap<FViewportClient*, EMovieSceneViewportParams>& ViewportParamsMap) = 0;
+	UE_DEPRECATED(5.5, "Viewport settings management has moved to FViewportSettingsPlaybackCapability")
+	virtual void SetViewportSettings(const TMap<FViewportClient*, EMovieSceneViewportParams>& ViewportParamsMap) {}
 
 	/*
 	 * Get the current perspective viewport settings
 	 *
 	 * @param ViewportParamMap A map from the viewport client to its settings
 	 */
-	virtual void GetViewportSettings(TMap<FViewportClient*, EMovieSceneViewportParams>& ViewportParamsMap) const = 0;
+	UE_DEPRECATED(5.5, "Viewport settings management has moved to FViewportSettingsPlaybackCapability")
+	virtual void GetViewportSettings(TMap<FViewportClient*, EMovieSceneViewportParams>& ViewportParamsMap) const {}
 
 	/** @return whether the player is currently playing, scrubbing, etc. */
 	virtual EMovieScenePlayerStatus::Type GetPlaybackStatus() const = 0;
diff --git a/Engine/Source/Runtime/MovieSceneTracks/Private/Evaluation/ViewportSettingsPlaybackCapability.cpp b/Engine/Source/Runtime/MovieSceneTracks/Private/Evaluation/ViewportSettingsPlaybackCapability.cpp
new file mode 100644
--- /dev/null
+++ b/Engine/Source/Runtime/MovieSceneTracks/Private/Evaluation/ViewportSettingsPlaybackCapability.cpp
@@ -0,0 +1,11 @@
+// Copyright Epic Games, Inc. All Rights Reserved.
+
+#include "Evaluation/ViewportSettingsPlaybackCapability.h"
+
+namespace UE::MovieScene
+{
+
+TPlaybackCapabilityID<FViewportSettingsPlaybackCapability> FViewportSettingsPlaybackCapability::ID = TPlaybackCapabilityID<FViewportSettingsPlaybackCapability>::Register();
+
+} // namespace UE::MovieScene
+
diff --git a/Engine/Source/Runtime/MovieSceneTracks/Private/Systems/MovieSceneFadeSystem.cpp b/Engine/Source/Runtime/MovieSceneTracks/Private/Systems/MovieSceneFadeSystem.cpp
--- a/Engine/Source/Runtime/MovieSceneTracks/Private/Systems/MovieSceneFadeSystem.cpp
+++ b/Engine/Source/Runtime/MovieSceneTracks/Private/Systems/MovieSceneFadeSystem.cpp
@@ -1,29 +1,71 @@
 // Copyright Epic Games, Inc. All Rights Reserved.
 
 #include "Systems/MovieSceneFadeSystem.h"
+
 #include "Async/TaskGraphInterfaces.h"
 #include "Engine/Engine.h"
 #include "Engine/GameInstance.h"
 #include "EngineGlobals.h"
 #include "EntitySystem/BuiltInComponentTypes.h"
+#include "EntitySystem/BuiltInComponentTypes.h"
 #include "Evaluation/PreAnimatedState/MovieScenePreAnimatedStateStorage.h"
 #include "Evaluation/PreAnimatedState/MovieScenePreAnimatedStorageID.inl"
+#include "Evaluation/ViewportSettingsPlaybackCapability.h"
 #include "GameFramework/PlayerController.h"
 #include "IMovieScenePlayer.h"
 #include "MovieSceneExecutionToken.h"
 #include "MovieSceneSequence.h"
-#include "Sections/MovieSceneFadeSection.h"
-#include "EntitySystem/BuiltInComponentTypes.h"
 #include "MovieSceneTracksComponentTypes.h"
+#include "Sections/MovieSceneFadeSection.h"
 
 #include UE_INLINE_GENERATED_CPP_BY_NAME(MovieSceneFadeSystem)
 
 namespace UE::MovieScene
 {
 
+struct FViewportSettingsPlaybackCapabilityCompatibilityWrapper
+{
+	FViewportSettingsPlaybackCapabilityCompatibilityWrapper(TSharedRef<FSharedPlaybackState> SharedPlaybackState)
+	{
+		ViewportSettingsCapability = SharedPlaybackState->FindCapability<FViewportSettingsPlaybackCapability>();
+		Player = FPlayerIndexPlaybackCapability::GetPlayer(SharedPlaybackState);
+	}
+
+	void SetViewportSettings(const TMap<FViewportClient*, EMovieSceneViewportParams>& ViewportParamsMap)
+	{
+		if (ViewportSettingsCapability)
+		{
+			ViewportSettingsCapability->SetViewportSettings(ViewportParamsMap);
+		}
+		else if (Player)
+		{
+			PRAGMA_DISABLE_DEPRECATION_WARNINGS
+			Player->SetViewportSettings(ViewportParamsMap);
+			PRAGMA_ENABLE_DEPRECATION_WARNINGS
+		}
+	}
+
+	void GetViewportSettings(TMap<FViewportClient*, EMovieSceneViewportParams>& ViewportParamsMap) const
+	{
+		if (ViewportSettingsCapability)
+		{
+			ViewportSettingsCapability->GetViewportSettings(ViewportParamsMap);
+		}
+		else if (Player)
+		{
+			PRAGMA_DISABLE_DEPRECATION_WARNINGS
+			Player->GetViewportSettings(ViewportParamsMap);
+			PRAGMA_ENABLE_DEPRECATION_WARNINGS
+		}
+	}
+
+	FViewportSettingsPlaybackCapability* ViewportSettingsCapability;
+	IMovieScenePlayer* Player;
+};
+
 struct FFadeUtil
 {
-	static void ApplyFade(IMovieScenePlayer& Player, float FadeValue, const FLinearColor& FadeColor, bool bFadeAudio)
+	static void ApplyFade(TSharedRef<FSharedPlaybackState> SharedPlaybackState, float FadeValue, const FLinearColor& FadeColor, bool bFadeAudio)
 	{
 		// Set editor preview/fade
 		EMovieSceneViewportParams ViewportParams;
@@ -31,17 +73,18 @@ struct FFadeUtil
 		ViewportParams.FadeAmount = FadeValue;
 		ViewportParams.FadeColor = FadeColor;
 
+		FViewportSettingsPlaybackCapabilityCompatibilityWrapper ViewportSettingsCapability(SharedPlaybackState);
 		TMap<FViewportClient*, EMovieSceneViewportParams> ViewportParamsMap;
-		Player.GetViewportSettings( ViewportParamsMap );
-		for( auto ViewportParamsPair : ViewportParamsMap )
+		ViewportSettingsCapability.GetViewportSettings(ViewportParamsMap);
+		for (auto ViewportParamsPair : ViewportParamsMap)
 		{
-			ViewportParamsMap[ ViewportParamsPair.Key ] = ViewportParams;
+			ViewportParamsMap[ViewportParamsPair.Key] = ViewportParams;
 		}
-		Player.SetViewportSettings( ViewportParamsMap );
+		ViewportSettingsCapability.SetViewportSettings(ViewportParamsMap);
 
 		// Set runtime fade
-		UObject* Context = Player.GetPlaybackContext();
-		UWorld* World = Context ? Context->GetWorld() : nullptr;
+		UObject* PlaybackContext = SharedPlaybackState->GetPlaybackContext();
+		UWorld* World = PlaybackContext ? PlaybackContext->GetWorld() : nullptr;
 		if( World && ( World->WorldType == EWorldType::Game || World->WorldType == EWorldType::PIE ) )
 		{
 			APlayerController* PlayerController = World->GetGameInstance()->GetFirstLocalPlayerController();
@@ -59,14 +102,13 @@ struct FPreAnimatedFadeState
 	FLinearColor FadeColor;
 	bool bFadeAudio;
 
-	static FPreAnimatedFadeState SaveState(IMovieScenePlayer* Player)
+	static FPreAnimatedFadeState SaveState(UObject* PlaybackContext)
 	{
 		float FadeAmount = 0.f;
 		FLinearColor FadeColor = FLinearColor::Black;
 		bool bFadeAudio = false;
 
-		UObject* Context = Player->GetPlaybackContext();
-		UWorld* World = Context ? Context->GetWorld() : nullptr;
+		UWorld* World = PlaybackContext ? PlaybackContext->GetWorld() : nullptr;
 		if (World && (World->WorldType == EWorldType::Game || World->WorldType == EWorldType::PIE))
 		{
 			APlayerController* PlayerController = World->GetGameInstance()->GetFirstLocalPlayerController();
@@ -83,13 +125,13 @@ struct FPreAnimatedFadeState
 
 	void RestoreState(const FMovieSceneAnimTypeID& Unused, const FRestoreStateParams& Params)
 	{
-		IMovieScenePlayer* Player = Params.GetTerminalPlayer();
-		if (!ensure(Player))
+		TSharedPtr<FSharedPlaybackState> SharedPlaybackState = Params.GetTerminalPlaybackState();
+		if (!ensure(SharedPlaybackState))
 		{
 			return;
 		}
 		
-		FFadeUtil::ApplyFade(*Player, FadeValue, FadeColor, bFadeAudio);
+		FFadeUtil::ApplyFade(SharedPlaybackState.ToSharedRef(), FadeValue, FadeColor, bFadeAudio);
 	}
 };
 
@@ -126,13 +168,13 @@ struct FEvaluateFade
 		{
 			FRootInstanceHandle RootInstanceHandle = RootInstanceHandles[Index];
 			const FSequenceInstance& Instance = InstanceRegistry->GetInstance(RootInstanceHandle);
-			IMovieScenePlayer* Player = Instance.GetPlayer();
+			UObject* PlaybackContext = Instance.GetSharedPlaybackState()->GetPlaybackContext();
 
 			PreAnimatedStorage->BeginTrackingEntity(EntityIDs[Index], bWantsRestoreState, RootInstanceHandle, Key);
-			PreAnimatedStorage->CachePreAnimatedValue(Key, [Player](const FMovieSceneAnimTypeID&) { return FPreAnimatedFadeState::SaveState(Player); });
+			PreAnimatedStorage->CachePreAnimatedValue(Key, [PlaybackContext](const FMovieSceneAnimTypeID&) { return FPreAnimatedFadeState::SaveState(PlaybackContext); });
 
 			const FFadeComponentData& FadeComponent(FadeComponents[Index]);
-			FFadeUtil::ApplyFade(*Player, FadeAmounts[Index], FadeComponent.FadeColor, FadeComponent.bFadeAudio);
+			FFadeUtil::ApplyFade(Instance.GetSharedPlaybackState(), FadeAmounts[Index], FadeComponent.FadeColor, FadeComponent.bFadeAudio);
 		}
 	}
 };
diff --git a/Engine/Source/Runtime/MovieSceneTracks/Public/Evaluation/ViewportSettingsPlaybackCapability.h b/Engine/Source/Runtime/MovieSceneTracks/Public/Evaluation/ViewportSettingsPlaybackCapability.h
new file mode 100644
--- /dev/null
+++ b/Engine/Source/Runtime/MovieSceneTracks/Public/Evaluation/ViewportSettingsPlaybackCapability.h
@@ -0,0 +1,64 @@
+// Copyright Epic Games, Inc. All Rights Reserved.
+
+#pragma once
+
+#include "Containers/Map.h"
+#include "CoreTypes.h"
+#include "Evaluation/MovieScenePlaybackCapabilities.h"
+#include "Math/MathFwd.h"
+
+class FViewportClient;
+
+struct EMovieSceneViewportParams
+{
+	EMovieSceneViewportParams()
+	{
+		FadeAmount = 0.f;
+		FadeColor = FLinearColor::Black;
+		bEnableColorScaling = false;
+	}
+
+	enum SetViewportParam
+	{
+		SVP_FadeAmount   = 0x00000001,
+		SVP_FadeColor    = 0x00000002,
+		SVP_ColorScaling = 0x00000004,
+		SVP_All          = SVP_FadeAmount | SVP_FadeColor | SVP_ColorScaling
+	};
+
+	SetViewportParam SetWhichViewportParam;
+
+	float FadeAmount;
+	FLinearColor FadeColor;
+	FVector ColorScale; 
+	bool bEnableColorScaling;
+};
+
+namespace UE::MovieScene
+{
+
+/**
+ * Playback capability for controlling game and editor viewports.
+ */
+struct MOVIESCENETRACKS_API FViewportSettingsPlaybackCapability
+{
+	/** Playback capability ID */
+	static TPlaybackCapabilityID<FViewportSettingsPlaybackCapability> ID;
+
+	/*
+	 * Set the perspective viewport settings
+	 *
+	 * @param ViewportParamMap A map from the viewport client to its settings
+	 */
+	virtual void SetViewportSettings(const TMap<FViewportClient*, EMovieSceneViewportParams>& ViewportParamsMap) = 0;
+
+	/*
+	 * Get the current perspective viewport settings
+	 *
+	 * @param ViewportParamMap A map from the viewport client to its settings
+	 */
+	virtual void GetViewportSettings(TMap<FViewportClient*, EMovieSceneViewportParams>& ViewportParamsMap) const = 0;
+};
+
+}  // UE::MovieScene
+
