Sequencer- When animating SkeletalMeshAsset property on Skeletal Mesh Components, ensure firing due to property not being found in cooked builds.

diff --git a/Engine/Source/Runtime/MovieScene/Public/EntitySystem/MovieScenePropertySystemTypes.h b/Engine/Source/Runtime/MovieScene/Public/EntitySystem/MovieScenePropertySystemTypes.h
--- a/Engine/Source/Runtime/MovieScene/Public/EntitySystem/MovieScenePropertySystemTypes.h
+++ b/Engine/Source/Runtime/MovieScene/Public/EntitySystem/MovieScenePropertySystemTypes.h
@@ -279,7 +279,7 @@ struct TCustomPropertyRegistration : ICustomPropertyRegistration
 		}
 	}
 
-private:
+protected:
 
 	/** */
 	TArray<TCustomPropertyAccessor<PropertyTraits>, TInlineAllocator<InlineSize>> CustomAccessors;
diff --git a/Engine/Source/Runtime/MovieSceneTracks/Private/MovieSceneTracksComponentTypes.cpp b/Engine/Source/Runtime/MovieSceneTracks/Private/MovieSceneTracksComponentTypes.cpp
--- a/Engine/Source/Runtime/MovieSceneTracks/Private/MovieSceneTracksComponentTypes.cpp
+++ b/Engine/Source/Runtime/MovieSceneTracks/Private/MovieSceneTracksComponentTypes.cpp
@@ -518,10 +518,21 @@ struct FObjectHandler : TPropertyComponentHandler<FObjectPropertyTraits, FObject
 		FEntityTaskBuilder()
 			.Read(BuiltInComponents->BoundObject)
 			.Read(BuiltInComponents->PropertyBinding)
+			.ReadOptional(BuiltInComponents->CustomPropertyIndex)
 			.Write(TrackComponents->Object.MetaDataComponents.GetType<0>())
 			.FilterAll({ BuiltInComponents->Tags.NeedsLink })
-			.Iterate_PerEntity(&Linker->EntityManager, [](UObject* Object, const FMovieScenePropertyBinding& Binding, FObjectPropertyTraits::FObjectMetadata& OutMetaData)
+			.Iterate_PerEntity(&Linker->EntityManager, [TrackComponents](UObject* Object, const FMovieScenePropertyBinding& Binding, const FCustomPropertyIndex* OptionalCustomPropertyIndex, FObjectPropertyTraits::FObjectMetadata& OutMetaData)
 				{
+					if (OptionalCustomPropertyIndex)
+					{
+						if (const auto* CustomObjectMetaData = TrackComponents->Accessors.Object.MetaData.Find(OptionalCustomPropertyIndex->Value))
+						{
+							OutMetaData.ObjectClass = CustomObjectMetaData->AllowedClass.Get();
+							OutMetaData.bAllowsClear = CustomObjectMetaData->bAllowsClear;
+							return;
+						}
+					}
+
 					FObjectPropertyBase* BoundProperty = CastField<FObjectPropertyBase>(FTrackInstancePropertyBindings::FindProperty(Object, Binding.PropertyPath.ToString()));
 					if (ensure(BoundProperty))
 					{
diff --git a/Engine/Source/Runtime/MovieSceneTracks/Private/MovieSceneTracksCustomAccessors.cpp b/Engine/Source/Runtime/MovieSceneTracks/Private/MovieSceneTracksCustomAccessors.cpp
--- a/Engine/Source/Runtime/MovieSceneTracks/Private/MovieSceneTracksCustomAccessors.cpp
+++ b/Engine/Source/Runtime/MovieSceneTracks/Private/MovieSceneTracksCustomAccessors.cpp
@@ -728,10 +728,10 @@ void InitializeMovieSceneTracksAccessors(FMovieSceneTracksComponentTypes* Tracks
 	TracksComponents->Accessors.ComponentTransform.Add(USceneComponent::StaticClass(), "Transform", &GetComponentTransform, &SetComponentTransformAndVelocity);
 
 	// SkeletalMeshComponent
-	TracksComponents->Accessors.Object.Add(USkeletalMeshComponent::StaticClass(), "SkeletalMeshAsset", &GetSkeletalMeshAsset, &SetSkeletalMeshAsset);
+	TracksComponents->Accessors.Object.Add(USkeletalMeshComponent::StaticClass(), "SkeletalMeshAsset", &GetSkeletalMeshAsset, &SetSkeletalMeshAsset, { USkeletalMeshComponent::StaticClass(), true });
 
 	//SkinnedMeshComponent
-	TracksComponents->Accessors.Object.Add(USkinnedMeshComponent::StaticClass(), "SkeletalMesh", &GetSkeletalMesh_DEPRECATED, &SetSkeletalMesh_DEPRECATED);
+	TracksComponents->Accessors.Object.Add(USkinnedMeshComponent::StaticClass(), "SkeletalMesh", &GetSkeletalMesh_DEPRECATED, &SetSkeletalMesh_DEPRECATED, { USkinnedMeshComponent::StaticClass() , true});
 
 	
 }
diff --git a/Engine/Source/Runtime/MovieSceneTracks/Public/MovieSceneTracksComponentTypes.h b/Engine/Source/Runtime/MovieSceneTracks/Public/MovieSceneTracksComponentTypes.h
--- a/Engine/Source/Runtime/MovieSceneTracks/Public/MovieSceneTracksComponentTypes.h
+++ b/Engine/Source/Runtime/MovieSceneTracks/Public/MovieSceneTracksComponentTypes.h
@@ -719,7 +719,30 @@ struct FMovieSceneTracksComponentTypes
 	TComponentTypeID<FMovieSceneCameraShakeComponentData> CameraShake;
 	TComponentTypeID<FMovieSceneCameraShakeInstanceData> CameraShakeInstance;
 
-	struct
+	struct FObjectPropertyRegistration : TCustomPropertyRegistration<FObjectPropertyTraits>
+	{
+		struct FMetaData
+		{
+			TWeakObjectPtr<UClass> AllowedClass;
+			bool bAllowsClear = true;
+		};
+
+		void Add(UClass* ClassType, FName PropertyName, GetterFunc Getter, SetterFunc Setter)
+		{
+			TCustomPropertyRegistration<FObjectPropertyTraits>::Add(ClassType, PropertyName, Getter, Setter);
+		}
+
+		void Add(UClass* ClassType, FName PropertyName, GetterFunc Getter, SetterFunc Setter, FMetaData InMetaData)
+		{
+			int32 CustomIndex = CustomAccessors.Num();
+			TCustomPropertyRegistration<FObjectPropertyTraits>::Add(ClassType, PropertyName, Getter, Setter);
+			MetaData.Add(CustomIndex, InMetaData);
+		}
+
+		TMap<int32, FMovieSceneTracksComponentTypes::FObjectPropertyRegistration::FMetaData> MetaData;
+	};
+
+	struct FAccessors
 	{
 		TCustomPropertyRegistration<FBoolPropertyTraits> Bool;
 		TCustomPropertyRegistration<FBytePropertyTraits> Byte;
@@ -731,7 +754,8 @@ struct FMovieSceneTracksComponentTypes
 		TCustomPropertyRegistration<FFloatVectorPropertyTraits> FloatVector;
 		TCustomPropertyRegistration<FDoubleVectorPropertyTraits> DoubleVector;
 		TCustomPropertyRegistration<FComponentTransformPropertyTraits, 1> ComponentTransform;
-		TCustomPropertyRegistration<FObjectPropertyTraits> Object;
+		FObjectPropertyRegistration Object;
+
 	} Accessors;
 
 	struct
