Sequencer- Fix issue where hierarchical bias wasn't working with visibility tracks.

diff --git a/Engine/Source/Runtime/MovieSceneTracks/Private/Sections/MovieSceneVisibilitySection.cpp b/Engine/Source/Runtime/MovieSceneTracks/Private/Sections/MovieSceneVisibilitySection.cpp
--- a/Engine/Source/Runtime/MovieSceneTracks/Private/Sections/MovieSceneVisibilitySection.cpp
+++ b/Engine/Source/Runtime/MovieSceneTracks/Private/Sections/MovieSceneVisibilitySection.cpp
@@ -10,6 +10,7 @@
 #include "MovieScene.h"
 #include "MovieSceneTimeHelpers.h"
 #include "MovieSceneTracksComponentTypes.h"
+#include "Tracks/MovieScenePropertyTrack.h"
 
 #include UE_INLINE_GENERATED_CPP_BY_NAME(MovieSceneVisibilitySection)
 
@@ -44,13 +45,18 @@ void UMovieSceneVisibilitySection::ImportEntityImpl(UMovieSceneEntitySystemLinke
 
 	FBuiltInComponentTypes* BuiltInComponents = FBuiltInComponentTypes::Get();
 	FMovieSceneTracksComponentTypes* TrackComponents = FMovieSceneTracksComponentTypes::Get();
+	UMovieScenePropertyTrack* OuterPropertyTrack = GetTypedOuter<UMovieScenePropertyTrack>();
 
-	OutImportedEntity->AddBuilder(
-		FEntityBuilder()
-		.AddConditional(BuiltInComponents->GenericObjectBinding, Params.GetObjectBindingID(), Params.GetObjectBindingID().IsValid())
-		.AddTag(TrackComponents->Tags.Visibility)
-		.Add(BuiltInComponents->BoolResult, EntityIDToVisibility(Params.EntityID))
-	);
+	if (OuterPropertyTrack)
+	{
+		OutImportedEntity->AddBuilder(
+			FEntityBuilder()
+			.AddConditional(BuiltInComponents->GenericObjectBinding, Params.GetObjectBindingID(), Params.GetObjectBindingID().IsValid())
+			.AddTag(TrackComponents->Tags.Visibility)
+			.Add(BuiltInComponents->BoolResult, EntityIDToVisibility(Params.EntityID))
+			.Add(BuiltInComponents->PropertyBinding, OuterPropertyTrack->GetPropertyBinding())
+		);
+	}
 }
 
 bool UMovieSceneVisibilitySection::PopulateEvaluationFieldImpl(const TRange<FFrameNumber>& EffectiveRange, const FMovieSceneEvaluationFieldEntityMetaData& InMetaData, FMovieSceneEntityComponentFieldBuilder* OutFieldBuilder)
diff --git a/Engine/Source/Runtime/MovieSceneTracks/Private/Systems/MovieSceneVisibilitySystem.cpp b/Engine/Source/Runtime/MovieSceneTracks/Private/Systems/MovieSceneVisibilitySystem.cpp
--- a/Engine/Source/Runtime/MovieSceneTracks/Private/Systems/MovieSceneVisibilitySystem.cpp
+++ b/Engine/Source/Runtime/MovieSceneTracks/Private/Systems/MovieSceneVisibilitySystem.cpp
@@ -10,7 +10,10 @@
 #include "Evaluation/PreAnimatedState/MovieScenePreAnimatedObjectStorage.h"
 #include "Evaluation/PreAnimatedState/MovieScenePreAnimatedStorageID.h"
 #include "Evaluation/PreAnimatedState/MovieScenePreAnimatedStorageID.inl"
+#include "Systems/MovieSceneHierarchicalBiasSystem.h"
 #include "MovieSceneTracksComponentTypes.h"
+#include "EntitySystem/MovieSceneEntityGroupingSystem.h"
+#include "Systems/MovieScenePropertyInstantiator.h"
 
 #include UE_INLINE_GENERATED_CPP_BY_NAME(MovieSceneVisibilitySystem)
 
@@ -151,7 +154,10 @@ UMovieSceneVisibilitySystem::UMovieSceneVisibilitySystem(const FObjectInitialize
 	if (HasAnyFlags(RF_ClassDefaultObject))
 	{
 		DefineComponentConsumer(GetClass(), FBuiltInComponentTypes::Get()->BoundObject);
+		DefineComponentConsumer(GetClass(), FBuiltInComponentTypes::Get()->Tags.Ignored);
 	}
+
+	DefineImplicitPrerequisite(UMovieSceneCachePreAnimatedStateSystem::StaticClass(), GetClass());
 }
 
 void UMovieSceneVisibilitySystem::OnLink()
@@ -181,6 +187,7 @@ void UMovieSceneVisibilitySystem::OnRun(FSystemTaskPrerequisites& InPrerequisite
 	.Read(BuiltInComponents->BoundObject)
 	.Read(BuiltInComponents->BoolResult)
 	.FilterAll({ BuiltInComponents->Tags.NeedsLink, TrackComponents->Tags.Visibility })
+	.FilterNone({BuiltInComponents->Tags.Ignored})
 	.RunInline_PerAllocation(&Linker->EntityManager, Task);
 }
 
