Sequencer: better handle camera cut preview in PIE vs Simulate

The code now handles the combination of PIE/SIE and being in/out of the player controller (with the "eject" button)

diff --git a/Engine/Source/Runtime/MovieSceneTracks/Private/TrackInstances/MovieSceneCameraCutEditorHandler.cpp b/Engine/Source/Runtime/MovieSceneTracks/Private/TrackInstances/MovieSceneCameraCutEditorHandler.cpp
--- a/Engine/Source/Runtime/MovieSceneTracks/Private/TrackInstances/MovieSceneCameraCutEditorHandler.cpp
+++ b/Engine/Source/Runtime/MovieSceneTracks/Private/TrackInstances/MovieSceneCameraCutEditorHandler.cpp
@@ -18,7 +18,12 @@ namespace UE::MovieScene
 bool FPreAnimatedCameraCutEditorTraits::ShouldHandleViewportCameraCuts(UWorld* ViewportWorld)
 {
 	return ViewportWorld && 
-		(ViewportWorld->WorldType == EWorldType::Editor || ViewportWorld->WorldType == EWorldType::EditorPreview);
+		// We can handle editor worlds, and game worlds that don't have an active player controller/pawn,
+		// such as PIE/SIE where the user has "ejected" out of the player controller.
+		(ViewportWorld->WorldType == EWorldType::Editor || 
+			ViewportWorld->WorldType == EWorldType::EditorPreview ||
+			(ViewportWorld->WorldType == EWorldType::PIE && 
+				GEditor && GEditor->bIsSimulatingInEditor));
 }
 
 FPreAnimatedCameraCutEditorState FPreAnimatedCameraCutEditorTraits::CachePreAnimatedValue(
@@ -155,6 +160,7 @@ void FCameraCutEditorHandler::ForcePreAnimatedValueOperation(
 				case EForcedCameraCutPreAnimatedStorageOperation::Cache:
 					{
 						PreAnimatedStorage->DiscardPreAnimatedStateStorage(StorageIndex, EPreAnimatedStorageRequirement::Transient);
+						PreAnimatedStorage->DiscardPreAnimatedStateStorage(StorageIndex, EPreAnimatedStorageRequirement::Persistent);
 						PreAnimatedStorage->CachePreAnimatedValue(
 								LevelVC,
 								[](FLevelEditorViewportClient* InKey) { return FPreAnimatedCameraCutEditorTraits::CachePreAnimatedValue(InKey); },
diff --git a/Engine/Source/Runtime/MovieSceneTracks/Private/TrackInstances/MovieSceneCameraCutGameHandler.cpp b/Engine/Source/Runtime/MovieSceneTracks/Private/TrackInstances/MovieSceneCameraCutGameHandler.cpp
--- a/Engine/Source/Runtime/MovieSceneTracks/Private/TrackInstances/MovieSceneCameraCutGameHandler.cpp
+++ b/Engine/Source/Runtime/MovieSceneTracks/Private/TrackInstances/MovieSceneCameraCutGameHandler.cpp
@@ -16,15 +16,27 @@
 #include "Systems/MovieSceneMotionVectorSimulationSystem.h"
 #include "TrackInstances/MovieSceneCameraCutTrackInstance.h"
 
+#if WITH_EDITOR
+#include "Editor.h"
+#endif
+
 namespace UE::MovieScene
 {
 
 bool FPreAnimatedCameraCutTraits::ShouldHandleWorldCameraCuts(UWorld* World)
 {
-	return World && 
+	return World &&
+		// We can handle any ongoing game worlds. We just don't handle worlds where there is
+		// no active player controller/pawn, such as PIE/SIE where the user has "ejected" out
+		// of the player controller.
 		World->GetGameInstance() != nullptr &&
 		World->WorldType != EWorldType::Editor &&
-		World->WorldType != EWorldType::EditorPreview;
+		World->WorldType != EWorldType::EditorPreview
+#if WITH_EDITOR
+		&&
+		(!GEditor || !GEditor->bIsSimulatingInEditor)
+#endif
+		;
 }
 
 FPreAnimatedCameraCutState FPreAnimatedCameraCutTraits::CachePreAnimatedValue(
@@ -147,6 +159,22 @@ static TTuple<EViewTargetBlendFunction, float> BuiltInEasingTypeToBlendFunction(
 	return Return(EViewTargetBlendFunction::VTBlend_Linear, 1.f);
 }
 
+void FCameraCutGameHandler::ForcePreAnimatedValueRestore(
+			UMovieSceneEntitySystemLinker* Linker,
+			const FSequenceInstance& SequenceInstance)
+{
+	TSharedPtr<FPreAnimatedCameraCutStorage> PreAnimatedStorage = Linker->PreAnimatedState.GetOrCreateStorage<FPreAnimatedCameraCutStorage>();
+
+	FRestoreStateParams Params;
+	Params.Linker = Linker;
+	Params.TerminalInstanceHandle = SequenceInstance.GetRootInstanceHandle();
+	PreAnimatedStorage->RestorePreAnimatedStateStorage(
+			(uint8)0,  // See comment below
+			EPreAnimatedStorageRequirement::Transient, 
+			EPreAnimatedStorageRequirement::Persistent,
+			Params);
+}
+
 void FCameraCutGameHandler::CachePreAnimatedValue(
 		UMovieSceneEntitySystemLinker* Linker,
 		const FSequenceInstance& SequenceInstance)
diff --git a/Engine/Source/Runtime/MovieSceneTracks/Private/TrackInstances/MovieSceneCameraCutGameHandler.h b/Engine/Source/Runtime/MovieSceneTracks/Private/TrackInstances/MovieSceneCameraCutGameHandler.h
--- a/Engine/Source/Runtime/MovieSceneTracks/Private/TrackInstances/MovieSceneCameraCutGameHandler.h
+++ b/Engine/Source/Runtime/MovieSceneTracks/Private/TrackInstances/MovieSceneCameraCutGameHandler.h
@@ -63,6 +63,9 @@ struct FCameraCutGameHandler
 	static void CachePreAnimatedValue(
 			UMovieSceneEntitySystemLinker* Linker,
 			const FSequenceInstance& SequenceInstance);
+	static void ForcePreAnimatedValueRestore(
+			UMovieSceneEntitySystemLinker* Linker,
+			const FSequenceInstance& SequenceInstance);
 
 private:
 	UMovieSceneEntitySystemLinker* Linker;
diff --git a/Engine/Source/Runtime/MovieSceneTracks/Private/TrackInstances/MovieSceneCameraCutTrackInstance.cpp b/Engine/Source/Runtime/MovieSceneTracks/Private/TrackInstances/MovieSceneCameraCutTrackInstance.cpp
--- a/Engine/Source/Runtime/MovieSceneTracks/Private/TrackInstances/MovieSceneCameraCutTrackInstance.cpp
+++ b/Engine/Source/Runtime/MovieSceneTracks/Private/TrackInstances/MovieSceneCameraCutTrackInstance.cpp
@@ -292,7 +292,7 @@ void UMovieSceneCameraCutTrackInstance::ToggleCameraCutLock(UMovieSceneEntitySys
 {
 	using namespace UE::MovieScene;
 
-	auto ForceEditorPreAnimatedStorageOperation = [](UMovieSceneCameraCutTrackInstance* This, UE::MovieScene::EForcedCameraCutPreAnimatedStorageOperation Operation)
+	auto ForceEditorPreAnimatedStorageOperation = [](UMovieSceneCameraCutTrackInstance* This, EForcedCameraCutPreAnimatedStorageOperation Operation)
 	{
 		using namespace UE::MovieScene;
 
@@ -305,6 +305,19 @@ void UMovieSceneCameraCutTrackInstance::ToggleCameraCutLock(UMovieSceneEntitySys
 			FCameraCutEditorHandler::ForcePreAnimatedValueOperation(Linker, SequenceInstance, Operation);
 		}
 	};
+	auto ForceGamePreAnimatedStorageRestore = [](UMovieSceneCameraCutTrackInstance* This)
+	{
+		using namespace UE::MovieScene;
+
+		UMovieSceneEntitySystemLinker* Linker = This->GetLinker();
+		const FInstanceRegistry* InstanceRegistry = Linker->GetInstanceRegistry();
+		for (const FCameraCutInputInfo& InputInfo : This->SortedInputInfos)
+		{
+			FScopedPreAnimatedCaptureSource CaptureSource(Linker, InputInfo.Input);
+			const FSequenceInstance& SequenceInstance = InstanceRegistry->GetInstance(InputInfo.Input.InstanceHandle);
+			FCameraCutGameHandler::ForcePreAnimatedValueRestore(Linker, SequenceInstance);
+		}
+	};
 
 	// Find the camera cut track instance and forcibly manage its pre-animated state for the editor 
 	// viewports depending on what we want to do with them.
@@ -349,6 +362,10 @@ void UMovieSceneCameraCutTrackInstance::ToggleCameraCutLock(UMovieSceneEntitySys
 				{
 					ForceEditorPreAnimatedStorageOperation(CameraCutTrackInstance, EForcedCameraCutPreAnimatedStorageOperation::Discard);
 				}
+
+				// If we have a PIE session active, we need to tell the game handler to restore its
+				// pre-animated state.
+				ForceGamePreAnimatedStorageRestore(CameraCutTrackInstance);
 			}
 		}
 	}
