[MovieSceneTracks]
*  Fixed non-unity compile error

diff --git a/Engine/Source/Runtime/MovieSceneTracks/Public/Evaluation/ViewportSettingsPlaybackCapability.h b/Engine/Source/Runtime/MovieSceneTracks/Public/Evaluation/ViewportSettingsPlaybackCapability.h
--- a/Engine/Source/Runtime/MovieSceneTracks/Public/Evaluation/ViewportSettingsPlaybackCapability.h
+++ b/Engine/Source/Runtime/MovieSceneTracks/Public/Evaluation/ViewportSettingsPlaybackCapability.h
@@ -5,7 +5,7 @@
 #include "Containers/Map.h"
 #include "CoreTypes.h"
 #include "Evaluation/MovieScenePlaybackCapabilities.h"
-#include "Math/MathFwd.h"
+#include "Math/Vector.h"
 
 class FViewportClient;
 
