Sequencer: Default to using spawnable actor binding as the preview spawnable type for replaceable from director blueprint binding.

diff --git a/Engine/Source/Runtime/MovieSceneTracks/Private/Bindings/MovieSceneReplaceableDirectorBlueprintBinding.cpp b/Engine/Source/Runtime/MovieSceneTracks/Private/Bindings/MovieSceneReplaceableDirectorBlueprintBinding.cpp
--- a/Engine/Source/Runtime/MovieSceneTracks/Private/Bindings/MovieSceneReplaceableDirectorBlueprintBinding.cpp
+++ b/Engine/Source/Runtime/MovieSceneTracks/Private/Bindings/MovieSceneReplaceableDirectorBlueprintBinding.cpp
@@ -48,4 +48,16 @@ FMovieSceneBindingResolveResult UMovieSceneReplaceableDirectorBlueprintBinding::
 	return ResolveResult;
 }
 
+UMovieSceneSpawnableBindingBase* UMovieSceneReplaceableDirectorBlueprintBinding::CreateInnerSpawnable(UObject* SourceObject, UMovieScene& OwnerMovieScene)
+{
+#if WITH_EDITOR
+	// If no inner spawnable class has been set, and it's available, set it to Spawnable Actor so we at least get some preview when converting an existing binding to this type
+	if (!PreviewSpawnableType && UMovieScene::IsCustomBindingClassAllowed(UMovieSceneSpawnableActorBinding::StaticClass()))
+	{
+		PreviewSpawnableType = UMovieSceneSpawnableActorBinding::StaticClass();
+	}
+#endif
+	return UMovieSceneReplaceableBindingBase::CreateInnerSpawnable(SourceObject, OwnerMovieScene);
+}
+
 #undef LOCTEXT_NAMESPACE
\ No newline at end of file
diff --git a/Engine/Source/Runtime/MovieSceneTracks/Public/Bindings/MovieSceneReplaceableDirectorBlueprintBinding.h b/Engine/Source/Runtime/MovieSceneTracks/Public/Bindings/MovieSceneReplaceableDirectorBlueprintBinding.h
--- a/Engine/Source/Runtime/MovieSceneTracks/Public/Bindings/MovieSceneReplaceableDirectorBlueprintBinding.h
+++ b/Engine/Source/Runtime/MovieSceneTracks/Public/Bindings/MovieSceneReplaceableDirectorBlueprintBinding.h
@@ -53,5 +53,7 @@ protected:
 	MOVIESCENETRACKS_API virtual TSubclassOf<UMovieSceneSpawnableBindingBase> GetInnerSpawnableClass() const override { return PreviewSpawnableType; }
 	
 	virtual bool SupportsBindingCreationFromObject(const UObject* SourceObject) const override { return true; }
+
+	virtual UMovieSceneSpawnableBindingBase* CreateInnerSpawnable(UObject* SourceObject, UMovieScene& OwnerMovieScene) override;
 };
 
