Move restoration of bClientSimulatingViewTarget from LevelSequencePlayer::OnStopped to FPreAnimatedCameraCutTraits::RestorePreAnimatedValue so it matches when ViewTarget is restored.

The same class that changes bClientSimulatingViewTarget now restores it as well, and at the correct time.
NOTE that if the camera cut is set to preserve state, you'll never get bClientSimulatingViewTarget restored, but we think that's the correct behavior.

diff --git a/Engine/Source/Runtime/LevelSequence/Private/LevelSequencePlayer.cpp b/Engine/Source/Runtime/LevelSequence/Private/LevelSequencePlayer.cpp
--- a/Engine/Source/Runtime/LevelSequence/Private/LevelSequencePlayer.cpp
+++ b/Engine/Source/Runtime/LevelSequence/Private/LevelSequencePlayer.cpp
@@ -177,19 +177,6 @@ void ULevelSequencePlayer::OnStartedPlaying()
 void ULevelSequencePlayer::OnStopped()
 {
 	EnableCinematicMode(false);
-
-	if (World != nullptr && World->GetGameInstance() != nullptr)
-	{
-		APlayerController* PC = World->GetGameInstance()->GetFirstLocalPlayerController();
-
-		if (PC != nullptr)
-		{
-			if (PC->PlayerCameraManager)
-			{
-				PC->PlayerCameraManager->bClientSimulatingViewTarget = false;
-			}
-		}
-	}
 }
 
 void ULevelSequencePlayer::UpdateMovieSceneInstance(FMovieSceneEvaluationRange InRange, EMovieScenePlayerStatus::Type PlayerStatus, const FMovieSceneUpdateArgs& Args)
diff --git a/Engine/Source/Runtime/MovieSceneTracks/Private/TrackInstances/MovieSceneCameraCutGameHandler.cpp b/Engine/Source/Runtime/MovieSceneTracks/Private/TrackInstances/MovieSceneCameraCutGameHandler.cpp
--- a/Engine/Source/Runtime/MovieSceneTracks/Private/TrackInstances/MovieSceneCameraCutGameHandler.cpp
+++ b/Engine/Source/Runtime/MovieSceneTracks/Private/TrackInstances/MovieSceneCameraCutGameHandler.cpp
@@ -77,6 +77,7 @@ void FPreAnimatedCameraCutTraits::RestorePreAnimatedValue(
 		return;
 	}
 
+	// TODO james.fleming really we should be getting the CameraManager from the sequence player or some other more reliable way that supports split screen.
 	APlayerController* PC = World->GetGameInstance()->GetFirstLocalPlayerController();
 	APlayerCameraManager* CameraManager = (PC != nullptr) ? PC->PlayerCameraManager : nullptr;
 
@@ -87,6 +88,8 @@ void FPreAnimatedCameraCutTraits::RestorePreAnimatedValue(
 	{
 		AActor* PreviousViewTarget = Cast<AActor>(CachedValue.LastViewTarget.ResolveObjectPtr());
 		CameraManager->SetViewTarget(PreviousViewTarget);
+		// TODO james.fleming ideally we would cache this before, just in case it had been set true (which is not usual, but could be possible) 
+		CameraManager->bClientSimulatingViewTarget = false;
 	}
 
 	// Restore previous aspect ratio axis constraint. Use the cached local player if there's no local player
