Sequencer: Fix crash where SkeletalMeshComponent's world was nullptr and not being checked. This code has been in there for a while without getting this crash found, but got triggered on the latest Metallica event. I removed the calling code since the variable wasn't being used.

diff --git a/Engine/Source/Runtime/MovieSceneTracks/Private/Systems/MovieSceneSkeletalAnimationSystem.cpp b/Engine/Source/Runtime/MovieSceneTracks/Private/Systems/MovieSceneSkeletalAnimationSystem.cpp
--- a/Engine/Source/Runtime/MovieSceneTracks/Private/Systems/MovieSceneSkeletalAnimationSystem.cpp
+++ b/Engine/Source/Runtime/MovieSceneTracks/Private/Systems/MovieSceneSkeletalAnimationSystem.cpp
@@ -284,7 +284,6 @@ void FBoneTransformFinalizeData::BoneTransformFinalized()
 		{
 			RelativeTransform = RelativeTransform * InitialTransform.GetValue();
 		}
-		const bool bIsGameWorld = (SkeletalMeshComponent->GetWorld()->IsGameWorld());
 
 		if (SwapRootBone == ESwapRootBone::SwapRootBone_Component)
 		{
