Sequencer: Fix some issues with USD Replaceable Director Bindings.
* We were having re-entrant code since changing a possessable binding into a director replaceable one may call 'OnActorAddedToSequencer' again.
* Fixed an edge case where adding a default transform track without a resolved preview actor initialized scale to 0.
* Slightly re-work the change from 34503068 to only default the preview spawnable for replaceable director blueprint bindings when converting from a binding with an actor. This is because USD bindings don't want a preview spawnable, as they resolve even at editor time.

diff --git a/Engine/Plugins/Importers/USDImporter/Source/USDStage/Private/USDStageActor.cpp b/Engine/Plugins/Importers/USDImporter/Source/USDStage/Private/USDStageActor.cpp
--- a/Engine/Plugins/Importers/USDImporter/Source/USDStage/Private/USDStageActor.cpp
+++ b/Engine/Plugins/Importers/USDImporter/Source/USDStage/Private/USDStageActor.cpp
@@ -4248,6 +4248,12 @@ void AUsdStageActor::UpdateSpawnedObjectsTransientFlag(bool bTransient)
 void AUsdStageActor::OnActorAddedToSequencer(AActor* NewActor, const FGuid Guid, TWeakPtr<ISequencer> WeakSequencer)
 {
 #if WITH_EDITOR
+
+	if (bIsAddingBinding)
+	{
+		// Prevent re-entrancy
+		return;
+	}
 	if (!NewActor || !NewActor->HasAnyFlags(RF_Transient))
 	{
 		return;
@@ -4290,6 +4296,7 @@ void AUsdStageActor::OnActorAddedToSequencer(AActor* NewActor, const FGuid Guid,
 		return;
 	}
 
+
 	// The UI action to add the actor to the sequencer should have spawned it's transaction already, but just in case it
 	// hasn't let's do that, as we'll end up modifying UObjects in here
 	FScopedTransaction Transaction(
@@ -4298,7 +4305,9 @@ void AUsdStageActor::OnActorAddedToSequencer(AActor* NewActor, const FGuid Guid,
 	Sequence->Modify();
 	MovieScene->Modify();
 
+	bIsAddingBinding = true;
 	FUsdStageActorImpl::SetupDynamicBinding(PrimPath, Guid, Sequence, *GetName(), PinnedSequencer);
+	bIsAddingBinding = false;
 #endif	  // WITH_EDITOR
 }
 
diff --git a/Engine/Plugins/Importers/USDImporter/Source/USDStage/Public/USDStageActor.h b/Engine/Plugins/Importers/USDImporter/Source/USDStage/Public/USDStageActor.h
--- a/Engine/Plugins/Importers/USDImporter/Source/USDStage/Public/USDStageActor.h
+++ b/Engine/Plugins/Importers/USDImporter/Source/USDStage/Public/USDStageActor.h
@@ -466,6 +466,7 @@ protected:
 	bool bIsTransitioningIntoPIE;
 	bool bIsModifyingAProperty;
 	bool bIsUndoRedoing;
+	bool bIsAddingBinding = false;
 
 	FDelegateHandle OnRedoHandle;
 	FDelegateHandle OnSequencerCreatedHandle;
diff --git a/Engine/Plugins/MovieScene/LevelSequenceEditor/Source/LevelSequenceEditor/Private/Misc/LevelSequenceEditorActorSpawner.cpp b/Engine/Plugins/MovieScene/LevelSequenceEditor/Source/LevelSequenceEditor/Private/Misc/LevelSequenceEditorActorSpawner.cpp
--- a/Engine/Plugins/MovieScene/LevelSequenceEditor/Source/LevelSequenceEditor/Private/Misc/LevelSequenceEditorActorSpawner.cpp
+++ b/Engine/Plugins/MovieScene/LevelSequenceEditor/Source/LevelSequenceEditor/Private/Misc/LevelSequenceEditorActorSpawner.cpp
@@ -202,6 +202,7 @@ bool FLevelSequenceEditorActorSpawner::CanSetupDefaultsForSpawnable(UObject* Spa
 void FLevelSequenceEditorActorSpawner::SetupDefaultsForSpawnable(UObject* SpawnedObject, const FGuid& Guid, const TOptional<FTransformData>& TransformData, TSharedRef<ISequencer> Sequencer, USequencerSettings* Settings)
 {
 	TOptional<FTransformData> DefaultTransform = TransformData;
+	DefaultTransform->Scale = FVector(1.0f, 1.0f, 1.0f);
 
 	AActor* SpawnedActor = Cast<AActor>(SpawnedObject);
 	if (SpawnedActor)
diff --git a/Engine/Source/Runtime/MovieSceneTracks/Private/Bindings/MovieSceneReplaceableDirectorBlueprintBinding.cpp b/Engine/Source/Runtime/MovieSceneTracks/Private/Bindings/MovieSceneReplaceableDirectorBlueprintBinding.cpp
--- a/Engine/Source/Runtime/MovieSceneTracks/Private/Bindings/MovieSceneReplaceableDirectorBlueprintBinding.cpp
+++ b/Engine/Source/Runtime/MovieSceneTracks/Private/Bindings/MovieSceneReplaceableDirectorBlueprintBinding.cpp
@@ -38,6 +38,27 @@ void UMovieSceneReplaceableDirectorBlueprintBinding::OnBindingAddedOrChanged(UMo
 		}
 	}
 }
+
+UMovieSceneCustomBinding* UMovieSceneReplaceableDirectorBlueprintBinding::CreateCustomBindingFromBinding(const FMovieSceneBindingReference& BindingReference, UObject* SourceObject, UMovieScene& OwnerMovieScene)
+{
+	// We override this specifically to initialize PreviewSpawnableType on conversions so that we end up with a custom preview from whatever object was passed in in that case.
+	UMovieSceneReplaceableDirectorBlueprintBinding* NewCustomBinding = nullptr;
+
+	const FName TemplateName = MakeUniqueObjectName(&OwnerMovieScene, UObject::StaticClass(), SourceObject ? SourceObject->GetFName() : GetClass()->GetFName());
+	const FName InstancedBindingName = MakeUniqueObjectName(&OwnerMovieScene, UObject::StaticClass(), *FString(TemplateName.ToString() + TEXT("_CustomBinding")));
+
+	NewCustomBinding = NewObject<UMovieSceneReplaceableDirectorBlueprintBinding>(&OwnerMovieScene, GetClass(), InstancedBindingName, RF_Transactional);
+
+	// If no inner spawnable class has been set, and it's available, set it to Spawnable Actor so we at least get some preview when converting an existing binding to this type
+	if (!NewCustomBinding->PreviewSpawnableType && UMovieScene::IsCustomBindingClassAllowed(UMovieSceneSpawnableActorBinding::StaticClass()))
+	{
+		NewCustomBinding->PreviewSpawnableType = UMovieSceneSpawnableActorBinding::StaticClass();
+	}
+
+	NewCustomBinding->PreviewSpawnable = NewCustomBinding->CreateInnerSpawnable(SourceObject, OwnerMovieScene);
+	NewCustomBinding->InitReplaceableBinding(SourceObject, OwnerMovieScene);
+	return NewCustomBinding;
+}
 #endif
 
 FMovieSceneBindingResolveResult UMovieSceneReplaceableDirectorBlueprintBinding::ResolveRuntimeBindingInternal(const FMovieSceneBindingResolveParams& ResolveParams, int32 BindingIndex, TSharedRef<const UE::MovieScene::FSharedPlaybackState> SharedPlaybackState) const
@@ -50,14 +71,6 @@ FMovieSceneBindingResolveResult UMovieSceneReplaceableDirectorBlueprintBinding::
 
 UMovieSceneSpawnableBindingBase* UMovieSceneReplaceableDirectorBlueprintBinding::CreateInnerSpawnable(UObject* SourceObject, UMovieScene& OwnerMovieScene)
 {
-#if WITH_EDITOR
-	// If no inner spawnable class has been set, and it's available, set it to Spawnable Actor so we at least get some preview when converting an existing binding to this type
-	if (!PreviewSpawnableType && UMovieScene::IsCustomBindingClassAllowed(UMovieSceneSpawnableActorBinding::StaticClass()))
-	{
-		PreviewSpawnableType = UMovieSceneSpawnableActorBinding::StaticClass();
-	}
-#endif
 	return UMovieSceneReplaceableBindingBase::CreateInnerSpawnable(SourceObject, OwnerMovieScene);
 }
-
 #undef LOCTEXT_NAMESPACE
\ No newline at end of file
diff --git a/Engine/Source/Runtime/MovieSceneTracks/Public/Bindings/MovieSceneReplaceableDirectorBlueprintBinding.h b/Engine/Source/Runtime/MovieSceneTracks/Public/Bindings/MovieSceneReplaceableDirectorBlueprintBinding.h
--- a/Engine/Source/Runtime/MovieSceneTracks/Public/Bindings/MovieSceneReplaceableDirectorBlueprintBinding.h
+++ b/Engine/Source/Runtime/MovieSceneTracks/Public/Bindings/MovieSceneReplaceableDirectorBlueprintBinding.h
@@ -36,13 +36,16 @@ public:
 
 	// Preview Spawnable Type to use for this replaceable
 	// TODO: Do I want to consider moving this sort of thing to the base class so any replaceable type can choose any other spawnable type?
-	// TODO: Changing this will require 
 	UPROPERTY(EditAnywhere, BlueprintReadWrite, Category="Sequencer")
 	TSubclassOf<UMovieSceneSpawnableBindingBase> PreviewSpawnableType;
 
 protected:
 	/* MovieSceneReplaceableBindingBase overrides*/
 
+#if WITH_EDITOR
+	MOVIESCENETRACKS_API virtual UMovieSceneCustomBinding* CreateCustomBindingFromBinding(const FMovieSceneBindingReference& BindingReference, UObject* SourceObject, UMovieScene& OwnerMovieScene) override;
+#endif
+
 	// By default we return nullptr here, as we rely on Sequencer's BindingOverride mechanism to bind these actors during runtime.
 	// This can be overridden if desired in subclasses to provide a different way to resolve to an actor at runtime while still using spawnable actor as the preview.
 	MOVIESCENETRACKS_API virtual FMovieSceneBindingResolveResult ResolveRuntimeBindingInternal(const FMovieSceneBindingResolveParams& ResolveParams, int32 BindingIndex, TSharedRef<const UE::MovieScene::FSharedPlaybackState> SharedPlaybackState) const override;
