[World Partition Sequencer] Properly handle the case where we are unable to set data layers from tracks due to authority.
Also changed logs to verbose to avoid logging them in shipping builds.

diff --git a/Engine/Source/Runtime/Engine/Private/WorldPartition/DataLayer/DataLayerManager.cpp b/Engine/Source/Runtime/Engine/Private/WorldPartition/DataLayer/DataLayerManager.cpp
--- a/Engine/Source/Runtime/Engine/Private/WorldPartition/DataLayer/DataLayerManager.cpp
+++ b/Engine/Source/Runtime/Engine/Private/WorldPartition/DataLayer/DataLayerManager.cpp
@@ -295,8 +295,7 @@ bool UDataLayerManager::SetDataLayerInstanceRuntimeState(const UDataLayerInstanc
 {
 	if (InDataLayerInstance)
 	{
-		InDataLayerInstance->GetOuterWorldDataLayers()->SetDataLayerRuntimeState(InDataLayerInstance, InState, bInIsRecursive);
-		return true;
+		return InDataLayerInstance->GetOuterWorldDataLayers()->SetDataLayerRuntimeState(InDataLayerInstance, InState, bInIsRecursive);
 	}
 	UE_LOG(LogWorldPartition, Error, TEXT("Invalid Data Layer Instance."));
 	return false;
diff --git a/Engine/Source/Runtime/Engine/Private/WorldPartition/DataLayer/WorldDataLayers.cpp b/Engine/Source/Runtime/Engine/Private/WorldPartition/DataLayer/WorldDataLayers.cpp
--- a/Engine/Source/Runtime/Engine/Private/WorldPartition/DataLayer/WorldDataLayers.cpp
+++ b/Engine/Source/Runtime/Engine/Private/WorldPartition/DataLayer/WorldDataLayers.cpp
@@ -179,11 +179,11 @@ void AWorldDataLayers::ResetDataLayerRuntimeStates()
 	AWORLDDATALAYERS_UPDATE_REPLICATED_DATALAYERS(RepEffectiveLoadedDataLayerNames, Empty);
 }
 
-void AWorldDataLayers::SetDataLayerRuntimeState(const UDataLayerInstance* InDataLayerInstance, EDataLayerRuntimeState InState, bool bInIsRecursive)
+bool AWorldDataLayers::SetDataLayerRuntimeState(const UDataLayerInstance* InDataLayerInstance, EDataLayerRuntimeState InState, bool bInIsRecursive)
 {
 	if (!InDataLayerInstance || !InDataLayerInstance->IsRuntime())
 	{
-		return;
+		return false;
 	}
 
 	const ENetMode NetMode = GetNetMode();
@@ -199,35 +199,53 @@ void AWorldDataLayers::SetDataLayerRuntimeState(const UDataLayerInstance* InData
 	{
 		if (NetMode != NM_Standalone && NetMode != NM_Client)
 		{
-			UE_LOG(LogWorldPartition, Log, TEXT("Client Only Data Layer state change '%s' was ignored: %s -> %s"),
+			UE_LOG(LogWorldPartition, Verbose, TEXT("Client Only Data Layer state change '%s' was ignored: %s -> %s"),
 				*InDataLayerInstance->GetDataLayerShortName(),
 				*StaticEnum<EDataLayerRuntimeState>()->GetDisplayNameTextByValue((int64)CurrentState).ToString(),
 				*StaticEnum<EDataLayerRuntimeState>()->GetDisplayNameTextByValue((int64)InState).ToString());
-			return;
+			return false;
 		}
 	}
 	else if (bDataLayerServerOnly)
 	{
 		if (NetMode == NM_Client)
 		{
-			UE_LOG(LogWorldPartition, Log, TEXT("Server Only Data Layer state change '%s' was ignored: %s -> %s"),
+			UE_LOG(LogWorldPartition, Verbose, TEXT("Server Only Data Layer state change '%s' was ignored: %s -> %s"),
 				*InDataLayerInstance->GetDataLayerShortName(),
 				*StaticEnum<EDataLayerRuntimeState>()->GetDisplayNameTextByValue((int64)CurrentState).ToString(),
 				*StaticEnum<EDataLayerRuntimeState>()->GetDisplayNameTextByValue((int64)InState).ToString());
-			return;
+			return false;
 		}
 	}
 	else if (NetMode == NM_Client)
 	{
-		UE_LOG(LogWorldPartition, Log, TEXT("Data Layer '%s' state change was ignored on client: %s -> %s"),
+		UE_LOG(LogWorldPartition, Verbose, TEXT("Data Layer '%s' state change was ignored on client: %s -> %s"),
 			*InDataLayerInstance->GetDataLayerShortName(),
 			*StaticEnum<EDataLayerRuntimeState>()->GetDisplayNameTextByValue((int64)CurrentState).ToString(),
 			*StaticEnum<EDataLayerRuntimeState>()->GetDisplayNameTextByValue((int64)InState).ToString());
-		return;
+		return false;
 	}
 
 	if (CurrentState != InState)
 	{
+		if (GetWorld()->IsGameWorld())
+		{
+			PRAGMA_DISABLE_DEPRECATION_WARNINGS
+			if (DataLayersFilterDelegate.IsBound())
+			{
+				const FString DataLayerShortName(InDataLayerInstance->GetDataLayerShortName());
+				if (!DataLayersFilterDelegate.Execute(FName(*DataLayerShortName), CurrentState, InState))
+				{
+					UE_LOG(LogWorldPartition, Verbose, TEXT("Data Layer '%s' state change was filtered out: %s -> %s"),
+						*DataLayerShortName,
+						*StaticEnum<EDataLayerRuntimeState>()->GetDisplayNameTextByValue((int64)CurrentState).ToString(),
+						*StaticEnum<EDataLayerRuntimeState>()->GetDisplayNameTextByValue((int64)InState).ToString());
+					return false;
+				}
+			}
+			PRAGMA_ENABLE_DEPRECATION_WARNINGS
+		}
+
 		TargetLoadedDataLayerNames.Remove(InDataLayerInstance->GetDataLayerFName());
 		TargetActiveDataLayerNames.Remove(InDataLayerInstance->GetDataLayerFName());
 
@@ -268,6 +286,8 @@ void AWorldDataLayers::SetDataLayerRuntimeState(const UDataLayerInstance* InData
 			return true;
 		});
 	}
+
+	return true;
 }
 
 void AWorldDataLayers::OnDataLayerRuntimeStateChanged_Implementation(const UDataLayerInstance* InDataLayer, EDataLayerRuntimeState InState)
diff --git a/Engine/Source/Runtime/Engine/Public/WorldPartition/DataLayer/WorldDataLayers.h b/Engine/Source/Runtime/Engine/Public/WorldPartition/DataLayer/WorldDataLayers.h
--- a/Engine/Source/Runtime/Engine/Public/WorldPartition/DataLayer/WorldDataLayers.h
+++ b/Engine/Source/Runtime/Engine/Public/WorldPartition/DataLayer/WorldDataLayers.h
@@ -176,7 +176,7 @@ public:
 	ENGINE_API bool IsExternalDataLayerWorldDataLayers() const;
 
 	// DataLayer Runtime State
-	ENGINE_API void SetDataLayerRuntimeState(const UDataLayerInstance* InDataLayerInstance, EDataLayerRuntimeState InState, bool bIsRecursive = false);
+	ENGINE_API bool SetDataLayerRuntimeState(const UDataLayerInstance* InDataLayerInstance, EDataLayerRuntimeState InState, bool bIsRecursive = false);
 	ENGINE_API EDataLayerRuntimeState GetDataLayerRuntimeStateByName(FName InDataLayerName) const;
 	ENGINE_API EDataLayerRuntimeState GetDataLayerEffectiveRuntimeStateByName(FName InDataLAyerName) const;
 	const TSet<FName>& GetEffectiveActiveDataLayerNames() const;
diff --git a/Engine/Source/Runtime/MovieSceneTracks/Private/Systems/MovieSceneDataLayerSystem.cpp b/Engine/Source/Runtime/MovieSceneTracks/Private/Systems/MovieSceneDataLayerSystem.cpp
--- a/Engine/Source/Runtime/MovieSceneTracks/Private/Systems/MovieSceneDataLayerSystem.cpp
+++ b/Engine/Source/Runtime/MovieSceneTracks/Private/Systems/MovieSceneDataLayerSystem.cpp
@@ -347,38 +347,39 @@ EDataLayerUpdateFlags FDesiredLayerStates::Apply(FPreAnimatedDataLayerStorage* P
 				}
 
 				const EDataLayerRuntimeState DesiredStateValue = DesiredState.GetValue();
-				UDataLayerManager::GetDataLayerManager(DataLayer)->SetDataLayerInstanceRuntimeState(DataLayer, DesiredStateValue);
-
-				if (StateValue.ShouldFlushStreaming(DesiredStateValue) && !IsDataLayerReady(DataLayer, DesiredStateValue, true))
+				if (UDataLayerManager::GetDataLayerManager(DataLayer)->SetDataLayerInstanceRuntimeState(DataLayer, DesiredStateValue))
 				{
-					// Exception for Full flush is if Desired State is Activated but we are not at least in Loaded state
-					if (DesiredStateValue == EDataLayerRuntimeState::Activated && !IsDataLayerReady(DataLayer, EDataLayerRuntimeState::Loaded, false))
+					if (StateValue.ShouldFlushStreaming(DesiredStateValue) && !IsDataLayerReady(DataLayer, DesiredStateValue, true))
 					{
-						Flags |= EDataLayerUpdateFlags::FlushStreamingFull;
-						UE_LOG(LogMovieScene, Warning, TEXT("[UMovieSceneDataLayerSystem] Data layer with name '%s' is causing a full streaming flush (%s)"),
-							*DataLayer->GetDataLayerShortName(),
-							*StaticEnum<EDataLayerRuntimeState>()->GetDisplayNameTextByValue((int64)DesiredStateValue).ToString());
+						// Exception for Full flush is if Desired State is Activated but we are not at least in Loaded state
+						if (DesiredStateValue == EDataLayerRuntimeState::Activated && !IsDataLayerReady(DataLayer, EDataLayerRuntimeState::Loaded, false))
+						{
+							Flags |= EDataLayerUpdateFlags::FlushStreamingFull;
+							UE_LOG(LogMovieScene, Warning, TEXT("[UMovieSceneDataLayerSystem] Data layer with name '%s' is causing a full streaming flush (%s)"),
+								*DataLayer->GetDataLayerShortName(),
+								*StaticEnum<EDataLayerRuntimeState>()->GetDisplayNameTextByValue((int64)DesiredStateValue).ToString());
 							
-						CSV_EVENT_GLOBAL(TEXT("SeqDataLayerFlushFull-%s-%s"),
-							*DataLayer->GetDataLayerShortName(),
-							*StaticEnum<EDataLayerRuntimeState>()->GetDisplayNameTextByValue((int64)DesiredStateValue).ToString());
-					}
-					else
-					{
-						Flags |= EDataLayerUpdateFlags::FlushStreamingVisibility;
-						UE_LOG(LogMovieScene, Log, TEXT("[UMovieSceneDataLayerSystem] Data layer with name '%s' is causing a visibility streaming flush (%s)"),
-							*DataLayer->GetDataLayerShortName(),
-							*StaticEnum<EDataLayerRuntimeState>()->GetDisplayNameTextByValue((int64)DesiredStateValue).ToString());
+							CSV_EVENT_GLOBAL(TEXT("SeqDataLayerFlushFull-%s-%s"),
+								*DataLayer->GetDataLayerShortName(),
+								*StaticEnum<EDataLayerRuntimeState>()->GetDisplayNameTextByValue((int64)DesiredStateValue).ToString());
+						}
+						else
+						{
+							Flags |= EDataLayerUpdateFlags::FlushStreamingVisibility;
+							UE_LOG(LogMovieScene, Log, TEXT("[UMovieSceneDataLayerSystem] Data layer with name '%s' is causing a visibility streaming flush (%s)"),
+								*DataLayer->GetDataLayerShortName(),
+								*StaticEnum<EDataLayerRuntimeState>()->GetDisplayNameTextByValue((int64)DesiredStateValue).ToString());
 							
-						CSV_EVENT_GLOBAL(TEXT("SeqDataLayerFlushVis-%s-%s"),
-							*DataLayer->GetDataLayerShortName(),
-							*StaticEnum<EDataLayerRuntimeState>()->GetDisplayNameTextByValue((int64)DesiredStateValue).ToString());
+							CSV_EVENT_GLOBAL(TEXT("SeqDataLayerFlushVis-%s-%s"),
+								*DataLayer->GetDataLayerShortName(),
+								*StaticEnum<EDataLayerRuntimeState>()->GetDisplayNameTextByValue((int64)DesiredStateValue).ToString());
+						}
 					}
-				}
 
-				if (StateValue.ShouldPerformGarbageCollect(DesiredStateValue))
-				{
-					Flags |= EDataLayerUpdateFlags::PerformGarbageCollect;
+					if (StateValue.ShouldPerformGarbageCollect(DesiredStateValue))
+					{
+						Flags |= EDataLayerUpdateFlags::PerformGarbageCollect;
+					}
 				}
 			}
 		}
