Sequencer: Added time-warp capabilities to sequences, sub-sections and skeletal animation sections

diff --git a/Engine/Source/Runtime/MovieScene/Public/MovieSceneSequenceVisitor.h b/Engine/Source/Runtime/MovieScene/Public/MovieSceneSequenceVisitor.h
--- a/Engine/Source/Runtime/MovieScene/Public/MovieSceneSequenceVisitor.h
+++ b/Engine/Source/Runtime/MovieScene/Public/MovieSceneSequenceVisitor.h
@@ -69,6 +69,8 @@ struct FSubSequenceSpace
 
 	/** Transform from the root time-space to the current sequence's time-space */
 	FMovieSceneSequenceTransform RootToSequenceTransform;
+	FMovieSceneTransformBreadcrumbs StartBreadcrumbs;
+	FMovieSceneTransformBreadcrumbs EndBreadcrumbs;
 	/** The ID of the sequence being compiled */
 	FMovieSceneSequenceID SequenceID;
 	/** A range to clamp the visited sequence to in the root's time-space */
diff --git a/Engine/Source/Runtime/MovieScene/Public/MovieSceneTimeHelpers.h b/Engine/Source/Runtime/MovieScene/Public/MovieSceneTimeHelpers.h
--- a/Engine/Source/Runtime/MovieScene/Public/MovieSceneTimeHelpers.h
+++ b/Engine/Source/Runtime/MovieScene/Public/MovieSceneTimeHelpers.h
@@ -263,15 +263,127 @@ inline TRange<FFrameTime> ConvertToFrameTimeRange(const TRange<FFrameNumber>& In
 	TRange<FFrameTime> Result;
 	if (InRange.HasLowerBound())
 	{
-		Result.SetLowerBound(TRangeBound<FFrameTime>(InRange.GetLowerBoundValue()));
+		if (InRange.GetLowerBound().IsInclusive())
+		{
+			Result.SetLowerBound(TRangeBound<FFrameTime>::Inclusive(InRange.GetLowerBoundValue()));
+		}
+		else
+		{
+			Result.SetLowerBound(TRangeBound<FFrameTime>::Exclusive(InRange.GetLowerBoundValue()));
+		}
 	}
 	if (InRange.HasUpperBound())
 	{
-		Result.SetUpperBound(TRangeBound<FFrameTime>(InRange.GetUpperBoundValue()));
+		if (InRange.GetUpperBound().IsInclusive())
+		{
+			Result.SetUpperBound(TRangeBound<FFrameTime>::Inclusive(InRange.GetUpperBoundValue()));
+		}
+		else
+		{
+			Result.SetUpperBound(TRangeBound<FFrameTime>::Exclusive(InRange.GetUpperBoundValue()));
+		}
 	}
 	return Result;
 }
 
+inline TRange<FFrameTime> ConvertToFrameTimeRange(const TRange<float>& InRange)
+{
+	TRange<FFrameTime> Result;
+	if (InRange.HasLowerBound())
+	{
+		if (InRange.GetLowerBound().IsInclusive())
+		{
+			Result.SetLowerBound(TRangeBound<FFrameTime>::Inclusive(FFrameTime::FromDecimal(InRange.GetLowerBoundValue())));
+		}
+		else
+		{
+			Result.SetLowerBound(TRangeBound<FFrameTime>::Exclusive(FFrameTime::FromDecimal(InRange.GetLowerBoundValue())));
+		}
+	}
+	if (InRange.HasUpperBound())
+	{
+		if (InRange.GetUpperBound().IsInclusive())
+		{
+			Result.SetUpperBound(TRangeBound<FFrameTime>::Inclusive(FFrameTime::FromDecimal(InRange.GetUpperBoundValue())));
+		}
+		else
+		{
+			Result.SetUpperBound(TRangeBound<FFrameTime>::Exclusive(FFrameTime::FromDecimal(InRange.GetUpperBoundValue())));
+		}
+	}
+	return Result;
+}
+
+inline TRange<FFrameNumber> ConvertToDiscreteRange(const TRange<FFrameTime>& InRange)
+{
+	TRange<FFrameNumber> Result;
+	if (InRange.HasLowerBound())
+	{
+		// A frame time of (10.5 does not include _all_ of frame 10, so we can't include that frame
+		FFrameTime LowerBound = InRange.GetLowerBoundValue();
+		if (InRange.GetLowerBound().IsInclusive() && LowerBound.GetSubFrame() == 0.f)
+		{
+			Result.SetLowerBound(TRangeBound<FFrameNumber>::Inclusive(LowerBound.GetFrame()));
+		}
+		else
+		{
+			Result.SetLowerBound(TRangeBound<FFrameNumber>::Inclusive(LowerBound.GetFrame()+1));
+		}
+	}
+	if (InRange.HasUpperBound())
+	{
+		FFrameTime UpperBound = InRange.GetUpperBoundValue();
+		Result.SetUpperBound(TRangeBound<FFrameNumber>::Exclusive(UpperBound.GetFrame()));
+	}
+	return Result;
+}
+
+
+/**
+ * Converts a range from one type of bounds to another. The output bounds type must be implicitly
+ * constructable from the input bounds type.
+ */
+template<typename InBoundType, typename OutBoundType>
+inline TRange<OutBoundType> ConvertRange(const TRange<InBoundType>& Range)
+{
+	const TRangeBound<InBoundType> SourceLower = Range.GetLowerBound();
+	TRangeBound<OutBoundType> DestLower = SourceLower.IsOpen() ?
+		TRangeBound<OutBoundType>() :
+		SourceLower.IsInclusive() ?
+			TRangeBound<OutBoundType>::Inclusive(SourceLower.GetValue()) :
+			TRangeBound<OutBoundType>::Exclusive(SourceLower.GetValue());
+
+	const TRangeBound<InBoundType> SourceUpper = Range.GetUpperBound();
+	TRangeBound<OutBoundType> DestUpper = SourceUpper.IsOpen() ?
+		TRangeBound<OutBoundType>() :
+		SourceUpper.IsInclusive() ?
+			TRangeBound<OutBoundType>::Inclusive(SourceUpper.GetValue()) :
+			TRangeBound<OutBoundType>::Exclusive(SourceUpper.GetValue());
+
+	return TRange<OutBoundType>(DestLower, DestUpper);
+}
+
+// Specialization of ConvertRange for round down FFrameTime to FFrameNumber.
+template<>
+inline TRange<FFrameNumber> ConvertRange(const TRange<FFrameTime>& Range)
+{
+	const TRangeBound<FFrameTime> SourceLower = Range.GetLowerBound();
+	TRangeBound<FFrameNumber> DestLower = SourceLower.IsOpen() ?
+		TRangeBound<FFrameNumber>() :
+		SourceLower.IsInclusive() ?
+			TRangeBound<FFrameNumber>::Inclusive(SourceLower.GetValue().FloorToFrame()) :
+			TRangeBound<FFrameNumber>::Exclusive(SourceLower.GetValue().FloorToFrame());
+
+	const TRangeBound<FFrameTime> SourceUpper = Range.GetUpperBound();
+	TRangeBound<FFrameNumber> DestUpper = SourceUpper.IsOpen() ?
+		TRangeBound<FFrameNumber>() :
+		SourceUpper.IsInclusive() ?
+			TRangeBound<FFrameNumber>::Inclusive(SourceUpper.GetValue().FloorToFrame()) :
+			TRangeBound<FFrameNumber>::Exclusive(SourceUpper.GetValue().FloorToFrame());
+
+	return TRange<FFrameNumber>(DestLower, DestUpper);
+}
+
 } // namespace MovieScene
 } // namespace UE
 
diff --git a/Engine/Source/Runtime/MovieScene/Public/Sections/MovieSceneSubSection.h b/Engine/Source/Runtime/MovieScene/Public/Sections/MovieSceneSubSection.h
--- a/Engine/Source/Runtime/MovieScene/Public/Sections/MovieSceneSubSection.h
+++ b/Engine/Source/Runtime/MovieScene/Public/Sections/MovieSceneSubSection.h
@@ -100,6 +100,16 @@ public:
 	 */
 	MOVIESCENE_API FMovieSceneSequenceTransform OuterToInnerTransform() const;
 
+	/**
+	 * Gets the transform that converts time from this section's time-base to its inner sequence's
+	 */
+	MOVIESCENE_API FMovieSceneSequenceTransform OuterToInnerTransform_NoInnerTimeWarp() const;
+
+	/**
+	 * Gets the transform that converts time from this section's time-base to its inner sequence's
+	 */
+	MOVIESCENE_API void AppendInnerTimeWarpTransform(FMovieSceneSequenceTransform& OutTransform) const;
+
 	/**
 	 * Gets the playrange of the inner sequence, in the inner sequence's time space, trimmed with any start/end offsets,
 	 * and validated to make sure we get at least a 1-frame long playback range (e.g. in the case where excessive
@@ -143,6 +153,8 @@ public:
 		NetworkMask = (uint8)InNetworkMask;
 	}
 
+	MOVIESCENE_API void DeleteChannels(TArrayView<const FName> ChannelNames);
+
 public:
 
 	//~ UMovieSceneSection interface
@@ -159,6 +171,8 @@ protected:
 	MOVIESCENE_API virtual bool PopulateEvaluationFieldImpl(const TRange<FFrameNumber>& EffectiveRange, const FMovieSceneEvaluationFieldEntityMetaData& InMetaData, FMovieSceneEntityComponentFieldBuilder* OutFieldBuilder) override;
 	MOVIESCENE_API virtual void ImportEntityImpl(UMovieSceneEntitySystemLinker* EntityLinker, const FEntityImportParams& Params, FImportedEntity* OutImportedEntity) override;
 
+	MOVIESCENE_API virtual EMovieSceneChannelProxyType CacheChannelProxy();
+
 public:
 
 	UPROPERTY(config, BlueprintReadWrite, EditAnywhere, Category="General", meta=(ShowOnlyInnerProperties))
diff --git a/Engine/Source/Runtime/MovieSceneTracks/Private/Evaluation/MovieSceneEventTemplate.cpp b/Engine/Source/Runtime/MovieSceneTracks/Private/Evaluation/MovieSceneEventTemplate.cpp
--- a/Engine/Source/Runtime/MovieSceneTracks/Private/Evaluation/MovieSceneEventTemplate.cpp
+++ b/Engine/Source/Runtime/MovieSceneTracks/Private/Evaluation/MovieSceneEventTemplate.cpp
@@ -176,7 +176,12 @@ void FMovieSceneEventSectionTemplate::EvaluateSwept(const FMovieSceneEvaluationO
 	const int32 Last = bBackwards ? 0 : KeyTimes.Num() - 1;
 	const int32 Inc = bBackwards ? -1 : 1;
 
-	const float PositionInSeconds = Context.GetTime() * Context.GetRootToSequenceTransform().InverseNoLooping() / Context.GetFrameRate();
+
+	// @todo: this is technically incorrect because it is converting a root-space time to seconds using the local-space frame rate,
+	//        but this code only exists for legacy reasons so it has not been updated.
+	TOptional<FFrameTime> RootTime = Context.GetRootToSequenceTransform().Inverse().TryTransformTime(Context.GetTime());
+
+	const float PositionInSeconds = RootTime.Get(0) / Context.GetFrameRate();
 
 	if (bBackwards)
 	{
diff --git a/Engine/Source/Runtime/MovieSceneTracks/Private/Sections/MovieSceneEventRepeaterSection.cpp b/Engine/Source/Runtime/MovieSceneTracks/Private/Sections/MovieSceneEventRepeaterSection.cpp
--- a/Engine/Source/Runtime/MovieSceneTracks/Private/Sections/MovieSceneEventRepeaterSection.cpp
+++ b/Engine/Source/Runtime/MovieSceneTracks/Private/Sections/MovieSceneEventRepeaterSection.cpp
@@ -61,11 +61,17 @@ void UMovieSceneEventRepeaterSection::ImportEntityImpl(UMovieSceneEntitySystemLi
 		EventSystem = EntityLinker->LinkSystem<UMovieScenePostEvalEventSystem>();
 	}
 
+	TOptional<FFrameTime> RootTime = Context.GetSequenceToRootSequenceTransform().TryTransformTime(Context.GetTime());
+	if (!RootTime)
+	{
+		return;
+	}
+
 	FMovieSceneEventTriggerData TriggerData = {
 		Event.Ptrs,
 		Params.GetObjectBindingID(),
 		ThisInstance.GetSequenceID(),
-		Context.GetTime() * Context.GetSequenceToRootSequenceTransform()
+		RootTime.GetValue()
 	};
 
 	EventSystem->AddEvent(ThisInstance.GetRootInstanceHandle(), TriggerData);
diff --git a/Engine/Source/Runtime/MovieSceneTracks/Private/Sections/MovieSceneEventTriggerSection.cpp b/Engine/Source/Runtime/MovieSceneTracks/Private/Sections/MovieSceneEventTriggerSection.cpp
--- a/Engine/Source/Runtime/MovieSceneTracks/Private/Sections/MovieSceneEventTriggerSection.cpp
+++ b/Engine/Source/Runtime/MovieSceneTracks/Private/Sections/MovieSceneEventTriggerSection.cpp
@@ -80,11 +80,17 @@ void UMovieSceneEventTriggerSection::ImportEntityImpl(UMovieSceneEntitySystemLin
 		EventSystem = EntityLinker->LinkSystem<UMovieScenePostEvalEventSystem>();
 	}
 
+	TOptional<FFrameTime> RootTime = Context.GetSequenceToRootSequenceTransform().TryTransformTime(Times[EventIndex]);
+	if (!RootTime)
+	{
+		return;
+	}
+
 	FMovieSceneEventTriggerData TriggerData = {
 		Events[EventIndex].Ptrs,
 		Params.GetObjectBindingID(),
 		ThisInstance.GetSequenceID(),
-		Times[EventIndex] * Context.GetSequenceToRootSequenceTransform()
+		RootTime.GetValue()
 	};
 
 	EventSystem->AddEvent(ThisInstance.GetRootInstanceHandle(), TriggerData);
diff --git a/Engine/Source/Runtime/MovieSceneTracks/Private/Sections/MovieSceneSkeletalAnimationSection.cpp b/Engine/Source/Runtime/MovieSceneTracks/Private/Sections/MovieSceneSkeletalAnimationSection.cpp
--- a/Engine/Source/Runtime/MovieSceneTracks/Private/Sections/MovieSceneSkeletalAnimationSection.cpp
+++ b/Engine/Source/Runtime/MovieSceneTracks/Private/Sections/MovieSceneSkeletalAnimationSection.cpp
@@ -10,12 +10,15 @@
 #include "MovieSceneTimeHelpers.h"
 #include "Tracks/MovieSceneSkeletalAnimationTrack.h"
 #include "BoneContainer.h"
+#include "MovieSceneTransformTypes.h"
 #include "Animation/AnimationPoseData.h"
 #include "Animation/AnimSequenceDecompressionContext.h"
 #include "Animation/AttributesRuntime.h"
 #include "Misc/FrameRate.h"
 #include "EntitySystem/BuiltInComponentTypes.h"
 #include "MovieSceneTracksComponentTypes.h"
+#include "Sections/MovieSceneSectionTimingParameters.h"
+#include "Variants/MovieSceneTimeWarpGetter.h"
 
 #include UE_INLINE_GENERATED_CPP_BY_NAME(MovieSceneSkeletalAnimationSection)
 
@@ -46,6 +49,32 @@ FMovieSceneSkeletalAnimationParams::FMovieSceneSkeletalAnimationParams()
 	SwapRootBone = ESwapRootBone::SwapRootBone_None;
 }
 
+FMovieSceneSequenceTransform FMovieSceneSkeletalAnimationParams::MakeTransform(const FFrameRate& OuterFrameRate, const TRange<FFrameNumber>& OuterRange) const
+{
+	if (!Animation)
+	{
+		return FMovieSceneSequenceTransform();
+	}
+
+	const double SequenceLength = GetSequenceLength();
+	const FFrameTime AnimationLength = SequenceLength * OuterFrameRate;
+	const int32 LengthInFrames = AnimationLength.FrameNumber.Value + (int)(AnimationLength.GetSubFrame() + 0.5f) + 1;
+
+	const bool bLooping = (UE::MovieScene::DiscreteSize(OuterRange) + StartFrameOffset + EndFrameOffset) > LengthInFrames;
+
+	FMovieSceneSectionTimingParametersSeconds TimingParams = {
+		PlayRate.ShallowCopy(),
+		StartFrameOffset / OuterFrameRate,
+		EndFrameOffset / OuterFrameRate,
+		FirstLoopStartFrameOffset / OuterFrameRate,
+		bLooping,
+		!bLooping,
+		bReverse
+	};
+
+	return TimingParams.MakeTransform(OuterFrameRate, OuterRange, SequenceLength, Animation->RateScale);
+}
+
 double FMovieSceneSkeletalAnimationParams::MapTimeToAnimation(const UMovieSceneSection* InSection, FFrameTime InPosition, FFrameRate InFrameRate) const
 {
 	const FFrameNumber SectionStartTime = InSection->GetInclusiveStartFrame();
@@ -55,47 +84,7 @@ double FMovieSceneSkeletalAnimationParams::MapTimeToAnimation(const UMovieSceneS
 
 double FMovieSceneSkeletalAnimationParams::MapTimeToAnimation(FFrameNumber InSectionStartTime, FFrameNumber InSectionEndTime, FFrameTime InPosition, FFrameRate InFrameRate) const
 {
-	// Get Animation Length and frame time
-	if (Animation)
-	{
-		const FFrameTime AnimationLength = GetSequenceLength() * InFrameRate;
-		const int32 LengthInFrames = AnimationLength.FrameNumber.Value + (int)(AnimationLength.GetSubFrame() + 0.5f) + 1;
-
-		// we only play end if we are not looping, and assuming we are looping if Length is greater than default length;
-		const bool bLooping = (InSectionEndTime.Value - InSectionStartTime.Value + StartFrameOffset + EndFrameOffset) > LengthInFrames;
-
-		// Make sure InPosition FrameTime doesn't underflow InSectionStartTime or overflow InSectionEndTime
-		InPosition = FMath::Clamp(InPosition, FFrameTime(InSectionStartTime), FFrameTime(InSectionEndTime - 1));
-
-		// Gather helper values
-		const float SectionPlayRate = PlayRate * Animation->RateScale;
-		const float AnimPlayRate = FMath::IsNearlyZero(SectionPlayRate) ? 1.0f : SectionPlayRate;
-		const double SeqLength = GetSequenceLength() - InFrameRate.AsSeconds(StartFrameOffset + EndFrameOffset);
-
-		// The Time from the beginning of InSectionStartTime to InPosition in seconds
-		double SecondsFromSectionStart = FFrameTime::FromDecimal((InPosition - InSectionStartTime).AsDecimal() * AnimPlayRate) / InFrameRate;
-
-		// Logic for reversed animation
-		if (bReverse)
-		{
-			// Duration of this section 
-			double SectionDuration = (((InSectionEndTime - InSectionStartTime) * AnimPlayRate) / InFrameRate);
-			SecondsFromSectionStart = SectionDuration - SecondsFromSectionStart;
-		}
-
-		SecondsFromSectionStart += InFrameRate.AsSeconds(FirstLoopStartFrameOffset);
-
-		// Make sure Seconds is in range
-		if (SeqLength > 0.0 && (bLooping || !FMath::IsNearlyEqual(SecondsFromSectionStart, SeqLength, 1e-4)))
-		{
-			SecondsFromSectionStart = FMath::Fmod(SecondsFromSectionStart, SeqLength);
-		}
-
-		// Add the StartFrameOffset to the current seconds in the section to get the right animation frame
-		SecondsFromSectionStart += InFrameRate.AsSeconds(StartFrameOffset);
-		return SecondsFromSectionStart;
-	}
-	return 0.0;
+	return MakeTransform(InFrameRate, TRange<FFrameNumber>(InSectionStartTime, InSectionEndTime)).TransformTime(InPosition).AsDecimal();
 }
 
 UMovieSceneSkeletalAnimationSection::UMovieSceneSkeletalAnimationSection( const FObjectInitializer& ObjectInitializer )
@@ -130,20 +119,59 @@ UMovieSceneSkeletalAnimationSection::UMovieSceneSkeletalAnimationSection( const
 	bMatchRotationRoll = false;
 	bMatchRotationPitch = false;
 	bMatchIncludeZHeight = false;
+	bDebugForceTickPose = false;
+}
 
-#if WITH_EDITOR
+EMovieSceneChannelProxyType UMovieSceneSkeletalAnimationSection::CacheChannelProxy()
+{
+	FMovieSceneChannelProxyData Channels;
+
+	if (Params.PlayRate.GetType() == EMovieSceneTimeWarpType::Custom)
+	{
+		UMovieSceneTimeWarpGetter* Custom = Params.PlayRate.AsCustom();
+		if (Custom)
+		{
+			Custom->PopulateChannelProxy(Channels, UMovieSceneTimeWarpGetter::EAllowTopLevelChannels::No);
+		}
+	}
 
-	PreviousPlayRate = Params.PlayRate;
+
+#if WITH_EDITOR
 
 	static FMovieSceneChannelMetaData MetaData("Weight", LOCTEXT("WeightChannelName", "Weight"));
 	MetaData.bCanCollapseToTrack = false;
-	ChannelProxy = MakeShared<FMovieSceneChannelProxy>(Params.Weight, MetaData, TMovieSceneExternalValue<float>());
+
+	Channels.Add(Params.Weight, MetaData, TMovieSceneExternalValue<float>());
 
 #else
 
-	ChannelProxy = MakeShared<FMovieSceneChannelProxy>(Params.Weight);
+	Channels.Add(Params.Weight);
 
 #endif
+
+	ChannelProxy = MakeShared<FMovieSceneChannelProxy>(MoveTemp(Channels));
+	return EMovieSceneChannelProxyType::Dynamic;
+}
+
+void UMovieSceneSkeletalAnimationSection::DeleteChannels(TArrayView<const FName> ChannelNames)
+{
+	bool bDeletedAny = false;
+
+	if (Params.PlayRate.GetType() == EMovieSceneTimeWarpType::Custom && TryModify())
+	{
+		if (UMovieSceneTimeWarpGetter* Getter = Params.PlayRate.AsCustom())
+		{
+			for (FName ChannelName : ChannelNames)
+			{
+				bDeletedAny |= Getter->DeleteChannel(Params.PlayRate, ChannelName);
+			}
+		}
+	}
+
+	if (bDeletedAny)
+	{
+		ChannelProxy = nullptr;
+	}
 }
 
 TOptional<FFrameTime> UMovieSceneSkeletalAnimationSection::GetOffsetTime() const
@@ -274,14 +302,23 @@ void UMovieSceneSkeletalAnimationSection::PostLoad()
 
 TOptional<TRange<FFrameNumber> > UMovieSceneSkeletalAnimationSection::GetAutoSizeRange() const
 {
-	const FFrameRate FrameRate = GetTypedOuter<UMovieScene>()->GetTickResolution();
-	const float AnimPlayRate = FMath::IsNearlyZero(Params.PlayRate) || Params.Animation == nullptr ? 1.0f : Params.PlayRate * Params.Animation->RateScale;
+	if (UMovieScene* MovieScene = GetTypedOuter<UMovieScene>())
+	{
+		FFrameRate FrameRate = MovieScene->GetTickResolution();
+		FMovieSceneInverseSequenceTransform InnerToOuterTransform = Params.MakeTransform(FrameRate, GetRange()).Inverse();
+
+		FFrameTime InnerStartTime = FFrameTime::FromDecimal((Params.StartFrameOffset + Params.FirstLoopStartFrameOffset) / FrameRate);
+		FFrameTime InnerEndTime   = FFrameTime::FromDecimal(Params.GetSequenceLength() - (Params.EndFrameOffset / FrameRate));
 
-	const FFrameTime UnscaledAnimationLength = FMath::Max(Params.GetSequenceLength() * FrameRate - Params.FirstLoopStartFrameOffset - Params.StartFrameOffset - Params.EndFrameOffset, FFrameTime(1));
-	const FFrameTime AnimationLength = UnscaledAnimationLength / AnimPlayRate;
-	const int32 IFrameNumber = AnimationLength.FrameNumber.Value + (int)(AnimationLength.GetSubFrame() + 0.5f);
+		TOptional<FFrameTime> OuterStartTime = InnerToOuterTransform.TryTransformTime(InnerStartTime);
+		TOptional<FFrameTime> OuterEndTime   = InnerToOuterTransform.TryTransformTime(InnerEndTime);
 
-	return TRange<FFrameNumber>(GetInclusiveStartFrame(), GetInclusiveStartFrame() + IFrameNumber + 1);
+		if (OuterStartTime && OuterEndTime)
+		{
+			return TRange<FFrameNumber>(OuterStartTime.GetValue().FrameNumber, OuterEndTime.GetValue().FrameNumber + 1);
+		}
+	}
+	return TOptional<TRange<FFrameNumber>>();
 }
 
 
@@ -294,8 +331,10 @@ void UMovieSceneSkeletalAnimationSection::TrimSection(FQualifiedFrameTime TrimTi
 		if (bTrimLeft)
 		{
 			FFrameRate FrameRate = GetTypedOuter<UMovieScene>()->GetTickResolution();
+			FMovieSceneSequenceTransform OuterToInnerTransform = Params.MakeTransform(FrameRate, GetRange());
 
-			Params.FirstLoopStartFrameOffset = HasStartFrame() ? GetFirstLoopStartOffsetAtTrimTime(TrimTime, Params, GetInclusiveStartFrame(), FrameRate) : 0;
+			FFrameTime AnimationTimeInSeconds = OuterToInnerTransform.TransformTime(TrimTime.Time);
+			Params.FirstLoopStartFrameOffset = (AnimationTimeInSeconds.AsDecimal() * FrameRate).FrameNumber;
 		}
 
 		Super::TrimSection(TrimTime, bTrimLeft, bDeleteKeys);
@@ -334,8 +373,10 @@ UMovieSceneSection* UMovieSceneSkeletalAnimationSection::SplitSection(FQualified
 	const FFrameNumber InitialFirstLoopStartFrameOffset = Params.FirstLoopStartFrameOffset;
 
 	FFrameRate FrameRate = GetTypedOuter<UMovieScene>()->GetTickResolution();
+	FMovieSceneSequenceTransform OuterToInnerTransform = Params.MakeTransform(FrameRate, GetRange());
 
-	const FFrameNumber NewOffset = HasStartFrame() ? GetFirstLoopStartOffsetAtTrimTime(SplitTime, Params, GetInclusiveStartFrame(), FrameRate) : 0;
+	FFrameTime AnimationTimeInSeconds = OuterToInnerTransform.TransformTime(SplitTime.Time);
+	const FFrameNumber NewOffset = (AnimationTimeInSeconds.AsDecimal() * FrameRate).FrameNumber;
 
 	UMovieSceneSkeletalAnimationSection* NewSection = Cast<UMovieSceneSkeletalAnimationSection>(Super::SplitSection(SplitTime, bDeleteKeys));
 	if (NewSection != nullptr)
@@ -367,35 +408,36 @@ UMovieSceneSection* UMovieSceneSkeletalAnimationSection::SplitSection(FQualified
 
 void UMovieSceneSkeletalAnimationSection::GetSnapTimes(TArray<FFrameNumber>& OutSnapTimes, bool bGetSectionBorders) const
 {
+	using namespace UE::MovieScene;
+
 	Super::GetSnapTimes(OutSnapTimes, bGetSectionBorders);
 
-	const FFrameRate   FrameRate  = GetTypedOuter<UMovieScene>()->GetTickResolution();
+	const UMovieScene* MovieScene = GetTypedOuter<UMovieScene>();
 	const FFrameNumber StartFrame = GetInclusiveStartFrame();
-	const FFrameNumber EndFrame   = GetExclusiveEndFrame() - 1; // -1 because we don't need to add the end frame twice
+	const FFrameNumber EndFrame   = GetExclusiveEndFrame();
 
-	const float AnimPlayRate     = FMath::IsNearlyZero(Params.PlayRate) || Params.Animation == nullptr ? 1.0f : Params.PlayRate * Params.Animation->RateScale;
-	const float SeqLengthSeconds = Params.GetSequenceLength() - FrameRate.AsSeconds(Params.StartFrameOffset + Params.EndFrameOffset) / AnimPlayRate;
-	const float FirstLoopSeqLengthSeconds = SeqLengthSeconds - FrameRate.AsSeconds(Params.FirstLoopStartFrameOffset) / AnimPlayRate;
+	if (!MovieScene)
+	{
+		return;
+	}
 
-	const FFrameTime SequenceFrameLength = SeqLengthSeconds * FrameRate;
-	const FFrameTime FirstLoopSequenceFrameLength = FirstLoopSeqLengthSeconds * FrameRate;
-	if (SequenceFrameLength.FrameNumber > 1)
+	auto VisitBoundary = [&OutSnapTimes](FFrameTime InTime)
 	{
-		// Snap to the repeat times
-		bool IsFirstLoop = true;
-		FFrameTime CurrentTime = StartFrame;
-		while (CurrentTime < EndFrame)
+		OutSnapTimes.Add(InTime.RoundToFrame());
+		return true;
+	};
+
+	FMovieSceneSequenceTransform OuterToInnerTransform = Params.MakeTransform(MovieScene->GetTickResolution(), GetRange());
+
+	if (!OuterToInnerTransform.ExtractBoundariesWithinRange(StartFrame, EndFrame, VisitBoundary))
+	{
+		FMovieSceneInverseSequenceTransform InnerToOuterTransform = OuterToInnerTransform.Inverse();
+
+		TOptional<FFrameTime> AnimEnd   = InnerToOuterTransform.TryTransformTime(FFrameTime::FromDecimal(Params.GetSequenceLength()));
+
+		if (AnimEnd && AnimEnd.GetValue() < EndFrame)
 		{
-			OutSnapTimes.Add(CurrentTime.FrameNumber);
-			if (IsFirstLoop)
-			{
-				CurrentTime += FirstLoopSequenceFrameLength;
-				IsFirstLoop = false;
-			}
-			else
-			{
-				CurrentTime += SequenceFrameLength;
-			}
+			VisitBoundary(AnimEnd.GetValue());
 		}
 	}
 }
@@ -469,7 +511,10 @@ bool UMovieSceneSkeletalAnimationSection::Modify(bool bAlwaysMarkDirty)
 void UMovieSceneSkeletalAnimationSection::PreEditChange(FProperty* PropertyAboutToChange)
 {
 	// Store the current play rate so that we can compute the amount to compensate the section end time when the play rate changes
-	PreviousPlayRate = Params.PlayRate;
+	if (Params.PlayRate.GetType() == EMovieSceneTimeWarpType::FixedPlayRate)
+	{
+		PreviousPlayRate = Params.PlayRate.AsFixedPlayRateFloat();
+	}
 
 	Super::PreEditChange(PropertyAboutToChange);
 }
@@ -478,9 +523,10 @@ void UMovieSceneSkeletalAnimationSection::PostEditChangeProperty(FPropertyChange
 {
 	// Adjust the duration automatically if the play rate changes
 	if (PropertyChangedEvent.Property != nullptr &&
-		PropertyChangedEvent.Property->GetFName() == TEXT("PlayRate"))
+		PropertyChangedEvent.Property->GetFName() == TEXT("PlayRate") && 
+		Params.PlayRate.GetType() == EMovieSceneTimeWarpType::FixedPlayRate)
 	{
-		float NewPlayRate = Params.PlayRate;
+		float NewPlayRate = Params.PlayRate.AsFixedPlayRateFloat();
 
 		if (!FMath::IsNearlyZero(NewPlayRate))
 		{
@@ -490,6 +536,8 @@ void UMovieSceneSkeletalAnimationSection::PostEditChangeProperty(FPropertyChange
 
 			PreviousPlayRate = NewPlayRate;
 		}
+
+		ChannelProxy = nullptr;
 	}
 	if (GetRootMotionParams())
 	{
diff --git a/Engine/Source/Runtime/MovieSceneTracks/Private/TrackInstances/MovieSceneCameraCutTrackInstance.cpp b/Engine/Source/Runtime/MovieSceneTracks/Private/TrackInstances/MovieSceneCameraCutTrackInstance.cpp
--- a/Engine/Source/Runtime/MovieSceneTracks/Private/TrackInstances/MovieSceneCameraCutTrackInstance.cpp
+++ b/Engine/Source/Runtime/MovieSceneTracks/Private/TrackInstances/MovieSceneCameraCutTrackInstance.cpp
@@ -395,6 +395,7 @@ void UMovieSceneCameraCutTrackInstance::OnAnimate()
 	{
 		const FMovieSceneTrackInstanceInput& Input = InputInfo.Input;
 		const FSequenceInstance& SequenceInstance = InstanceRegistry->GetInstance(Input.InstanceHandle);
+		const FSequenceInstance& RootSequenceInstance = InstanceRegistry->GetInstance(SequenceInstance.GetRootInstanceHandle());
 		const FMovieSceneContext& Context = SequenceInstance.GetContext();
 
 		const UMovieSceneCameraCutSection* Section = Cast<const UMovieSceneCameraCutSection>(Input.Section);
@@ -420,7 +421,7 @@ void UMovieSceneCameraCutTrackInstance::OnAnimate()
 		else
 		{
 			const UMovieSceneCameraCutTrack* Track = Section->GetTypedOuter<UMovieSceneCameraCutTrack>();
-			const FMovieSceneSequenceTransform SequenceToRootTransform = Context.GetSequenceToRootSequenceTransform();
+			const FMovieSceneInverseSequenceTransform SequenceToRootTransform = Context.GetSequenceToRootSequenceTransform();
 
 			FBlendedCameraCut Params(Input, CameraBindingID, SequenceInstance.GetSequenceID());
 			Params.bCanBlend = Track->bCanBlend;
@@ -437,14 +438,29 @@ void UMovieSceneCameraCutTrackInstance::OnAnimate()
 			if (Section->HasStartFrame() && Section->Easing.GetEaseInDuration() > 0)
 			{
 				Params.LocalEaseInEndTime = Params.LocalStartTime + Section->Easing.GetEaseInDuration();
-				const float RootEaseInTime = SequenceToRootTransform.GetTimeScale() * Context.GetFrameRate().AsSeconds(FFrameNumber(Section->Easing.GetEaseInDuration()));
-				Params.EaseIn = FBlendedCameraCutEasingInfo(RootEaseInTime, Section->Easing.EaseIn);
+
+				TOptional<FFrameTime> RootStart = SequenceToRootTransform.TryTransformTime(Params.LocalStartTime);
+				TOptional<FFrameTime> RootEnd   = SequenceToRootTransform.TryTransformTime(Params.LocalEaseInEndTime);
+
+				if (RootStart && RootEnd)
+				{
+					const float RootEaseInTime = RootSequenceInstance.GetContext().GetFrameRate().AsSeconds(RootEnd.GetValue() - RootStart.GetValue());
+					Params.EaseIn = FBlendedCameraCutEasingInfo(RootEaseInTime, Section->Easing.EaseIn);
+				}
 			}
 			if (Section->HasEndFrame() && Section->Easing.GetEaseOutDuration() > 0)
 			{
 				Params.LocalEaseOutStartTime = Params.LocalEndTime - Section->Easing.GetEaseOutDuration();
-				const float RootEaseOutTime = SequenceToRootTransform.GetTimeScale() * Context.GetFrameRate().AsSeconds(FFrameNumber(Section->Easing.GetEaseOutDuration()));
-				Params.EaseOut = FBlendedCameraCutEasingInfo(RootEaseOutTime, Section->Easing.EaseOut);
+
+
+				TOptional<FFrameTime> RootStart = SequenceToRootTransform.TryTransformTime(Params.LocalEaseOutStartTime);
+				TOptional<FFrameTime> RootEnd   = SequenceToRootTransform.TryTransformTime(Params.LocalEndTime);
+
+				if (RootStart && RootEnd)
+				{
+					const float RootEaseOutTime = RootSequenceInstance.GetContext().GetFrameRate().AsSeconds(RootEnd.GetValue() - RootStart.GetValue());
+					Params.EaseOut = FBlendedCameraCutEasingInfo(RootEaseOutTime, Section->Easing.EaseOut);
+				}
 			}
 
 			// Remember locking option.
diff --git a/Engine/Source/Runtime/MovieSceneTracks/Public/Sections/MovieSceneSkeletalAnimationSection.h b/Engine/Source/Runtime/MovieSceneTracks/Public/Sections/MovieSceneSkeletalAnimationSection.h
--- a/Engine/Source/Runtime/MovieSceneTracks/Public/Sections/MovieSceneSkeletalAnimationSection.h
+++ b/Engine/Source/Runtime/MovieSceneTracks/Public/Sections/MovieSceneSkeletalAnimationSection.h
@@ -23,7 +23,8 @@ struct FMovieSceneSkeletalAnimationParams
 	FMovieSceneSkeletalAnimationParams();
 
 	/** Gets the animation duration, modified by play rate */
-	float GetDuration() const { return FMath::IsNearlyZero(PlayRate) || Animation == nullptr ? 0.f : Animation->GetPlayLength() / PlayRate; }
+	UE_DEPRECATED(5.5, "Animation lengh no longer has a single, consistent length if there is timewarp.")
+	float GetDuration() const { return 0.f; }
 
 	/** Gets the animation sequence length, not modified by play rate */
 	float GetSequenceLength() const { return Animation != nullptr ? Animation->GetPlayLength() : 0.f; }
@@ -39,6 +40,11 @@ struct FMovieSceneSkeletalAnimationParams
 	 */
 	double MapTimeToAnimation(FFrameNumber InSectionStartTime, FFrameNumber InSectionEndTime, FFrameTime InPosition, FFrameRate InFrameRate) const;
 
+	/**
+	 * Make a transform structure from these animation parameters
+	 */
+	MOVIESCENETRACKS_API FMovieSceneSequenceTransform MakeTransform(const FFrameRate& OuterFrameRate, const TRange<FFrameNumber>& OuterRange) const;
+
 	/** The animation this section plays */
 	UPROPERTY(BlueprintReadWrite, EditAnywhere, Category="Animation", meta=(AllowedClasses = "/Script/Engine.AnimSequence,/Script/Engine.AnimComposite,/Script/Engine.AnimStreamable"))
 	TObjectPtr<UAnimSequenceBase> Animation;
@@ -57,7 +63,7 @@ struct FMovieSceneSkeletalAnimationParams
 
 	/** The playback rate of the animation clip */
 	UPROPERTY(BlueprintReadWrite, EditAnywhere, Category="Animation")
-	float PlayRate;
+	FMovieSceneTimeWarpVariant PlayRate;
 
 	/** Reverse the playback of the animation clip */
 	UPROPERTY(BlueprintReadWrite, EditAnywhere, Category="Animation")
@@ -113,10 +119,13 @@ public:
 	MOVIESCENETRACKS_API double MapTimeToAnimation(FFrameTime InPosition, FFrameRate InFrameRate) const;
 	
 	//~ UMovieSceneSection interface
+	virtual EMovieSceneChannelProxyType CacheChannelProxy() override;
 	virtual void SetRange(const TRange<FFrameNumber>& NewRange) override;
 	virtual void SetStartFrame(TRangeBound<FFrameNumber> NewStartFrame) override;
 	virtual void SetEndFrame(TRangeBound<FFrameNumber> NewEndFrame)override;
 
+	MOVIESCENETRACKS_API void DeleteChannels(TArrayView<const FName> ChannelNames);
+
 	struct FRootMotionParams
 	{
 		bool bBlendFirstChildOfRoot;
@@ -191,9 +200,6 @@ public:
 
 	FRotator StartRotationOffset;
 
-	UPROPERTY()
-	bool bMatchWithPrevious;
-
 	UPROPERTY()
 	FName MatchedBoneName;
 
@@ -206,24 +212,30 @@ public:
 	FRotator MatchedRotationOffset;
 
 	UPROPERTY()
-	bool bMatchTranslation;
+	uint8 bMatchWithPrevious : 1;
+
+	UPROPERTY()
+	uint8 bMatchTranslation : 1;
+
+	UPROPERTY()
+	uint8 bMatchIncludeZHeight : 1;
 
 	UPROPERTY()
-	bool bMatchIncludeZHeight;
+	uint8 bMatchRotationYaw : 1;
 
 	UPROPERTY()
-	bool bMatchRotationYaw;
+	uint8 bMatchRotationPitch : 1;
 
 	UPROPERTY()
-	bool bMatchRotationPitch;
+	uint8 bMatchRotationRoll : 1;
 
 	UPROPERTY()
-	bool bMatchRotationRoll;
+	uint8 bDebugForceTickPose : 1;
 
 #if WITH_EDITORONLY_DATA
 	/** Whether to show the underlying skeleton for this section. */
 	UPROPERTY(BlueprintReadWrite, EditAnywhere, Category = "Root Motions")
-	bool bShowSkeleton;
+	uint8 bShowSkeleton : 1;
 #endif
 	//Previous transform used to specify the global OffsetTransform while calculting the root motions.
 	FTransform PreviousTransform;
