Sequencer DataLayer sections - Removed the "Force GC on unload" option
* Added as a band-aid option for a project which ended up not using it
* Was not handled correctly on the server vs the clients
* Global GLevelStreamingContinuouslyIncrementalGCWhileLevelsPendingPurgeOverride option was never reset, which would cause the GC to run continously for the rest of the process execution

diff --git a/Engine/Source/Editor/UnrealEd/Private/WorldPartition/WorldPartitionBuilder.cpp b/Engine/Source/Editor/UnrealEd/Private/WorldPartition/WorldPartitionBuilder.cpp
--- a/Engine/Source/Editor/UnrealEd/Private/WorldPartition/WorldPartitionBuilder.cpp
+++ b/Engine/Source/Editor/UnrealEd/Private/WorldPartition/WorldPartitionBuilder.cpp
@@ -172,7 +172,6 @@ FWorldBuilderCellCoord FCellInfo::GetCellCount(const FBox& InBounds, const int32
 bool UWorldPartitionBuilder::Run(UWorld* World, FPackageSourceControlHelper& PackageHelper)
 {
 	// As we manage GC frequency ourselves during builds, make sure to turn off automated GC based on other settings
-	TGuardValue<int32> Guard_GLevelStreamingContinuouslyIncrementalGCWhileLevelsPendingPurgeOverride(GLevelStreamingContinuouslyIncrementalGCWhileLevelsPendingPurgeOverride, 0);
 	TGuardValue<int32> Guard_GLevelStreamingContinuouslyIncrementalGCWhileLevelsPendingPurge(GLevelStreamingContinuouslyIncrementalGCWhileLevelsPendingPurge, 0);
 	TGuardValue<int32> Guard_GLevelStreamingLowMemoryPendingPurgeCount(GLevelStreamingLowMemoryPendingPurgeCount, MAX_int32);
 	TGuardValue<int32> Guard_GLevelStreamingForceGCAfterLevelStreamedOut(GLevelStreamingForceGCAfterLevelStreamedOut, 0);
diff --git a/Engine/Source/Runtime/Engine/Classes/Engine/CoreSettings.h b/Engine/Source/Runtime/Engine/Classes/Engine/CoreSettings.h
--- a/Engine/Source/Runtime/Engine/Classes/Engine/CoreSettings.h
+++ b/Engine/Source/Runtime/Engine/Classes/Engine/CoreSettings.h
@@ -147,8 +147,6 @@ extern ENGINE_API float GLevelStreamingUnregisterComponentsTimeLimit;
 extern ENGINE_API int32 GLevelStreamingForceGCAfterLevelStreamedOut;
 /** Whether to kick off incremental GC when there are over the specified amount of levels still waiting to be purged. */
 extern ENGINE_API int32 GLevelStreamingContinuouslyIncrementalGCWhileLevelsPendingPurge;
-/** Whether to override GLevelStreamingContinuouslyIncrementalGCWhileLevelsPendingPurge temporarily until at least one level needs to be purged.  */
-extern ENGINE_API int32 GLevelStreamingContinuouslyIncrementalGCWhileLevelsPendingPurgeOverride;
 /** Enables level streaming requests while async loading (of anything) while the match is already in progress and no loading screen is up. */
 extern ENGINE_API int32 GLevelStreamingAllowLevelRequestsWhileAsyncLoadingInMatch;
 /** When we're already loading this many levels and actively in match, don't allow any more requests until one of those completes.  Set to zero to disable. */
diff --git a/Engine/Source/Runtime/Engine/Private/CoreSettings.cpp b/Engine/Source/Runtime/Engine/Private/CoreSettings.cpp
--- a/Engine/Source/Runtime/Engine/Private/CoreSettings.cpp
+++ b/Engine/Source/Runtime/Engine/Private/CoreSettings.cpp
@@ -22,7 +22,6 @@ int32 GLevelStreamingComponentsUnregistrationGranularity = 5;
 int32 GLevelStreamingRouteActorInitializationGranularity = 10;
 int32 GLevelStreamingForceGCAfterLevelStreamedOut = 1;
 int32 GLevelStreamingContinuouslyIncrementalGCWhileLevelsPendingPurge = 1;
-int32 GLevelStreamingContinuouslyIncrementalGCWhileLevelsPendingPurgeOverride = 0;
 int32 GLevelStreamingAllowLevelRequestsWhileAsyncLoadingInMatch = 1;
 int32 GLevelStreamingMaxLevelRequestsAtOnceWhileInMatch = 0;
 int32 GLevelStreamingForceVerifyLevelsGotRemovedByGC = 0;
diff --git a/Engine/Source/Runtime/Engine/Private/World.cpp b/Engine/Source/Runtime/Engine/Private/World.cpp
--- a/Engine/Source/Runtime/Engine/Private/World.cpp
+++ b/Engine/Source/Runtime/Engine/Private/World.cpp
@@ -4446,8 +4446,7 @@ void UWorld::UpdateLevelStreaming()
 	}
 
 	const int32 CurrentNumLevelsPendingPurge = FLevelStreamingGCHelper::GetNumLevelsPendingPurge();
-	const int32 LevelStreamingContinuouslyIncrementalGCWhileLevelsPendingPurge = GLevelStreamingContinuouslyIncrementalGCWhileLevelsPendingPurgeOverride ? 1 : GLevelStreamingContinuouslyIncrementalGCWhileLevelsPendingPurge;
-	const bool bShouldPurgeLevels = LevelStreamingContinuouslyIncrementalGCWhileLevelsPendingPurge && CurrentNumLevelsPendingPurge >= LevelStreamingContinuouslyIncrementalGCWhileLevelsPendingPurge;
+	const bool bShouldPurgeLevels = GLevelStreamingContinuouslyIncrementalGCWhileLevelsPendingPurge && CurrentNumLevelsPendingPurge >= GLevelStreamingContinuouslyIncrementalGCWhileLevelsPendingPurge;
 
 	// Are we currently in a low memory situation and number of pending levels to purge meets or exceeds our threshold?
 	const bool bShouldDoLowMemoryGC = GIsLowMemory && CurrentNumLevelsPendingPurge >= GLevelStreamingLowMemoryPendingPurgeCount;
diff --git a/Engine/Source/Runtime/MovieSceneTracks/Private/Sections/MovieSceneDataLayerSection.cpp b/Engine/Source/Runtime/MovieSceneTracks/Private/Sections/MovieSceneDataLayerSection.cpp
--- a/Engine/Source/Runtime/MovieSceneTracks/Private/Sections/MovieSceneDataLayerSection.cpp
+++ b/Engine/Source/Runtime/MovieSceneTracks/Private/Sections/MovieSceneDataLayerSection.cpp
@@ -13,7 +13,6 @@ UMovieSceneDataLayerSection::UMovieSceneDataLayerSection(const FObjectInitialize
 	PrerollState = EDataLayerRuntimeState::Activated;
 	bFlushOnActivated = true;
 	bFlushOnUnload = false;
-	bPerformGCOnUnload = false;
 	EvalOptions.EnableAndSetCompletionMode(EMovieSceneCompletionMode::RestoreState);
 }
 
@@ -57,16 +56,6 @@ void UMovieSceneDataLayerSection::SetFlushOnUnload(bool bInFlushOnUnload)
 	bFlushOnUnload = bInFlushOnUnload;
 }
 
-bool UMovieSceneDataLayerSection::GetPerformGCOnUnload() const
-{
-	return bPerformGCOnUnload;
-}
-
-void UMovieSceneDataLayerSection::SetPerformGCOnUnload(bool bInPerformGCOnUnload)
-{
-	bPerformGCOnUnload = bInPerformGCOnUnload;
-}
-
 void UMovieSceneDataLayerSection::ImportEntityImpl(UMovieSceneEntitySystemLinker* EntityLinker, const FEntityImportParams& Params, FImportedEntity* OutImportedEntity)
 {
 	using namespace UE::MovieScene;
diff --git a/Engine/Source/Runtime/MovieSceneTracks/Private/Systems/MovieSceneDataLayerSystem.cpp b/Engine/Source/Runtime/MovieSceneTracks/Private/Systems/MovieSceneDataLayerSystem.cpp
--- a/Engine/Source/Runtime/MovieSceneTracks/Private/Systems/MovieSceneDataLayerSystem.cpp
+++ b/Engine/Source/Runtime/MovieSceneTracks/Private/Systems/MovieSceneDataLayerSystem.cpp
@@ -42,8 +42,7 @@ enum class EDataLayerUpdateFlags : uint8
 {
 	None						= 0,
 	FlushStreamingVisibility	= 1,
-	FlushStreamingFull			= 2,
-	PerformGarbageCollect		= 4,
+	FlushStreamingFull			= 2
 };
 ENUM_CLASS_FLAGS(EDataLayerUpdateFlags)
 
@@ -83,10 +82,9 @@ struct FDataLayerState
 {
 	void Reset();
 	bool IsEmpty() const;
-	void AddRequest(int16 InBias, EDataLayerRuntimeState RequestedState, bool bRequiresStreamingFlush, bool bPerformGC);
+	void AddRequest(int16 InBias, EDataLayerRuntimeState RequestedState, bool bRequiresStreamingFlush);
 	TOptional<EDataLayerRuntimeState> ComputeDesiredState() const;
 	bool ShouldFlushStreaming(EDataLayerRuntimeState ComputedState) const;
-	bool ShouldPerformGarbageCollect(EDataLayerRuntimeState ComputedState) const;
 
 private:
 
@@ -96,7 +94,6 @@ private:
 	int32 ActivatedCount   = 0;
 	bool  bFlushUnloaded   = false;
 	bool  bFlushActivated  = false;
-	bool  bGCUnloaded      = false;
 };
 
 
@@ -108,7 +105,7 @@ struct FDesiredLayerStates
 #if WITH_EDITOR
 	void ApplyInEditor(FPreAnimatedDataLayerStorage* PreAnimatedStorage, UDataLayerEditorSubsystem* EditorSubSystem);
 #endif
-	void ApplyNewState(const UDataLayerInstance* InDataLayer, int16 HierarchicalBias, EDataLayerRuntimeState DesiredState, bool bRequiresStreamingFlush, bool bPerformGC);
+	void ApplyNewState(const UDataLayerInstance* InDataLayer, int16 HierarchicalBias, EDataLayerRuntimeState DesiredState, bool bRequiresStreamingFlush);
 
 	TMap<TObjectKey<UDataLayerInstance>, FDataLayerState> StatesByInstance;
 };
@@ -213,10 +210,9 @@ void FDataLayerState::Reset()
 	ActivatedCount   = 0;
 	bFlushUnloaded   = false;
 	bFlushActivated  = false;
-	bGCUnloaded      = false;
 }
 
-void FDataLayerState::AddRequest(int16 InBias, EDataLayerRuntimeState RequestedState, bool bRequiresStreamingFlush, bool bPerformGC)
+void FDataLayerState::AddRequest(int16 InBias, EDataLayerRuntimeState RequestedState, bool bRequiresStreamingFlush)
 {
 	if (InBias > HierarchicalBias)
 	{
@@ -228,7 +224,7 @@ void FDataLayerState::AddRequest(int16 InBias, EDataLayerRuntimeState RequestedS
 	{
 		switch (RequestedState)
 		{
-		case EDataLayerRuntimeState::Unloaded:  ++UnloadedCount;  bFlushUnloaded |= bRequiresStreamingFlush; bGCUnloaded |= bPerformGC; break;
+		case EDataLayerRuntimeState::Unloaded:  ++UnloadedCount;  bFlushUnloaded |= bRequiresStreamingFlush; break;
 		case EDataLayerRuntimeState::Loaded:    ++LoadedCount;    break;
 		case EDataLayerRuntimeState::Activated: ++ActivatedCount; bFlushActivated |= bRequiresStreamingFlush; break;
 		}
@@ -269,16 +265,6 @@ bool FDataLayerState::ShouldFlushStreaming(EDataLayerRuntimeState ComputedState)
 	return false;
 }
 
-bool FDataLayerState::ShouldPerformGarbageCollect(EDataLayerRuntimeState ComputedState) const
-{
-	switch (ComputedState)
-	{
-	case EDataLayerRuntimeState::Unloaded:  return bGCUnloaded;
-	}
-
-	return false;
-}
-
 // ---------------------------------------------------------------------
 // FDataLayerStates definitions
 bool FDesiredLayerStates::IsEmpty() const
@@ -375,11 +361,6 @@ EDataLayerUpdateFlags FDesiredLayerStates::Apply(FPreAnimatedDataLayerStorage* P
 								*StaticEnum<EDataLayerRuntimeState>()->GetDisplayNameTextByValue((int64)DesiredStateValue).ToString());
 						}
 					}
-
-					if (StateValue.ShouldPerformGarbageCollect(DesiredStateValue))
-					{
-						Flags |= EDataLayerUpdateFlags::PerformGarbageCollect;
-					}
 				}
 			}
 		}
@@ -452,7 +433,7 @@ void FDesiredLayerStates::ApplyInEditor(FPreAnimatedDataLayerStorage* PreAnimate
 }
 #endif
 
-void FDesiredLayerStates::ApplyNewState(const UDataLayerInstance* InDataLayer, int16 HierarchicalBias, EDataLayerRuntimeState DesiredState, bool bRequiresStreamingFlush, bool bPerformGC)
+void FDesiredLayerStates::ApplyNewState(const UDataLayerInstance* InDataLayer, int16 HierarchicalBias, EDataLayerRuntimeState DesiredState, bool bRequiresStreamingFlush)
 {
 	using namespace UE::MovieScene;
 
@@ -462,7 +443,7 @@ void FDesiredLayerStates::ApplyNewState(const UDataLayerInstance* InDataLayer, i
 		LayerState = &StatesByInstance.Add(InDataLayer, FDataLayerState());
 	}
 
-	LayerState->AddRequest(HierarchicalBias, DesiredState, bRequiresStreamingFlush, bPerformGC);
+	LayerState->AddRequest(HierarchicalBias, DesiredState, bRequiresStreamingFlush);
 }
 
 } // namespace MovieScene
@@ -565,12 +546,6 @@ void UMovieSceneDataLayerSystem::OnRun(FSystemTaskPrerequisites& InPrerequisites
 				World->FlushLevelStreaming(EFlushLevelStreamingType::Visibility);
 			}
 
-			if (EnumHasAnyFlags(UpdateFlags, EDataLayerUpdateFlags::PerformGarbageCollect))
-			{
-				UE_LOG(LogMovieScene, Warning, TEXT("[UMovieSceneDataLayerSystem] Forcing garbage collection"));
-				GLevelStreamingContinuouslyIncrementalGCWhileLevelsPendingPurgeOverride = 1;
-			}
-
 			UE_SUPPRESS(LogMovieScene, Warning,
 			{
 				if (FlushTypeString)
@@ -646,14 +621,13 @@ void UMovieSceneDataLayerSystem::UpdateDesiredStates()
 				EDataLayerRuntimeState DesiredState = bPreroll ? Section->GetPrerollState() : Section->GetDesiredState();
 				const bool bRequiresStreamingFlush = (DesiredState == EDataLayerRuntimeState::Unloaded) ? Section->GetFlushOnUnload() :
 													 (DesiredState == EDataLayerRuntimeState::Activated) ? Section->GetFlushOnActivated() : false;
-				const bool bPerformGC = (DesiredState == EDataLayerRuntimeState::Unloaded) ? Section->GetPerformGCOnUnload() : false;
 
 				for (const UDataLayerAsset* DataLayerAsset : Section->GetDataLayerAssets())
 				{
 					const UDataLayerInstance* DataLayerInstance = DataLayerManager->GetDataLayerInstanceFromAsset(DataLayerAsset);
 					if (DataLayerInstance)
 					{
-						this->DesiredLayerStates->ApplyNewState(DataLayerInstance, OptHBiases ? OptHBiases[Index] : 0, DesiredState, bRequiresStreamingFlush, bPerformGC);
+						this->DesiredLayerStates->ApplyNewState(DataLayerInstance, OptHBiases ? OptHBiases[Index] : 0, DesiredState, bRequiresStreamingFlush);
 					}
 				}
 			}
diff --git a/Engine/Source/Runtime/MovieSceneTracks/Public/Sections/MovieSceneDataLayerSection.h b/Engine/Source/Runtime/MovieSceneTracks/Public/Sections/MovieSceneDataLayerSection.h
--- a/Engine/Source/Runtime/MovieSceneTracks/Public/Sections/MovieSceneDataLayerSection.h
+++ b/Engine/Source/Runtime/MovieSceneTracks/Public/Sections/MovieSceneDataLayerSection.h
@@ -45,12 +45,6 @@ public:
 	UFUNCTION(BlueprintCallable, Category = "Sequencer|Section")
 	MOVIESCENETRACKS_API void SetFlushOnUnload(bool bFlushOnUnload);
 		
-	UFUNCTION(BlueprintPure, Category = "Sequencer|Section")
-	MOVIESCENETRACKS_API bool GetPerformGCOnUnload() const;
-
-	UFUNCTION(BlueprintCallable, Category = "Sequencer|Section")
-	MOVIESCENETRACKS_API void SetPerformGCOnUnload(bool bPerformGCOnUnload);
-
 	UFUNCTION(BlueprintPure, Category = "Sequencer|Section")
 	const TArray<UDataLayerAsset*>& GetDataLayerAssets() const { return DataLayerAssets; }
 
@@ -82,10 +76,6 @@ private:
 	UPROPERTY(EditAnywhere, Category=DataLayer, Meta=(EditCondition="DesiredState == EDataLayerRuntimeState::Unloaded", EditConditionHides))
 	bool bFlushOnUnload;
 
-	/** Determine if we need to perform a GC when the data layers unloads. */
-	UPROPERTY(EditAnywhere, Category=DataLayer, Meta=(EditCondition="DesiredState == EDataLayerRuntimeState::Unloaded", EditConditionHides, DisplayName="Perform GC On Unload"))
-	bool bPerformGCOnUnload;
-
 	UFUNCTION()
 	bool HasPreRoll() const { return GetPreRollFrames() > 0.0f; }
 };
