Sequencer- fix crash when restoring state when level visbility and skeletal animation have both been animated.

diff --git a/Engine/Source/Runtime/MovieSceneTracks/Private/Systems/MovieSceneSkeletalAnimationSystem.cpp b/Engine/Source/Runtime/MovieSceneTracks/Private/Systems/MovieSceneSkeletalAnimationSystem.cpp
--- a/Engine/Source/Runtime/MovieSceneTracks/Private/Systems/MovieSceneSkeletalAnimationSystem.cpp
+++ b/Engine/Source/Runtime/MovieSceneTracks/Private/Systems/MovieSceneSkeletalAnimationSystem.cpp
@@ -94,7 +94,7 @@ struct FPreAnimatedSkeletalAnimationTraits : FBoundObjectPreAnimatedStateTraits
 	void RestorePreAnimatedValue(const KeyType& Object, StorageType& InOutCachedValue, const FRestoreStateParams& Params)
 	{
 		USkeletalMeshComponent* Component = Cast<USkeletalMeshComponent>(Object.ResolveObjectPtr());
-		if (!Component)
+		if (!Component || !Component->IsRegistered())
 		{
 			return;
 		}
