Sequencer: Fixed adding timewarp to skeletal animation sections not showing up in the UI

diff --git a/Engine/Source/Runtime/MovieSceneTracks/Private/Sections/MovieSceneSkeletalAnimationSection.cpp b/Engine/Source/Runtime/MovieSceneTracks/Private/Sections/MovieSceneSkeletalAnimationSection.cpp
--- a/Engine/Source/Runtime/MovieSceneTracks/Private/Sections/MovieSceneSkeletalAnimationSection.cpp
+++ b/Engine/Source/Runtime/MovieSceneTracks/Private/Sections/MovieSceneSkeletalAnimationSection.cpp
@@ -523,18 +523,20 @@ void UMovieSceneSkeletalAnimationSection::PostEditChangeProperty(FPropertyChange
 {
 	// Adjust the duration automatically if the play rate changes
 	if (PropertyChangedEvent.Property != nullptr &&
-		PropertyChangedEvent.Property->GetFName() == TEXT("PlayRate") && 
-		Params.PlayRate.GetType() == EMovieSceneTimeWarpType::FixedPlayRate)
+		PropertyChangedEvent.Property->GetFName() == TEXT("PlayRate"))
 	{
-		float NewPlayRate = Params.PlayRate.AsFixedPlayRateFloat();
-
-		if (!FMath::IsNearlyZero(NewPlayRate))
+		if (Params.PlayRate.GetType() == EMovieSceneTimeWarpType::FixedPlayRate)
 		{
-			float CurrentDuration = UE::MovieScene::DiscreteSize(GetRange());
-			float NewDuration = CurrentDuration * (PreviousPlayRate / NewPlayRate);
-			SetEndFrame( GetInclusiveStartFrame() + FMath::FloorToInt(NewDuration) );
+			float NewPlayRate = Params.PlayRate.AsFixedPlayRateFloat();
+
+			if (!FMath::IsNearlyZero(NewPlayRate))
+			{
+				float CurrentDuration = UE::MovieScene::DiscreteSize(GetRange());
+				float NewDuration = CurrentDuration * (PreviousPlayRate / NewPlayRate);
+				SetEndFrame( GetInclusiveStartFrame() + FMath::FloorToInt(NewDuration) );
 
-			PreviousPlayRate = NewPlayRate;
+				PreviousPlayRate = NewPlayRate;
+			}
 		}
 
 		ChannelProxy = nullptr;
