MovieScene : MovieSceneTracks
- Add FRotator keyframe support in sequencer (over rotation allowed)

diff --git a/Engine/Source/Editor/MovieSceneTools/Private/MovieSceneToolsModule.cpp b/Engine/Source/Editor/MovieSceneTools/Private/MovieSceneToolsModule.cpp
--- a/Engine/Source/Editor/MovieSceneTools/Private/MovieSceneToolsModule.cpp
+++ b/Engine/Source/Editor/MovieSceneTools/Private/MovieSceneToolsModule.cpp
@@ -26,6 +26,7 @@
 #include "TrackEditors/PropertyTrackEditors/VectorPropertyTrackEditor.h"
 #include "TrackEditors/PropertyTrackEditors/TransformPropertyTrackEditor.h"
 #include "TrackEditors/PropertyTrackEditors/EulerTransformPropertyTrackEditor.h"
+#include "TrackEditors/PropertyTrackEditors/RotatorPropertyTrackEditor.h"
 #include "TrackEditors/PropertyTrackEditors/VisibilityPropertyTrackEditor.h"
 #include "TrackEditors/PropertyTrackEditors/ActorReferencePropertyTrackEditor.h"
 #include "TrackEditors/PropertyTrackEditors/StringPropertyTrackEditor.h"
@@ -162,6 +163,7 @@ void FMovieSceneToolsModule::StartupModule()
 		DoubleVectorPropertyTrackCreateEditorHandle = SequencerModule.RegisterPropertyTrackEditor<FDoubleVectorPropertyTrackEditor>();
 		TransformPropertyTrackCreateEditorHandle = SequencerModule.RegisterPropertyTrackEditor<FTransformPropertyTrackEditor>();
 		EulerTransformPropertyTrackCreateEditorHandle = SequencerModule.RegisterPropertyTrackEditor<FEulerTransformPropertyTrackEditor>();
+		RotatorPropertyTrackCreateEditorHandle = SequencerModule.RegisterPropertyTrackEditor<FRotatorPropertyTrackEditor>();
 		VisibilityPropertyTrackCreateEditorHandle = SequencerModule.RegisterPropertyTrackEditor<FVisibilityPropertyTrackEditor>();
 		ActorReferencePropertyTrackCreateEditorHandle = SequencerModule.RegisterPropertyTrackEditor<FActorReferencePropertyTrackEditor>();
 		StringPropertyTrackCreateEditorHandle = SequencerModule.RegisterPropertyTrackEditor<FStringPropertyTrackEditor>();
@@ -335,6 +337,7 @@ void FMovieSceneToolsModule::ShutdownModule()
 	SequencerModule.UnRegisterTrackEditor( DoubleVectorPropertyTrackCreateEditorHandle );
 	SequencerModule.UnRegisterTrackEditor( TransformPropertyTrackCreateEditorHandle );
 	SequencerModule.UnRegisterTrackEditor( EulerTransformPropertyTrackCreateEditorHandle );
+	SequencerModule.UnRegisterTrackEditor( RotatorPropertyTrackCreateEditorHandle );
 	SequencerModule.UnRegisterTrackEditor( VisibilityPropertyTrackCreateEditorHandle );
 	SequencerModule.UnRegisterTrackEditor( ActorReferencePropertyTrackCreateEditorHandle );
 	SequencerModule.UnRegisterTrackEditor( StringPropertyTrackCreateEditorHandle );
diff --git a/Engine/Source/Editor/MovieSceneTools/Private/TrackEditors/PropertyTrackEditors/RotatorPropertyTrackEditor.cpp b/Engine/Source/Editor/MovieSceneTools/Private/TrackEditors/PropertyTrackEditors/RotatorPropertyTrackEditor.cpp
new file mode 100644
--- /dev/null
+++ b/Engine/Source/Editor/MovieSceneTools/Private/TrackEditors/PropertyTrackEditors/RotatorPropertyTrackEditor.cpp
@@ -0,0 +1,41 @@
+// Copyright Epic Games, Inc. All Rights Reserved.
+
+#include "RotatorPropertyTrackEditor.h"
+
+#include "MovieSceneTracksComponentTypes.h"
+
+TSharedRef<ISequencerTrackEditor> FRotatorPropertyTrackEditor::CreateTrackEditor(TSharedRef<ISequencer> InOwningSequencer)
+{
+	return MakeShared<FRotatorPropertyTrackEditor>(InOwningSequencer);
+}
+
+void FRotatorPropertyTrackEditor::GenerateKeysFromPropertyChanged(const FPropertyChangedParams& InPropertyChangedParams, UMovieSceneSection* InSectionToKey, FGeneratedTrackKeys& OutGeneratedKeys)
+{
+	using namespace  UE::MovieScene;
+
+	const FPropertyPath& StructPath = InPropertyChangedParams.StructPathToKey;
+
+	const FName ChannelName = StructPath.GetNumProperties() != 0
+		? StructPath.GetLeafMostProperty().Property->GetFName()
+		: NAME_None;
+
+	const FRotator CurrentRotator = InPropertyChangedParams.GetPropertyValue<FRotator>();
+
+	const bool bKeyRoll = ChannelName == GET_MEMBER_NAME_CHECKED(FRotator, Roll) || ChannelName == NAME_None;
+	const bool bKeyPitch = ChannelName == GET_MEMBER_NAME_CHECKED(FRotator, Pitch) || ChannelName == NAME_None;
+	const bool bKeyYaw = ChannelName == GET_MEMBER_NAME_CHECKED(FRotator, Yaw) || ChannelName == NAME_None;
+
+	int32 ChannelIndex = 0;
+
+	OutGeneratedKeys.Add(FMovieSceneChannelValueSetter::Create<FMovieSceneDoubleChannel>(ChannelIndex++
+		, CurrentRotator.Pitch
+		, bKeyPitch));
+
+	OutGeneratedKeys.Add(FMovieSceneChannelValueSetter::Create<FMovieSceneDoubleChannel>(ChannelIndex++
+		, CurrentRotator.Yaw
+		, bKeyYaw));
+
+	OutGeneratedKeys.Add(FMovieSceneChannelValueSetter::Create<FMovieSceneDoubleChannel>(ChannelIndex
+		, CurrentRotator.Roll
+		, bKeyRoll));
+}
diff --git a/Engine/Source/Editor/MovieSceneTools/Private/TrackEditors/PropertyTrackEditors/RotatorPropertyTrackEditor.h b/Engine/Source/Editor/MovieSceneTools/Private/TrackEditors/PropertyTrackEditors/RotatorPropertyTrackEditor.h
new file mode 100644
--- /dev/null
+++ b/Engine/Source/Editor/MovieSceneTools/Private/TrackEditors/PropertyTrackEditors/RotatorPropertyTrackEditor.h
@@ -0,0 +1,36 @@
+// Copyright Epic Games, Inc. All Rights Reserved.
+
+#pragma once
+
+#include "ISequencer.h"
+#include "ISequencerSection.h"
+#include "ISequencerTrackEditor.h"
+#include "KeyPropertyParams.h"
+#include "Misc/Guid.h"
+#include "PropertyTrackEditor.h"
+#include "Tracks/MovieSceneRotatorTrack.h"
+
+/** Property track editor for rotator properties */
+class FRotatorPropertyTrackEditor : public FPropertyTrackEditor<UMovieSceneRotatorTrack>
+{
+public:
+	FRotatorPropertyTrackEditor(TSharedRef<ISequencer> InSequencer)
+		: FPropertyTrackEditor(InSequencer, GetAnimatedPropertyTypes())
+	{}
+
+	static TArray<FAnimatedPropertyKey, TInlineAllocator<1>> GetAnimatedPropertyTypes()
+	{
+		return TArray<FAnimatedPropertyKey, TInlineAllocator<1>>(
+			{
+				FAnimatedPropertyKey::FromStructType(NAME_Rotator)
+			}
+		);
+	}
+
+	static TSharedRef<ISequencerTrackEditor> CreateTrackEditor(TSharedRef<ISequencer> InOwningSequencer);
+
+protected:
+	//~ Begin FPropertyTrackEditor
+	virtual void GenerateKeysFromPropertyChanged(const FPropertyChangedParams& InPropertyChangedParams, UMovieSceneSection* InSectionToKey, FGeneratedTrackKeys& OutGeneratedKeys) override;
+	//~ End FPropertyTrackEditor
+};
diff --git a/Engine/Source/Editor/MovieSceneTools/Public/MovieSceneToolsModule.h b/Engine/Source/Editor/MovieSceneTools/Public/MovieSceneToolsModule.h
--- a/Engine/Source/Editor/MovieSceneTools/Public/MovieSceneToolsModule.h
+++ b/Engine/Source/Editor/MovieSceneTools/Public/MovieSceneToolsModule.h
@@ -112,6 +112,7 @@ private:
 	FDelegateHandle DoubleVectorPropertyTrackCreateEditorHandle;
 	FDelegateHandle TransformPropertyTrackCreateEditorHandle;
 	FDelegateHandle EulerTransformPropertyTrackCreateEditorHandle;
+	FDelegateHandle RotatorPropertyTrackCreateEditorHandle;
 	FDelegateHandle VisibilityPropertyTrackCreateEditorHandle;
 	FDelegateHandle ActorReferencePropertyTrackCreateEditorHandle;
 	FDelegateHandle StringPropertyTrackCreateEditorHandle;
diff --git a/Engine/Source/Runtime/MovieSceneTracks/Private/MovieSceneTracksComponentTypes.cpp b/Engine/Source/Runtime/MovieSceneTracks/Private/MovieSceneTracksComponentTypes.cpp
--- a/Engine/Source/Runtime/MovieSceneTracks/Private/MovieSceneTracksComponentTypes.cpp
+++ b/Engine/Source/Runtime/MovieSceneTracks/Private/MovieSceneTracksComponentTypes.cpp
@@ -563,6 +563,8 @@ FMovieSceneTracksComponentTypes::FMovieSceneTracksComponentTypes()
 	ComponentRegistry->NewPropertyType(EulerTransform, TEXT("FEulerTransform"));
 	ComponentRegistry->NewPropertyType(ComponentTransform, TEXT("Component Transform"));
 
+	ComponentRegistry->NewPropertyType(Rotator, TEXT("FRotator"));
+
 	ComponentRegistry->NewPropertyType(FloatParameter, TEXT("float parameter"));
 	ComponentRegistry->NewPropertyType(ColorParameter, TEXT("color parameter"));
 
@@ -576,6 +578,10 @@ FMovieSceneTracksComponentTypes::FMovieSceneTracksComponentTypes()
 	ComponentRegistry->NewComponentType(&QuaternionRotationChannel[1], TEXT("Quaternion Rotation Channel 1"));
 	ComponentRegistry->NewComponentType(&QuaternionRotationChannel[2], TEXT("Quaternion Rotation Channel 2"));
 
+	ComponentRegistry->NewComponentType(&RotatorChannel[0], TEXT("Rotator Channel Y (Pitch)"));
+	ComponentRegistry->NewComponentType(&RotatorChannel[1], TEXT("Rotator Channel Z (Yaw)"));
+	ComponentRegistry->NewComponentType(&RotatorChannel[2], TEXT("Rotator Channel X (Roll)"));
+
 	ComponentRegistry->NewComponentType(&ConstraintChannel, TEXT("Constraint Channel"));
 
 	ComponentRegistry->NewComponentType(&AttachParent, TEXT("Attach Parent"));
@@ -636,6 +642,7 @@ FMovieSceneTracksComponentTypes::FMovieSceneTracksComponentTypes()
 	.SetCustomAccessors(&Accessors.Bool)
 	.Commit(FBoolHandler());
 
+	// --------------------------------------------------------------------------------------------
 	// Set up FTransform properties
 	BuiltInComponents->PropertyRegistry.DefineCompositeProperty(Transform, TEXT("Apply FTransform Properties"))
 	.AddComposite(BuiltInComponents->DoubleResult[0], &FIntermediate3DTransform::T_X)
@@ -650,6 +657,16 @@ FMovieSceneTracksComponentTypes::FMovieSceneTracksComponentTypes()
 	.SetBlenderSystem<UMovieScenePiecewiseDoubleBlenderSystem>()
 	.Commit();
 
+	// --------------------------------------------------------------------------------------------
+	// Set up FRotator properties
+	BuiltInComponents->PropertyRegistry.DefineCompositeProperty(Rotator, TEXT("Apply FRotator Properties"))
+		.AddComposite(BuiltInComponents->DoubleResult[0], &FRotator::Pitch)
+		.AddComposite(BuiltInComponents->DoubleResult[1], &FRotator::Yaw)
+		.AddComposite(BuiltInComponents->DoubleResult[2], &FRotator::Roll)
+		// Use this blender since we want over rotation 720deg = 2 cycles instead of 720deg = 0, we could later add support for multiple interpolation
+		.SetBlenderSystem<UMovieScenePiecewiseDoubleBlenderSystem>()
+		.Commit();
+
 	// --------------------------------------------------------------------------------------------
 	// Set up byte properties
 	BuiltInComponents->PropertyRegistry.DefineProperty(Byte, TEXT("Apply Byte Properties"))
@@ -800,6 +817,15 @@ FMovieSceneTracksComponentTypes::FMovieSceneTracksComponentTypes()
 		ComponentRegistry->Factories.DefineMutuallyInclusiveComponent(QuaternionRotationChannel[Index], BuiltInComponents->EvalTime);
 	}
 
+	// --------------------------------------------------------------------------------------------
+	// Set up rotator components
+	for (int32 Index = 0; Index < UE_ARRAY_COUNT(RotatorChannel); ++Index)
+	{
+		ComponentRegistry->Factories.DuplicateChildComponent(RotatorChannel[Index]);
+		ComponentRegistry->Factories.DefineMutuallyInclusiveComponent(RotatorChannel[Index], BuiltInComponents->DoubleResult[Index]);
+		ComponentRegistry->Factories.DefineMutuallyInclusiveComponent(RotatorChannel[Index], BuiltInComponents->EvalTime);
+	}
+
 	// -------------------------------------------------------------------------------------------
 	// Set up constraint components
 	ComponentRegistry->Factories.DuplicateChildComponent(ConstraintChannel);
diff --git a/Engine/Source/Runtime/MovieSceneTracks/Private/Sections/MovieSceneRotatorSection.cpp b/Engine/Source/Runtime/MovieSceneTracks/Private/Sections/MovieSceneRotatorSection.cpp
new file mode 100644
--- /dev/null
+++ b/Engine/Source/Runtime/MovieSceneTracks/Private/Sections/MovieSceneRotatorSection.cpp
@@ -0,0 +1,119 @@
+// Copyright Epic Games, Inc. All Rights Reserved.
+
+#include "Sections/MovieSceneRotatorSection.h"
+
+#include "Channels/MovieSceneChannelProxy.h"
+#include "MovieSceneTracksComponentTypes.h"
+#include "Tracks/MovieScenePropertyTrack.h"
+
+#if WITH_EDITOR
+struct FRotatorSectionEditorData
+{
+	FRotatorSectionEditorData()
+	{
+		MetaData[0].SetIdentifiers(TEXT("Rotation.Y"), FCommonChannelData::ChannelY);
+		MetaData[0].SubPropertyPath = TEXT("Pitch");
+		MetaData[0].SortOrder = 1;
+		MetaData[0].Color = FCommonChannelData::GreenChannelColor;
+		MetaData[0].bCanCollapseToTrack = false;
+		ExternalValues[0].OnGetExternalValue = ExtractChannelY;
+
+		MetaData[1].SetIdentifiers(TEXT("Rotation.Z"), FCommonChannelData::ChannelZ);
+		MetaData[1].SubPropertyPath = TEXT("Yaw");
+		MetaData[1].SortOrder = 2;
+		MetaData[1].Color = FCommonChannelData::BlueChannelColor;
+		MetaData[1].bCanCollapseToTrack = false;
+		ExternalValues[1].OnGetExternalValue = ExtractChannelZ;
+
+		MetaData[2].SetIdentifiers(TEXT("Rotation.X"), FCommonChannelData::ChannelX);
+		MetaData[2].SubPropertyPath = TEXT("Roll");
+		MetaData[2].SortOrder = 0;
+		MetaData[2].Color = FCommonChannelData::RedChannelColor;
+		MetaData[2].bCanCollapseToTrack = false;
+		ExternalValues[2].OnGetExternalValue = ExtractChannelX;
+	}
+
+	static FRotator GetPropertyValue(const UObject& InObject, FTrackInstancePropertyBindings& InBindings)
+	{
+		if (const FStructProperty* RotatorStructProperty = CastField<FStructProperty>(InBindings.GetProperty(InObject)))
+		{
+			if (RotatorStructProperty->Struct->GetFName() == NAME_Rotator)
+			{
+				return InBindings.GetCurrentValue<FRotator>(InObject);
+			}
+		}
+
+		return FRotator::ZeroRotator;
+	}
+
+	static TOptional<double> ExtractChannelX(UObject& InObject, FTrackInstancePropertyBindings* InBindings)
+	{
+		return InBindings ? GetPropertyValue(InObject, *InBindings).Roll : TOptional<double>();
+	}
+
+	static TOptional<double> ExtractChannelY(UObject& InObject, FTrackInstancePropertyBindings* InBindings)
+	{
+		return InBindings ? GetPropertyValue(InObject, *InBindings).Pitch : TOptional<double>();
+	}
+
+	static TOptional<double> ExtractChannelZ(UObject& InObject, FTrackInstancePropertyBindings* InBindings)
+	{
+		return InBindings ? GetPropertyValue(InObject, *InBindings).Yaw : TOptional<double>();
+	}
+
+	FMovieSceneChannelMetaData MetaData[3];
+	TMovieSceneExternalValue<double> ExternalValues[3];
+};
+#endif
+
+UMovieSceneRotatorSection::UMovieSceneRotatorSection(const FObjectInitializer& InObjectInitializer)
+	: Super(InObjectInitializer)
+{
+	EvalOptions.EnableAndSetCompletionMode(EMovieSceneCompletionMode::ProjectDefault);
+	bSupportsInfiniteRange = true;
+	BlendType = EMovieSceneBlendType::Absolute;
+
+	Rotation[0].SetDefault(0.f);
+	Rotation[1].SetDefault(0.f);
+	Rotation[2].SetDefault(0.f);
+
+	FMovieSceneChannelProxyData Channels;
+
+#if WITH_EDITOR
+	static FRotatorSectionEditorData EditorData;
+	Channels.Add(Rotation[0], EditorData.MetaData[0], EditorData.ExternalValues[0]);
+	Channels.Add(Rotation[1], EditorData.MetaData[1], EditorData.ExternalValues[1]);
+	Channels.Add(Rotation[2], EditorData.MetaData[2], EditorData.ExternalValues[2]);
+#else
+	Channels.Add(Rotation[0]);
+	Channels.Add(Rotation[1]);
+	Channels.Add(Rotation[2]);
+#endif
+
+	ChannelProxy = MakeShared<FMovieSceneChannelProxy>(MoveTemp(Channels));
+}
+
+bool UMovieSceneRotatorSection::PopulateEvaluationFieldImpl(const TRange<FFrameNumber>& InEffectiveRange, const FMovieSceneEvaluationFieldEntityMetaData& InMetaData, FMovieSceneEntityComponentFieldBuilder* OutFieldBuilder)
+{
+	FMovieScenePropertyTrackEntityImportHelper::PopulateEvaluationField(*this, InEffectiveRange, InMetaData, OutFieldBuilder);
+	return true;
+}
+
+void UMovieSceneRotatorSection::ImportEntityImpl(UMovieSceneEntitySystemLinker* InEntityLinker, const FEntityImportParams& InParams, FImportedEntity* OutImportedEntity)
+{
+	if (!Rotation[0].HasAnyData() && !Rotation[1].HasAnyData() && !Rotation[2].HasAnyData())
+	{
+		return;
+	}
+
+	using namespace UE::MovieScene;
+
+	const FBuiltInComponentTypes* BuiltInComponents = FBuiltInComponentTypes::Get();
+	const FMovieSceneTracksComponentTypes* TracksComponents = FMovieSceneTracksComponentTypes::Get();
+
+	FPropertyTrackEntityImportHelper(TracksComponents->Rotator)
+		.AddConditional(BuiltInComponents->DoubleChannel[0], &Rotation[0], Rotation[0].HasAnyData())
+		.AddConditional(BuiltInComponents->DoubleChannel[1], &Rotation[1], Rotation[1].HasAnyData())
+		.AddConditional(BuiltInComponents->DoubleChannel[2], &Rotation[2], Rotation[2].HasAnyData())
+		.Commit(this, InParams, OutImportedEntity);
+}
diff --git a/Engine/Source/Runtime/MovieSceneTracks/Private/Systems/MovieSceneRotatorPropertySystem.cpp b/Engine/Source/Runtime/MovieSceneTracks/Private/Systems/MovieSceneRotatorPropertySystem.cpp
new file mode 100644
--- /dev/null
+++ b/Engine/Source/Runtime/MovieSceneTracks/Private/Systems/MovieSceneRotatorPropertySystem.cpp
@@ -0,0 +1,24 @@
+// Copyright Epic Games, Inc. All Rights Reserved.
+
+#include "Systems/MovieSceneRotatorPropertySystem.h"
+
+#include "MovieSceneTracksComponentTypes.h"
+#include "Systems/DoubleChannelEvaluatorSystem.h"
+#include "Systems/MovieScenePiecewiseDoubleBlenderSystem.h"
+
+UMovieSceneRotatorPropertySystem::UMovieSceneRotatorPropertySystem(const FObjectInitializer& InObjectInitializer)
+	: Super(InObjectInitializer)
+{
+	BindToProperty(UE::MovieScene::FMovieSceneTracksComponentTypes::Get()->Rotator);
+
+	if (HasAnyFlags(RF_ClassDefaultObject))
+	{
+		DefineImplicitPrerequisite(UMovieScenePiecewiseDoubleBlenderSystem::StaticClass(), GetClass());
+		DefineImplicitPrerequisite(UDoubleChannelEvaluatorSystem::StaticClass(), GetClass());
+	}
+}
+
+void UMovieSceneRotatorPropertySystem::OnRun(FSystemTaskPrerequisites& InPrerequisites, FSystemSubsequentTasks& InSubsequents)
+{
+	Super::OnRun(InPrerequisites, InSubsequents);
+}
diff --git a/Engine/Source/Runtime/MovieSceneTracks/Private/Tracks/MovieSceneRotatorTrack.cpp b/Engine/Source/Runtime/MovieSceneTracks/Private/Tracks/MovieSceneRotatorTrack.cpp
new file mode 100644
--- /dev/null
+++ b/Engine/Source/Runtime/MovieSceneTracks/Private/Tracks/MovieSceneRotatorTrack.cpp
@@ -0,0 +1,58 @@
+// Copyright Epic Games, Inc. All Rights Reserved.
+
+#include "Tracks/MovieSceneRotatorTrack.h"
+
+#include "Evaluation/Blending/MovieSceneBlendType.h"
+#include "Sections/MovieSceneRotatorSection.h"
+
+UMovieSceneRotatorTrack::UMovieSceneRotatorTrack(const FObjectInitializer& InObjectInitializer)
+	: Super(InObjectInitializer)
+{
+	SupportedBlendTypes = FMovieSceneBlendTypeField::All();
+}
+
+bool UMovieSceneRotatorTrack::SupportsType(TSubclassOf<UMovieSceneSection> InSectionClass) const
+{
+	return InSectionClass == UMovieSceneRotatorSection::StaticClass();
+}
+
+UMovieSceneSection* UMovieSceneRotatorTrack::CreateNewSection()
+{
+	return NewObject<UMovieSceneRotatorSection>(this, NAME_None, RF_Transactional);
+}
+
+void UMovieSceneRotatorTrack::AddSection(UMovieSceneSection& InSection)
+{
+	Sections.Add(&InSection);
+}
+
+const TArray<UMovieSceneSection*>& UMovieSceneRotatorTrack::GetAllSections() const
+{
+	return Sections;
+}
+
+bool UMovieSceneRotatorTrack::HasSection(const UMovieSceneSection& InSection) const
+{
+	return Sections.Contains(&InSection);
+}
+
+bool UMovieSceneRotatorTrack::IsEmpty() const
+{
+	return Sections.IsEmpty();
+}
+
+void UMovieSceneRotatorTrack::RemoveAllAnimationData()
+{
+	Sections.Empty();
+}
+
+void UMovieSceneRotatorTrack::RemoveSection(UMovieSceneSection& InSection)
+{
+	Sections.Remove(&InSection);
+}
+
+void UMovieSceneRotatorTrack::RemoveSectionAt(int32 InSectionIndex)
+{
+	Sections.RemoveAt(InSectionIndex);
+}
+
diff --git a/Engine/Source/Runtime/MovieSceneTracks/Public/MovieSceneTracksComponentTypes.h b/Engine/Source/Runtime/MovieSceneTracks/Public/MovieSceneTracksComponentTypes.h
--- a/Engine/Source/Runtime/MovieSceneTracks/Public/MovieSceneTracksComponentTypes.h
+++ b/Engine/Source/Runtime/MovieSceneTracks/Public/MovieSceneTracksComponentTypes.h
@@ -655,8 +655,8 @@ using FDoublePropertyTraits             = TDirectPropertyTraits<double, false>;
 using FTransformPropertyTraits          = TIndirectPropertyTraits<FTransform, FIntermediate3DTransform>;
 using FEulerTransformPropertyTraits     = TIndirectPropertyTraits<FEulerTransform, FIntermediate3DTransform>;
 using FComponentTransformPropertyTraits = TDirectPropertyTraits<FIntermediate3DTransform>;
-using FStringPropertyTraits			    = TDirectPropertyTraits<FString>;
-
+using FRotatorPropertyTraits            = TDirectPropertyTraits<FRotator>;
+using FStringPropertyTraits             = TDirectPropertyTraits<FString>;
 using FFloatParameterTraits             = TIndirectPropertyTraits<float, double, false>;
 using FColorParameterTraits             = TIndirectPropertyTraits<FLinearColor, FIntermediateColor>;
 
@@ -676,6 +676,7 @@ struct FMovieSceneTracksComponentTypes
 	TPropertyComponents<FTransformPropertyTraits> Transform;
 	TPropertyComponents<FEulerTransformPropertyTraits> EulerTransform;
 	TPropertyComponents<FComponentTransformPropertyTraits> ComponentTransform;
+	TPropertyComponents<FRotatorPropertyTraits> Rotator;
 	TPropertyComponents<FStringPropertyTraits> String;
 	TPropertyComponents<FObjectPropertyTraits> Object;
 
@@ -683,6 +684,7 @@ struct FMovieSceneTracksComponentTypes
 	TPropertyComponents<FColorParameterTraits> ColorParameter;
 
 	TComponentTypeID<FSourceDoubleChannel> QuaternionRotationChannel[3];
+	TComponentTypeID<FSourceDoubleChannel> RotatorChannel[3];
 
 	TComponentTypeID<FConstraintComponentData> ConstraintChannel;
 
diff --git a/Engine/Source/Runtime/MovieSceneTracks/Public/Sections/MovieSceneRotatorSection.h b/Engine/Source/Runtime/MovieSceneTracks/Public/Sections/MovieSceneRotatorSection.h
new file mode 100644
--- /dev/null
+++ b/Engine/Source/Runtime/MovieSceneTracks/Public/Sections/MovieSceneRotatorSection.h
@@ -0,0 +1,58 @@
+// Copyright Epic Games, Inc. All Rights Reserved.
+
+#pragma once
+
+#include "Channels/MovieSceneDoubleChannel.h"
+#include "EntitySystem/IMovieSceneEntityProvider.h"
+#include "MovieSceneSection.h"
+#include "MovieSceneRotatorSection.generated.h"
+
+/** Movie scene section that animates each component (X, Y, Z) of an FRotator property */
+UCLASS(MinimalAPI)
+class UMovieSceneRotatorSection
+	: public UMovieSceneSection
+	, public IMovieSceneEntityProvider
+{
+	GENERATED_BODY()
+
+public:
+	static constexpr int32 RollChannelIndex = 2;
+	static constexpr int32 PitchChannelIndex = 0;
+	static constexpr int32 YawChannelIndex = 1;
+
+	UMovieSceneRotatorSection(const FObjectInitializer& InObjectInitializer);
+
+	/** Get the Roll channel */
+	const FMovieSceneDoubleChannel& GetChannelX() const
+	{
+		return Rotation[RollChannelIndex];
+	}
+
+	/** Get the Pitch channel */
+	const FMovieSceneDoubleChannel& GetChannelY() const
+	{
+		return Rotation[PitchChannelIndex];
+	}
+
+	/** Get the Yaw channel */
+	const FMovieSceneDoubleChannel& GetChannelZ() const
+	{
+		return Rotation[YawChannelIndex];
+	}
+
+	const FMovieSceneDoubleChannel& GetChannel(int32 InIndex) const
+	{
+		constexpr int32 Size = UE_ARRAY_COUNT(Rotation);
+		check(InIndex >= 0 && InIndex < Size)
+		return Rotation[InIndex];
+	}
+
+	//~ Begin IMovieSceneEntityProvider
+	virtual bool PopulateEvaluationFieldImpl(const TRange<FFrameNumber>& InEffectiveRange, const FMovieSceneEvaluationFieldEntityMetaData& InMetaData, FMovieSceneEntityComponentFieldBuilder* OutFieldBuilder) override;
+	virtual void ImportEntityImpl(UMovieSceneEntitySystemLinker* InEntityLinker, const FEntityImportParams& InParams, FImportedEntity* OutImportedEntity) override;
+	//~ End IMovieSceneEntityProvider
+
+private:
+	UPROPERTY()
+	FMovieSceneDoubleChannel Rotation[3];
+};
diff --git a/Engine/Source/Runtime/MovieSceneTracks/Public/Systems/MovieSceneRotatorPropertySystem.h b/Engine/Source/Runtime/MovieSceneTracks/Public/Systems/MovieSceneRotatorPropertySystem.h
new file mode 100644
--- /dev/null
+++ b/Engine/Source/Runtime/MovieSceneTracks/Public/Systems/MovieSceneRotatorPropertySystem.h
@@ -0,0 +1,20 @@
+// Copyright Epic Games, Inc. All Rights Reserved.
+
+#pragma once
+
+#include "Systems/MovieScenePropertySystem.h"
+#include "MovieSceneRotatorPropertySystem.generated.h"
+
+/** System for rotator property registered within PropertyRegistry */
+UCLASS(MinimalAPI)
+class UMovieSceneRotatorPropertySystem : public UMovieScenePropertySystem
+{
+	GENERATED_BODY()
+
+public:
+	UMovieSceneRotatorPropertySystem(const FObjectInitializer& InObjectInitializer);
+
+	//~ Begin UMovieSceneEntitySystem
+	virtual void OnRun(FSystemTaskPrerequisites& InPrerequisites, FSystemSubsequentTasks& InSubsequents) override;
+	//~ End UMovieSceneEntitySystem
+};
diff --git a/Engine/Source/Runtime/MovieSceneTracks/Public/Tracks/MovieSceneRotatorTrack.h b/Engine/Source/Runtime/MovieSceneTracks/Public/Tracks/MovieSceneRotatorTrack.h
new file mode 100644
--- /dev/null
+++ b/Engine/Source/Runtime/MovieSceneTracks/Public/Tracks/MovieSceneRotatorTrack.h
@@ -0,0 +1,28 @@
+// Copyright Epic Games, Inc. All Rights Reserved.
+
+#pragma once
+
+#include "Tracks/MovieScenePropertyTrack.h"
+#include "MovieSceneRotatorTrack.generated.h"
+
+/** Movie scene track that animates an FRotator property */
+UCLASS(MinimalAPI)
+class UMovieSceneRotatorTrack : public UMovieScenePropertyTrack
+{
+	GENERATED_BODY()
+
+public:
+	UMovieSceneRotatorTrack(const FObjectInitializer& InObjectInitializer);
+
+	//~ Begin UMovieSceneTrack
+	virtual bool SupportsType(TSubclassOf<UMovieSceneSection> InSectionClass) const override;
+	virtual UMovieSceneSection* CreateNewSection() override;
+	virtual void AddSection(UMovieSceneSection& InSection) override;
+	virtual const TArray<UMovieSceneSection*>& GetAllSections() const override;
+	virtual bool HasSection(const UMovieSceneSection& InSection) const override;
+	virtual bool IsEmpty() const override;
+	virtual void RemoveAllAnimationData() override;
+	virtual void RemoveSection(UMovieSceneSection& InSection) override;
+	virtual void RemoveSectionAt(int32 InSectionIndex) override;
+	//~ End UMovieSceneTrack
+};
