Sequencer: Skeletal animation evaluation tasks are now dependent upon transforms and object properties

diff --git a/Engine/Source/Runtime/MovieSceneTracks/Private/Systems/MovieSceneSkeletalAnimationSystem.cpp b/Engine/Source/Runtime/MovieSceneTracks/Private/Systems/MovieSceneSkeletalAnimationSystem.cpp
--- a/Engine/Source/Runtime/MovieSceneTracks/Private/Systems/MovieSceneSkeletalAnimationSystem.cpp
+++ b/Engine/Source/Runtime/MovieSceneTracks/Private/Systems/MovieSceneSkeletalAnimationSystem.cpp
@@ -34,6 +34,7 @@
 #include "Systems/MovieSceneComponentTransformSystem.h"
 #include "Systems/MovieSceneQuaternionInterpolationRotationSystem.h"
 #include "Systems/WeightAndEasingEvaluatorSystem.h"
+#include "Systems/MovieSceneObjectPropertySystem.h"
 
 #include "Components/SkeletalMeshComponent.h"
 #include "Animation/AnimationPoseData.h"
@@ -988,6 +989,7 @@ UMovieSceneSkeletalAnimationSystem::UMovieSceneSkeletalAnimationSystem(const FOb
 		DefineImplicitPrerequisite(UWeightAndEasingEvaluatorSystem::StaticClass(), GetClass());
 		DefineImplicitPrerequisite(UMovieSceneComponentTransformSystem::StaticClass(), GetClass());
 		DefineImplicitPrerequisite(UMovieSceneQuaternionInterpolationRotationSystem::StaticClass(), GetClass());
+		DefineImplicitPrerequisite(UMovieSceneObjectPropertySystem::StaticClass(), GetClass());
 	}
 }
 
@@ -998,6 +1000,55 @@ void UMovieSceneSkeletalAnimationSystem::OnSchedulePersistentTasks(UE::MovieScen
 	FBuiltInComponentTypes* BuiltInComponents = FBuiltInComponentTypes::Get();
 	FMovieSceneTracksComponentTypes* TrackComponents = FMovieSceneTracksComponentTypes::Get();
 
+
+	// Facade task that mimics a write dependency to transform results to guarantee that
+	//     skeletal animation evaluation tasks are scheduled after transforms.
+	// @todo: this currently makes all skel anims dependent upon all transforms, which
+	//        is not ideal, but a more granular dependency is not currently possible (or would be prohibitively complex)
+	struct FTransformDependencyTask
+	{
+		static void ForEachAllocation(const FEntityAllocation*,
+			TWriteOptional<double>,TWriteOptional<double>,TWriteOptional<double>,
+			TWriteOptional<double>,TWriteOptional<double>,TWriteOptional<double>,
+			TWriteOptional<double>,TWriteOptional<double>,TWriteOptional<double>)
+		{}
+	};
+
+	// Facade task that mimics a write dependency to object property results to guarantee that
+	//     skeletal animation evaluation tasks are scheduled after any SetMesh calls.
+	// @todo: this currently makes all skel anims dependent upon all transforms, which
+	//        is not ideal, but a more granular dependency is not currently possible (or would be prohibitively complex)
+	struct FWriteObjectResultNoop
+	{
+		static void ForEachAllocation(FEntityAllocationIteratorItem, TWrite<FObjectComponent>)
+		{}
+	};
+
+
+	// Schedule a dummy task that writes to all object results to guarantee that animation eval
+	//    operates after any calls to SetMesh
+	FTaskID WaitForObjectProperties = FEntityTaskBuilder()
+	.Write(BuiltInComponents->ObjectResult)
+	.Schedule_PerAllocation<FWriteObjectResultNoop>(&Linker->EntityManager, TaskScheduler);
+
+
+	// Schedule a dummy task that does nothing but open all transform components for write
+	//  this is used to ensure that previously scheduled transform setter tasks have completed before
+	//  we evaluate skel animations and root motion
+	FTaskID WaitForAllTransforms = FEntityTaskBuilder()
+	.WriteOptional(BuiltInComponents->DoubleResult[0])
+	.WriteOptional(BuiltInComponents->DoubleResult[1])
+	.WriteOptional(BuiltInComponents->DoubleResult[2])
+	.WriteOptional(BuiltInComponents->DoubleResult[3])
+	.WriteOptional(BuiltInComponents->DoubleResult[4])
+	.WriteOptional(BuiltInComponents->DoubleResult[5])
+	.WriteOptional(BuiltInComponents->DoubleResult[6])
+	.WriteOptional(BuiltInComponents->DoubleResult[7])
+	.WriteOptional(BuiltInComponents->DoubleResult[8])
+	.FilterAll({ TrackComponents->ComponentTransform.PropertyTag })
+	.Schedule_PerAllocation<FTransformDependencyTask>(&Linker->EntityManager, TaskScheduler);
+
+
 	FTaskID GatherTask = FEntityTaskBuilder()
 	.ReadEntityIDs()
 	.Read(BuiltInComponents->RootInstanceHandle)
@@ -1018,6 +1069,8 @@ void UMovieSceneSkeletalAnimationSystem::OnSchedulePersistentTasks(UE::MovieScen
 	FTaskID EvaluateTask = TaskScheduler->AddTask<FEvaluateSkeletalAnimations>(Params, Linker, &SystemData);
 
 	TaskScheduler->AddPrerequisite(GatherTask, EvaluateTask);
+	TaskScheduler->AddPrerequisite(WaitForAllTransforms, EvaluateTask);
+	TaskScheduler->AddPrerequisite(WaitForObjectProperties, EvaluateTask);
 }
 
 void UMovieSceneSkeletalAnimationSystem::OnRun(FSystemTaskPrerequisites& InPrerequisites, FSystemSubsequentTasks& Subsequents)
