Sequencer- Expose FComponentMaterialInfo to blueprint/script to properly allow users to build the new upgraded ComponentMaterialTracks in script.

diff --git a/Engine/Source/Runtime/MovieSceneTracks/Public/Tracks/MovieSceneMaterialTrack.h b/Engine/Source/Runtime/MovieSceneTracks/Public/Tracks/MovieSceneMaterialTrack.h
--- a/Engine/Source/Runtime/MovieSceneTracks/Public/Tracks/MovieSceneMaterialTrack.h
+++ b/Engine/Source/Runtime/MovieSceneTracks/Public/Tracks/MovieSceneMaterialTrack.h
@@ -13,8 +13,8 @@
 #endif
 #include "MovieSceneMaterialTrack.generated.h"
 
-UENUM()
-enum class EComponentMaterialType
+UENUM(BlueprintType)
+enum class EComponentMaterialType : uint8
 {
 	/* Empty/Uninitialized*/
 	Empty,
@@ -36,13 +36,13 @@ struct FComponentMaterialInfo
 {
 	GENERATED_BODY()
 
-	UPROPERTY()
+	UPROPERTY(EditAnywhere, BlueprintReadWrite, Category="Sequencer|Material")
 	FName MaterialSlotName;
 
-	UPROPERTY()
-	int MaterialSlotIndex = 0;
+	UPROPERTY(EditAnywhere, BlueprintReadWrite, Category = "Sequencer|Material")
+	int32 MaterialSlotIndex = 0;
 
-	UPROPERTY()
+	UPROPERTY(EditAnywhere, BlueprintReadWrite, Category = "Sequencer|Material")
 	EComponentMaterialType MaterialType = EComponentMaterialType::Empty;
 
 	friend uint32 GetTypeHash(const FComponentMaterialInfo& In)
