Sequencer: Fixed pre-animated state being cached too late for skel animations

Since pre-animated state was only being tracked and cached inside the Evaluation phase, any changes that removed tracked entities could lead to a temporary state whereby no tracking exists for skeletal animation (ie, when manipulating a skeletal animation track) because pre-animated state restoration occurs at the end of the instantiation phase, _before_ the new entity is evaluated. This was leading to unnecessary UI flicker caused by re-initialization of the anim instance. The fix is to move the pre-animated state tracking into the instantiation phase and ensure this runs before the RestorePreAnimatedState system is allowed to run.

diff --git a/Engine/Source/Runtime/MovieSceneTracks/Private/Systems/MovieSceneSkeletalAnimationSystem.cpp b/Engine/Source/Runtime/MovieSceneTracks/Private/Systems/MovieSceneSkeletalAnimationSystem.cpp
--- a/Engine/Source/Runtime/MovieSceneTracks/Private/Systems/MovieSceneSkeletalAnimationSystem.cpp
+++ b/Engine/Source/Runtime/MovieSceneTracks/Private/Systems/MovieSceneSkeletalAnimationSystem.cpp
@@ -35,6 +35,7 @@
 #include "Systems/MovieSceneQuaternionInterpolationRotationSystem.h"
 #include "Systems/WeightAndEasingEvaluatorSystem.h"
 #include "Systems/MovieSceneObjectPropertySystem.h"
+#include "EntitySystem/MovieScenePreAnimatedStateSystem.h"
 
 #include "Components/SkeletalMeshComponent.h"
 #include "Animation/AnimationPoseData.h"
@@ -480,7 +481,6 @@ private:
 	UMovieSceneEntitySystemLinker* Linker;
 	FSkeletalAnimationSystemData* SystemData;
 
-	TSharedPtr<FPreAnimatedSkeletalAnimationStorage> PreAnimatedStorage;
 	TSharedPtr<FPreAnimatedSkeletalAnimationMontageStorage> PreAnimatedMontageStorage;
 	TSharedPtr<FPreAnimatedSkeletalAnimationAnimInstanceStorage> PreAnimatedAnimInstanceStorage;
 
@@ -490,7 +490,6 @@ public:
 		: Linker(InLinker)
 		, SystemData(InSystemData)
 	{
-		PreAnimatedStorage = InLinker->PreAnimatedState.GetOrCreateStorage<FPreAnimatedSkeletalAnimationStorage>();
 		PreAnimatedMontageStorage = InLinker->PreAnimatedState.GetOrCreateStorage<FPreAnimatedSkeletalAnimationMontageStorage>();
 		PreAnimatedAnimInstanceStorage = InLinker->PreAnimatedState.GetOrCreateStorage<FPreAnimatedSkeletalAnimationAnimInstanceStorage>();
 	}
@@ -518,17 +517,6 @@ private:
 			return;
 		}
 
-		// Cache pre-animated state for this bound object before doing anything.
-		// We don't yet track what entities have already started animated vs. entities that just started this frame,
-		// so we just process all the currently active ones. If they are already tracked and have already had their
-		// pre-animated state saved, it these calls will just early return.
-		for (const FActiveSkeletalAnimation& SkeletalAnimation : InSkeletalAnimations.Animations)
-		{
-			PreAnimatedStorage->BeginTrackingEntity(SkeletalAnimation.EntityID, SkeletalAnimation.bWantsRestoreState, SkeletalAnimation.RootInstanceHandle, SkeletalMeshComponent);
-		}
-		FCachePreAnimatedValueParams CacheParams;
-		PreAnimatedStorage->CachePreAnimatedValue(CacheParams, SkeletalMeshComponent);
-
 		// Setup any needed animation nodes for sequencer playback.
 		UAnimInstance* ExistingAnimInstance = GetSourceAnimInstance(SkeletalMeshComponent);
 		bool bWasCreated = false;
@@ -990,6 +978,8 @@ UMovieSceneSkeletalAnimationSystem::UMovieSceneSkeletalAnimationSystem(const FOb
 		DefineImplicitPrerequisite(UMovieSceneComponentTransformSystem::StaticClass(), GetClass());
 		DefineImplicitPrerequisite(UMovieSceneQuaternionInterpolationRotationSystem::StaticClass(), GetClass());
 		DefineImplicitPrerequisite(UMovieSceneObjectPropertySystem::StaticClass(), GetClass());
+
+		DefineImplicitPrerequisite(GetClass(), UMovieSceneRestorePreAnimatedStateSystem::StaticClass());
 	}
 }
 
@@ -1079,16 +1069,26 @@ void UMovieSceneSkeletalAnimationSystem::OnRun(FSystemTaskPrerequisites& InPrere
 
 	const TStatId GatherStatId = GET_STATID(MovieSceneEval_GatherSkeletalAnimations);
 
+	FBuiltInComponentTypes*          BuiltInComponents = FBuiltInComponentTypes::Get();
+	FMovieSceneTracksComponentTypes* TrackComponents   = FMovieSceneTracksComponentTypes::Get();
+
 	TSharedRef<FMovieSceneEntitySystemRunner> Runner = Linker->GetRunner();
 	if (Runner->GetCurrentPhase() == ESystemPhase::Instantiation)
 	{
+		// Begin tracking pre-animated state for all bound skel animation components
+		TSharedPtr<FPreAnimatedSkeletalAnimationStorage> PreAnimatedStorage = Linker->PreAnimatedState.GetOrCreateStorage<FPreAnimatedSkeletalAnimationStorage>();
+
+		struct FTask
+		{
+			FEntityComponentFilter AdditionalFilter;
+		} Task;
+		Task.AdditionalFilter.All({ TrackComponents->SkeletalAnimation, BuiltInComponents->Tags.NeedsLink });
+		PreAnimatedStorage->BeginTrackingAndCachePreAnimatedValuesTask(Linker, Task, BuiltInComponents->BoundObject);
+
 		CleanSystemData();
 		return;
 	}
 
-	FBuiltInComponentTypes* BuiltInComponents = FBuiltInComponentTypes::Get();
-	FMovieSceneTracksComponentTypes* TrackComponents = FMovieSceneTracksComponentTypes::Get();
-
 	FGraphEventRef GatherTask = FEntityTaskBuilder()
 	.ReadEntityIDs()
 	.Read(BuiltInComponents->RootInstanceHandle)
