Sequencer: Skeletal animation sections now always bind directly to a USkeletalMeshComponent

This uses the new BoundObjectResolver component to ensure that bindings always resolve directly to a skeletal mesh component. This simplifies the evaluation system and also makes pre-animated state more robust.

diff --git a/Engine/Source/Runtime/MovieSceneTracks/Private/Sections/MovieSceneSkeletalAnimationSection.cpp b/Engine/Source/Runtime/MovieSceneTracks/Private/Sections/MovieSceneSkeletalAnimationSection.cpp
--- a/Engine/Source/Runtime/MovieSceneTracks/Private/Sections/MovieSceneSkeletalAnimationSection.cpp
+++ b/Engine/Source/Runtime/MovieSceneTracks/Private/Sections/MovieSceneSkeletalAnimationSection.cpp
@@ -19,6 +19,7 @@
 #include "MovieSceneTracksComponentTypes.h"
 #include "Sections/MovieSceneSectionTimingParameters.h"
 #include "Variants/MovieSceneTimeWarpGetter.h"
+#include "Systems/MovieSceneSkeletalAnimationSystem.h"
 
 #include UE_INLINE_GENERATED_CPP_BY_NAME(MovieSceneSkeletalAnimationSection)
 
@@ -488,12 +489,19 @@ void UMovieSceneSkeletalAnimationSection::ImportEntityImpl(UMovieSceneEntitySyst
 	const FMovieSceneTracksComponentTypes* TrackComponents = FMovieSceneTracksComponentTypes::Get();
 
 	const FGuid ObjectBindingID = InParams.GetObjectBindingID();
+
+	if (!ObjectBindingID.IsValid())
+	{
+		return;
+	}
+
 	FMovieSceneSkeletalAnimationComponentData ComponentData { this };
 
 	OutImportedEntity->AddBuilder(
 		FEntityBuilder()
 		.Add(TrackComponents->SkeletalAnimation, ComponentData)
-		.AddConditional(BuiltInComponents->GenericObjectBinding, ObjectBindingID, ObjectBindingID.IsValid())
+		.Add(BuiltInComponents->GenericObjectBinding, ObjectBindingID)
+		.Add(BuiltInComponents->BoundObjectResolver, UMovieSceneSkeletalAnimationSystem::ResolveSkeletalMeshComponentBinding)
 		.AddConditional(BuiltInComponents->WeightChannel, &Params.Weight, Params.Weight.HasAnyData())
 	);
 }
diff --git a/Engine/Source/Runtime/MovieSceneTracks/Private/Systems/MovieSceneSkeletalAnimationSystem.cpp b/Engine/Source/Runtime/MovieSceneTracks/Private/Systems/MovieSceneSkeletalAnimationSystem.cpp
--- a/Engine/Source/Runtime/MovieSceneTracks/Private/Systems/MovieSceneSkeletalAnimationSystem.cpp
+++ b/Engine/Source/Runtime/MovieSceneTracks/Private/Systems/MovieSceneSkeletalAnimationSystem.cpp
@@ -369,7 +369,7 @@ struct FGatherSkeletalAnimations
 			const FMovieSceneSkeletalAnimationParams& AnimParams = AnimSection->Params;
 
 			// Get the bound skeletal mesh component.
-			USkeletalMeshComponent* SkeletalMeshComponent = SkeletalMeshComponentFromObject(BoundObject);
+			USkeletalMeshComponent* SkeletalMeshComponent = CastChecked<USkeletalMeshComponent>(BoundObject);
 			if (!SkeletalMeshComponent || AnimParams.Animation == nullptr)
 			{
 				continue;
@@ -446,31 +446,6 @@ private:
 		bool bIsNotInPIEOrNotPlaying = (RuntimeObject.GetWorld() && !RuntimeObject.GetWorld()->HasBegunPlay()) || PlayerStatus != EMovieScenePlayerStatus::Playing;
 		return GIsEditor && bIsNotInPIEOrNotPlaying;
 	}
-
-	static USkeletalMeshComponent* SkeletalMeshComponentFromObject(UObject* InObject)
-	{
-		// Check if we are bound directly to a skeletal mesh component.
-		USkeletalMeshComponent* SkeletalMeshComponent = Cast<USkeletalMeshComponent>(InObject);
-		if (SkeletalMeshComponent)
-		{
-			return SkeletalMeshComponent;
-		}
-
-		// Then check to see if we are controlling an actor. If so use its first skeletal mesh component.
-		AActor* Actor = Cast<AActor>(InObject);
-		if (!Actor)
-		{
-			if (UChildActorComponent* ChildActorComponent = Cast<UChildActorComponent>(InObject))
-			{
-				Actor = ChildActorComponent->GetChildActor();
-			}
-		}
-		if (Actor)
-		{
-			return Actor->FindComponentByClass<USkeletalMeshComponent>();
-		}
-		return nullptr;
-	}
 };
 
 /** Task for evaluating skeletal animations */
@@ -983,6 +958,31 @@ UMovieSceneSkeletalAnimationSystem::UMovieSceneSkeletalAnimationSystem(const FOb
 	}
 }
 
+UObject* UMovieSceneSkeletalAnimationSystem::ResolveSkeletalMeshComponentBinding(UObject* InObject)
+{
+	// Check if we are bound directly to a skeletal mesh component.
+	USkeletalMeshComponent* SkeletalMeshComponent = Cast<USkeletalMeshComponent>(InObject);
+	if (SkeletalMeshComponent)
+	{
+		return SkeletalMeshComponent;
+	}
+
+	// Then check to see if we are controlling an actor. If so use its first skeletal mesh component.
+	AActor* Actor = Cast<AActor>(InObject);
+	if (!Actor)
+	{
+		if (UChildActorComponent* ChildActorComponent = Cast<UChildActorComponent>(InObject))
+		{
+			Actor = ChildActorComponent->GetChildActor();
+		}
+	}
+	if (Actor)
+	{
+		return Actor->FindComponentByClass<USkeletalMeshComponent>();
+	}
+	return nullptr;
+}
+
 void UMovieSceneSkeletalAnimationSystem::OnSchedulePersistentTasks(UE::MovieScene::IEntitySystemScheduler* TaskScheduler)
 {
 	using namespace UE::MovieScene;
diff --git a/Engine/Source/Runtime/MovieSceneTracks/Public/Systems/MovieSceneSkeletalAnimationSystem.h b/Engine/Source/Runtime/MovieSceneTracks/Public/Systems/MovieSceneSkeletalAnimationSystem.h
--- a/Engine/Source/Runtime/MovieSceneTracks/Public/Systems/MovieSceneSkeletalAnimationSystem.h
+++ b/Engine/Source/Runtime/MovieSceneTracks/Public/Systems/MovieSceneSkeletalAnimationSystem.h
@@ -95,6 +95,8 @@ public:
 
 	UMovieSceneSkeletalAnimationSystem(const FObjectInitializer& ObjInit);
 
+	static UObject* ResolveSkeletalMeshComponentBinding(UObject* InObject);
+
 private:
 
 	virtual void OnSchedulePersistentTasks(UE::MovieScene::IEntitySystemScheduler* TaskScheduler) override;
