Sequencer: Added (and implemented) a virtual method for retrieving timewarp from a section

diff --git a/Engine/Source/Runtime/MovieScene/Private/Sections/MovieSceneSubSection.cpp b/Engine/Source/Runtime/MovieScene/Private/Sections/MovieSceneSubSection.cpp
--- a/Engine/Source/Runtime/MovieScene/Private/Sections/MovieSceneSubSection.cpp
+++ b/Engine/Source/Runtime/MovieScene/Private/Sections/MovieSceneSubSection.cpp
@@ -299,6 +299,11 @@ UMovieSceneSequence* UMovieSceneSubSection::GetSequence() const
 	return SubSequence;
 }
 
+FMovieSceneTimeWarpVariant* UMovieSceneSubSection::GetTimeWarp()
+{
+	return &Parameters.TimeScale;
+}
+
 #if WITH_EDITOR
 void UMovieSceneSubSection::PreEditChange(FProperty* PropertyAboutToChange)
 {
diff --git a/Engine/Source/Runtime/MovieScene/Private/Sections/MovieSceneTimeWarpSection.cpp b/Engine/Source/Runtime/MovieScene/Private/Sections/MovieSceneTimeWarpSection.cpp
--- a/Engine/Source/Runtime/MovieScene/Private/Sections/MovieSceneTimeWarpSection.cpp
+++ b/Engine/Source/Runtime/MovieScene/Private/Sections/MovieSceneTimeWarpSection.cpp
@@ -41,6 +41,11 @@ EMovieSceneChannelProxyType UMovieSceneTimeWarpSection::CacheChannelProxy()
 	return EMovieSceneChannelProxyType::Dynamic;
 }
 
+FMovieSceneTimeWarpVariant* UMovieSceneTimeWarpSection::GetTimeWarp()
+{
+	return &TimeWarp;
+}
+
 #if WITH_EDITOR
 
 bool UMovieSceneTimeWarpSection::Modify(bool bAlwaysMarkDirty)
diff --git a/Engine/Source/Runtime/MovieScene/Public/MovieSceneSection.h b/Engine/Source/Runtime/MovieScene/Public/MovieSceneSection.h
--- a/Engine/Source/Runtime/MovieScene/Public/MovieSceneSection.h
+++ b/Engine/Source/Runtime/MovieScene/Public/MovieSceneSection.h
@@ -54,6 +54,7 @@ struct FMovieSceneSequenceHierarchy;
 struct FMovieSceneSequenceID;
 struct FPropertyChangedEvent;
 struct FQualifiedFrameTime;
+struct FMovieSceneTimeWarpVariant;
 
 namespace UE::MovieScene
 {
@@ -483,6 +484,14 @@ public:
 	 */
 	MOVIESCENE_API virtual TSharedPtr<FStructOnScope> GetKeyStruct(TArrayView<const FKeyHandle> KeyHandles);
 
+	/**
+	 * Retrieve the time warp variant for this section, if it has one
+	 */
+	virtual FMovieSceneTimeWarpVariant* GetTimeWarp()
+	{
+		return nullptr;
+	}
+
 	/**
 	 * Gets all snap times for this section
 	 *
diff --git a/Engine/Source/Runtime/MovieScene/Public/Sections/MovieSceneSubSection.h b/Engine/Source/Runtime/MovieScene/Public/Sections/MovieSceneSubSection.h
--- a/Engine/Source/Runtime/MovieScene/Public/Sections/MovieSceneSubSection.h
+++ b/Engine/Source/Runtime/MovieScene/Public/Sections/MovieSceneSubSection.h
@@ -163,6 +163,7 @@ public:
 	virtual TOptional<FFrameTime> GetOffsetTime() const override { return TOptional<FFrameTime>(FFrameTime(Parameters.StartFrameOffset)); }
 	MOVIESCENE_API virtual void GetSnapTimes(TArray<FFrameNumber>& OutSnapTimes, bool bGetSectionBorders) const override;
 	MOVIESCENE_API virtual void MigrateFrameTimes(FFrameRate SourceRate, FFrameRate DestinationRate) override;
+	MOVIESCENE_API virtual FMovieSceneTimeWarpVariant* GetTimeWarp() override;
 
 protected:
 
diff --git a/Engine/Source/Runtime/MovieScene/Public/Sections/MovieSceneTimeWarpSection.h b/Engine/Source/Runtime/MovieScene/Public/Sections/MovieSceneTimeWarpSection.h
--- a/Engine/Source/Runtime/MovieScene/Public/Sections/MovieSceneTimeWarpSection.h
+++ b/Engine/Source/Runtime/MovieScene/Public/Sections/MovieSceneTimeWarpSection.h
@@ -34,4 +34,5 @@ private:
 	virtual void PostEditChangeProperty(FPropertyChangedEvent& PropertyChangedEvent) override;
 #endif
 	virtual EMovieSceneChannelProxyType CacheChannelProxy() override;
+	virtual FMovieSceneTimeWarpVariant* GetTimeWarp() override;
 };
diff --git a/Engine/Source/Runtime/MovieSceneTracks/Private/Sections/MovieSceneSkeletalAnimationSection.cpp b/Engine/Source/Runtime/MovieSceneTracks/Private/Sections/MovieSceneSkeletalAnimationSection.cpp
--- a/Engine/Source/Runtime/MovieSceneTracks/Private/Sections/MovieSceneSkeletalAnimationSection.cpp
+++ b/Engine/Source/Runtime/MovieSceneTracks/Private/Sections/MovieSceneSkeletalAnimationSection.cpp
@@ -481,6 +481,11 @@ void UMovieSceneSkeletalAnimationSection::SetEndFrame(TRangeBound<FFrameNumber>
 	}
 }
 
+FMovieSceneTimeWarpVariant* UMovieSceneSkeletalAnimationSection::GetTimeWarp()
+{
+	return &Params.PlayRate;
+}
+
 void UMovieSceneSkeletalAnimationSection::ImportEntityImpl(UMovieSceneEntitySystemLinker* EntityLinker, const FEntityImportParams& InParams, FImportedEntity* OutImportedEntity)
 {
 	using namespace UE::MovieScene;
diff --git a/Engine/Source/Runtime/MovieSceneTracks/Public/Sections/MovieSceneSkeletalAnimationSection.h b/Engine/Source/Runtime/MovieSceneTracks/Public/Sections/MovieSceneSkeletalAnimationSection.h
--- a/Engine/Source/Runtime/MovieSceneTracks/Public/Sections/MovieSceneSkeletalAnimationSection.h
+++ b/Engine/Source/Runtime/MovieSceneTracks/Public/Sections/MovieSceneSkeletalAnimationSection.h
@@ -123,6 +123,7 @@ public:
 	virtual void SetRange(const TRange<FFrameNumber>& NewRange) override;
 	virtual void SetStartFrame(TRangeBound<FFrameNumber> NewStartFrame) override;
 	virtual void SetEndFrame(TRangeBound<FFrameNumber> NewEndFrame)override;
+	virtual FMovieSceneTimeWarpVariant* GetTimeWarp() override;
 
 	MOVIESCENETRACKS_API void DeleteChannels(TArrayView<const FName> ChannelNames);
 
