Sequencer: Fix GetKeyStruct to allow for when channels are overridden.

If a channel is overriden, the channel data doesn't exist in the proxy, so we can't hardcode the index to 0,1,2...8. Check for the existance of the channel, and get the channel by index.

diff --git a/Engine/Source/Runtime/MovieSceneTracks/Private/Sections/MovieScene3DTransformSection.cpp b/Engine/Source/Runtime/MovieSceneTracks/Private/Sections/MovieScene3DTransformSection.cpp
--- a/Engine/Source/Runtime/MovieSceneTracks/Private/Sections/MovieScene3DTransformSection.cpp
+++ b/Engine/Source/Runtime/MovieSceneTracks/Private/Sections/MovieScene3DTransformSection.cpp
@@ -843,29 +843,78 @@ TSharedPtr<FStructOnScope> UMovieScene3DTransformSection::GetKeyStruct(TArrayVie
 		GetChannelProxy();
 	}
 
-	TArrayView<FMovieSceneDoubleChannel* const> DoubleChannels = ChannelProxy->GetChannels<FMovieSceneDoubleChannel>();
+	FMovieSceneChannelHandle LocationXChannel, LocationYChannel, LocationZChannel;
+	FMovieSceneChannelHandle RotationXChannel, RotationYChannel, RotationZChannel;
+	FMovieSceneChannelHandle ScaleXChannel, ScaleYChannel, ScaleZChannel;
 
-	TOptional<TTuple<FKeyHandle, FFrameNumber>> LocationKeys[3] = {
-		FMovieSceneChannelValueHelper::FindFirstKey(DoubleChannels[0], KeyHandles),
-		FMovieSceneChannelValueHelper::FindFirstKey(DoubleChannels[1], KeyHandles),
-		FMovieSceneChannelValueHelper::FindFirstKey(DoubleChannels[2], KeyHandles)
-	};
+#if WITH_EDITOR
+	LocationXChannel = ChannelProxy->GetChannelByName<FMovieSceneDoubleChannel>("Location.X");
+	LocationYChannel = ChannelProxy->GetChannelByName<FMovieSceneDoubleChannel>("Location.Y");
+	LocationZChannel = ChannelProxy->GetChannelByName<FMovieSceneDoubleChannel>("Location.Z");
 
-	TOptional<TTuple<FKeyHandle, FFrameNumber>> RotationKeys[3] = {
-		FMovieSceneChannelValueHelper::FindFirstKey(DoubleChannels[3], KeyHandles),
-		FMovieSceneChannelValueHelper::FindFirstKey(DoubleChannels[4], KeyHandles),
-		FMovieSceneChannelValueHelper::FindFirstKey(DoubleChannels[5], KeyHandles)
-	};
+	RotationXChannel = ChannelProxy->GetChannelByName<FMovieSceneDoubleChannel>("Rotation.X");
+	RotationYChannel = ChannelProxy->GetChannelByName<FMovieSceneDoubleChannel>("Rotation.Y");
+	RotationZChannel = ChannelProxy->GetChannelByName<FMovieSceneDoubleChannel>("Rotation.Z");
 
-	TOptional<TTuple<FKeyHandle, FFrameNumber>> ScaleKeys[3] = {
-		FMovieSceneChannelValueHelper::FindFirstKey(DoubleChannels[6], KeyHandles),
-		FMovieSceneChannelValueHelper::FindFirstKey(DoubleChannels[7], KeyHandles),
-		FMovieSceneChannelValueHelper::FindFirstKey(DoubleChannels[8], KeyHandles)
-	};
+	ScaleXChannel = ChannelProxy->GetChannelByName<FMovieSceneDoubleChannel>("Scale.X");
+	ScaleYChannel = ChannelProxy->GetChannelByName<FMovieSceneDoubleChannel>("Scale.Y");
+	ScaleZChannel = ChannelProxy->GetChannelByName<FMovieSceneDoubleChannel>("Scale.Z");
+#endif
+
+	int32 AnyLocationKeys = 0;
+	int32 AnyRotationKeys = 0;
+	int32 AnyScaleKeys = 0;
+
+	TOptional<TTuple<FKeyHandle, FFrameNumber>> LocationKeys[3];
+	if (LocationXChannel.Get())
+	{
+		LocationKeys[0] = FMovieSceneChannelValueHelper::FindFirstKey(static_cast<FMovieSceneDoubleChannel*>(LocationXChannel.Get()), KeyHandles);
+		AnyLocationKeys++;
+	}
+	if (LocationYChannel.Get())
+	{
+		LocationKeys[1] = FMovieSceneChannelValueHelper::FindFirstKey(static_cast<FMovieSceneDoubleChannel*>(LocationYChannel.Get()), KeyHandles);
+		AnyLocationKeys++;
+	}
+	if (LocationZChannel.Get())
+	{
+		LocationKeys[2] = FMovieSceneChannelValueHelper::FindFirstKey(static_cast<FMovieSceneDoubleChannel*>(LocationZChannel.Get()), KeyHandles);
+		AnyLocationKeys++;
+	}
+
+	TOptional<TTuple<FKeyHandle, FFrameNumber>> RotationKeys[3];
+	if (RotationXChannel.Get())
+	{
+		RotationKeys[0] = FMovieSceneChannelValueHelper::FindFirstKey(static_cast<FMovieSceneDoubleChannel*>(RotationXChannel.Get()), KeyHandles);
+		AnyRotationKeys++;
+	}
+	if (RotationYChannel.Get())
+	{
+		RotationKeys[1] = FMovieSceneChannelValueHelper::FindFirstKey(static_cast<FMovieSceneDoubleChannel*>(RotationYChannel.Get()), KeyHandles);
+		AnyRotationKeys++;
+	}
+	if (RotationZChannel.Get())
+	{
+		RotationKeys[2] = FMovieSceneChannelValueHelper::FindFirstKey(static_cast<FMovieSceneDoubleChannel*>(RotationZChannel.Get()), KeyHandles);
+		AnyRotationKeys++;
+	}
 
-	const int32 AnyLocationKeys = Algo::AnyOf(LocationKeys);
-	const int32 AnyRotationKeys = Algo::AnyOf(RotationKeys);
-	const int32 AnyScaleKeys =    Algo::AnyOf(ScaleKeys);
+	TOptional<TTuple<FKeyHandle, FFrameNumber>> ScaleKeys[3];
+	if (ScaleXChannel.Get())
+	{
+		ScaleKeys[0] = FMovieSceneChannelValueHelper::FindFirstKey(static_cast<FMovieSceneDoubleChannel*>(ScaleXChannel.Get()), KeyHandles);
+		AnyScaleKeys++;
+	}
+	if (ScaleYChannel.Get())
+	{
+		ScaleKeys[1] = FMovieSceneChannelValueHelper::FindFirstKey(static_cast<FMovieSceneDoubleChannel*>(ScaleYChannel.Get()), KeyHandles);
+		AnyScaleKeys++;
+	}
+	if (ScaleZChannel.Get())
+	{
+		ScaleKeys[2] = FMovieSceneChannelValueHelper::FindFirstKey(static_cast<FMovieSceneDoubleChannel*>(ScaleZChannel.Get()), KeyHandles);
+		AnyScaleKeys++;
+	}
 
 	// do we have multiple keys on multiple parts of the transform?
 	if (AnyLocationKeys + AnyRotationKeys + AnyScaleKeys > 1)
@@ -873,17 +922,50 @@ TSharedPtr<FStructOnScope> UMovieScene3DTransformSection::GetKeyStruct(TArrayVie
 		TSharedRef<FStructOnScope> KeyStruct = MakeShareable(new FStructOnScope(FMovieScene3DTransformKeyStruct::StaticStruct()));
 		auto Struct = (FMovieScene3DTransformKeyStruct*)KeyStruct->GetStructMemory();
 
-		Struct->KeyStructInterop.Add(FMovieSceneChannelValueHelper(ChannelProxy->MakeHandle<FMovieSceneDoubleChannel>(0), &Struct->Location.X,     LocationKeys[0]));
-		Struct->KeyStructInterop.Add(FMovieSceneChannelValueHelper(ChannelProxy->MakeHandle<FMovieSceneDoubleChannel>(1), &Struct->Location.Y,     LocationKeys[1]));
-		Struct->KeyStructInterop.Add(FMovieSceneChannelValueHelper(ChannelProxy->MakeHandle<FMovieSceneDoubleChannel>(2), &Struct->Location.Z,     LocationKeys[2]));
+		if (LocationKeys[0].IsSet())
+		{
+			Struct->KeyStructInterop.Add(FMovieSceneChannelValueHelper(ChannelProxy->MakeHandle<FMovieSceneDoubleChannel>(LocationXChannel.GetChannelIndex()), &Struct->Location.X,     LocationKeys[0]));
+		}
+		
+		if (LocationKeys[1].IsSet())
+		{
+			Struct->KeyStructInterop.Add(FMovieSceneChannelValueHelper(ChannelProxy->MakeHandle<FMovieSceneDoubleChannel>(LocationYChannel.GetChannelIndex()), &Struct->Location.Y,     LocationKeys[1]));
+		}
+		
+		if (LocationKeys[2].IsSet())
+		{
+			Struct->KeyStructInterop.Add(FMovieSceneChannelValueHelper(ChannelProxy->MakeHandle<FMovieSceneDoubleChannel>(LocationZChannel.GetChannelIndex()), &Struct->Location.Z,     LocationKeys[2]));
+		}
 
-		Struct->KeyStructInterop.Add(FMovieSceneChannelValueHelper(ChannelProxy->MakeHandle<FMovieSceneDoubleChannel>(3), &Struct->Rotation.Roll,  RotationKeys[0]));
-		Struct->KeyStructInterop.Add(FMovieSceneChannelValueHelper(ChannelProxy->MakeHandle<FMovieSceneDoubleChannel>(4), &Struct->Rotation.Pitch, RotationKeys[1]));
-		Struct->KeyStructInterop.Add(FMovieSceneChannelValueHelper(ChannelProxy->MakeHandle<FMovieSceneDoubleChannel>(5), &Struct->Rotation.Yaw,   RotationKeys[2]));
+		if (RotationKeys[0].IsSet())
+		{
+			Struct->KeyStructInterop.Add(FMovieSceneChannelValueHelper(ChannelProxy->MakeHandle<FMovieSceneDoubleChannel>(RotationXChannel.GetChannelIndex()), &Struct->Rotation.Roll,  RotationKeys[0]));
+		}
 
-		Struct->KeyStructInterop.Add(FMovieSceneChannelValueHelper(ChannelProxy->MakeHandle<FMovieSceneDoubleChannel>(6), &Struct->Scale.X,        ScaleKeys[0]));
-		Struct->KeyStructInterop.Add(FMovieSceneChannelValueHelper(ChannelProxy->MakeHandle<FMovieSceneDoubleChannel>(7), &Struct->Scale.Y,        ScaleKeys[1]));
-		Struct->KeyStructInterop.Add(FMovieSceneChannelValueHelper(ChannelProxy->MakeHandle<FMovieSceneDoubleChannel>(8), &Struct->Scale.Z,        ScaleKeys[2]));
+		if (RotationKeys[1].IsSet())
+		{
+			Struct->KeyStructInterop.Add(FMovieSceneChannelValueHelper(ChannelProxy->MakeHandle<FMovieSceneDoubleChannel>(RotationYChannel.GetChannelIndex()), &Struct->Rotation.Pitch, RotationKeys[1]));
+		}
+		
+		if (RotationKeys[2].IsSet())
+		{
+			Struct->KeyStructInterop.Add(FMovieSceneChannelValueHelper(ChannelProxy->MakeHandle<FMovieSceneDoubleChannel>(RotationZChannel.GetChannelIndex()), &Struct->Rotation.Yaw,   RotationKeys[2]));
+		}
+
+		if (ScaleKeys[0].IsSet())
+		{
+			Struct->KeyStructInterop.Add(FMovieSceneChannelValueHelper(ChannelProxy->MakeHandle<FMovieSceneDoubleChannel>(ScaleXChannel.GetChannelIndex()), &Struct->Scale.X,        ScaleKeys[0]));
+		}
+		
+		if (ScaleKeys[1].IsSet())
+		{
+			Struct->KeyStructInterop.Add(FMovieSceneChannelValueHelper(ChannelProxy->MakeHandle<FMovieSceneDoubleChannel>(ScaleYChannel.GetChannelIndex()), &Struct->Scale.Y,        ScaleKeys[1]));
+		}
+		
+		if (ScaleKeys[2].IsSet())
+		{
+			Struct->KeyStructInterop.Add(FMovieSceneChannelValueHelper(ChannelProxy->MakeHandle<FMovieSceneDoubleChannel>(ScaleZChannel.GetChannelIndex()), &Struct->Scale.Z,        ScaleKeys[2]));
+		}
 
 		Struct->KeyStructInterop.SetStartingValues();
 		Struct->Time = Struct->KeyStructInterop.GetUnifiedKeyTime().Get(0);
@@ -895,9 +977,20 @@ TSharedPtr<FStructOnScope> UMovieScene3DTransformSection::GetKeyStruct(TArrayVie
 		TSharedRef<FStructOnScope> KeyStruct = MakeShareable(new FStructOnScope(FMovieScene3DLocationKeyStruct::StaticStruct()));
 		auto Struct = (FMovieScene3DLocationKeyStruct*)KeyStruct->GetStructMemory();
 
-		Struct->KeyStructInterop.Add(FMovieSceneChannelValueHelper(ChannelProxy->MakeHandle<FMovieSceneDoubleChannel>(0), &Struct->Location.X,     LocationKeys[0]));
-		Struct->KeyStructInterop.Add(FMovieSceneChannelValueHelper(ChannelProxy->MakeHandle<FMovieSceneDoubleChannel>(1), &Struct->Location.Y,     LocationKeys[1]));
-		Struct->KeyStructInterop.Add(FMovieSceneChannelValueHelper(ChannelProxy->MakeHandle<FMovieSceneDoubleChannel>(2), &Struct->Location.Z,     LocationKeys[2]));
+		if (LocationKeys[0].IsSet())
+		{
+			Struct->KeyStructInterop.Add(FMovieSceneChannelValueHelper(ChannelProxy->MakeHandle<FMovieSceneDoubleChannel>(LocationXChannel.GetChannelIndex()), &Struct->Location.X,     LocationKeys[0]));
+		}
+		
+		if (LocationKeys[1].IsSet())
+		{
+			Struct->KeyStructInterop.Add(FMovieSceneChannelValueHelper(ChannelProxy->MakeHandle<FMovieSceneDoubleChannel>(LocationYChannel.GetChannelIndex()), &Struct->Location.Y,     LocationKeys[1]));
+		}
+		
+		if (LocationKeys[2].IsSet())
+		{
+			Struct->KeyStructInterop.Add(FMovieSceneChannelValueHelper(ChannelProxy->MakeHandle<FMovieSceneDoubleChannel>(LocationZChannel.GetChannelIndex()), &Struct->Location.Z,     LocationKeys[2]));
+		}
 
 		Struct->KeyStructInterop.SetStartingValues();
 		Struct->Time = Struct->KeyStructInterop.GetUnifiedKeyTime().Get(0);
@@ -909,9 +1002,20 @@ TSharedPtr<FStructOnScope> UMovieScene3DTransformSection::GetKeyStruct(TArrayVie
 		TSharedRef<FStructOnScope> KeyStruct = MakeShareable(new FStructOnScope(FMovieScene3DRotationKeyStruct::StaticStruct()));
 		auto Struct = (FMovieScene3DRotationKeyStruct*)KeyStruct->GetStructMemory();
 
-		Struct->KeyStructInterop.Add(FMovieSceneChannelValueHelper(ChannelProxy->MakeHandle<FMovieSceneDoubleChannel>(3), &Struct->Rotation.Roll,  RotationKeys[0]));
-		Struct->KeyStructInterop.Add(FMovieSceneChannelValueHelper(ChannelProxy->MakeHandle<FMovieSceneDoubleChannel>(4), &Struct->Rotation.Pitch, RotationKeys[1]));
-		Struct->KeyStructInterop.Add(FMovieSceneChannelValueHelper(ChannelProxy->MakeHandle<FMovieSceneDoubleChannel>(5), &Struct->Rotation.Yaw,   RotationKeys[2]));
+		if (RotationKeys[0].IsSet())
+		{
+			Struct->KeyStructInterop.Add(FMovieSceneChannelValueHelper(ChannelProxy->MakeHandle<FMovieSceneDoubleChannel>(RotationXChannel.GetChannelIndex()), &Struct->Rotation.Roll,  RotationKeys[0]));
+		}
+		
+		if (RotationKeys[1].IsSet())
+		{
+			Struct->KeyStructInterop.Add(FMovieSceneChannelValueHelper(ChannelProxy->MakeHandle<FMovieSceneDoubleChannel>(RotationYChannel.GetChannelIndex()), &Struct->Rotation.Pitch, RotationKeys[1]));
+		}
+		
+		if (RotationKeys[2].IsSet())
+		{
+			Struct->KeyStructInterop.Add(FMovieSceneChannelValueHelper(ChannelProxy->MakeHandle<FMovieSceneDoubleChannel>(RotationZChannel.GetChannelIndex()), &Struct->Rotation.Yaw,   RotationKeys[2]));
+		}
 
 		Struct->KeyStructInterop.SetStartingValues();
 		Struct->Time = Struct->KeyStructInterop.GetUnifiedKeyTime().Get(0);
@@ -923,10 +1027,20 @@ TSharedPtr<FStructOnScope> UMovieScene3DTransformSection::GetKeyStruct(TArrayVie
 		TSharedRef<FStructOnScope> KeyStruct = MakeShareable(new FStructOnScope(FMovieScene3DScaleKeyStruct::StaticStruct()));
 		auto Struct = (FMovieScene3DScaleKeyStruct*)KeyStruct->GetStructMemory();
 
-		Struct->KeyStructInterop.Add(FMovieSceneChannelValueHelper(ChannelProxy->MakeHandle<FMovieSceneDoubleChannel>(6), &Struct->Scale.X,        ScaleKeys[0]));
-		Struct->KeyStructInterop.Add(FMovieSceneChannelValueHelper(ChannelProxy->MakeHandle<FMovieSceneDoubleChannel>(7), &Struct->Scale.Y,        ScaleKeys[1]));
-		Struct->KeyStructInterop.Add(FMovieSceneChannelValueHelper(ChannelProxy->MakeHandle<FMovieSceneDoubleChannel>(8), &Struct->Scale.Z,        ScaleKeys[2]));
-
+		if (ScaleKeys[0].IsSet())
+		{
+			Struct->KeyStructInterop.Add(FMovieSceneChannelValueHelper(ChannelProxy->MakeHandle<FMovieSceneDoubleChannel>(ScaleXChannel.GetChannelIndex()), &Struct->Scale.X,        ScaleKeys[0]));
+		}
+		
+		if (ScaleKeys[1].IsSet())
+		{
+			Struct->KeyStructInterop.Add(FMovieSceneChannelValueHelper(ChannelProxy->MakeHandle<FMovieSceneDoubleChannel>(ScaleYChannel.GetChannelIndex()), &Struct->Scale.Y,        ScaleKeys[1]));
+		}
+		
+		if (ScaleKeys[2].IsSet())
+		{
+			Struct->KeyStructInterop.Add(FMovieSceneChannelValueHelper(ChannelProxy->MakeHandle<FMovieSceneDoubleChannel>(ScaleZChannel.GetChannelIndex()), &Struct->Scale.Z,        ScaleKeys[2]));
+		}
 		Struct->KeyStructInterop.SetStartingValues();
 		Struct->Time = Struct->KeyStructInterop.GetUnifiedKeyTime().Get(0);
 		return KeyStruct;
