Sequencer: Fix crash on condition creation.

diff --git a/Engine/Source/Runtime/MovieSceneTracks/Private/Conditions/MovieSceneDirectorBlueprintCondition.cpp b/Engine/Source/Runtime/MovieSceneTracks/Private/Conditions/MovieSceneDirectorBlueprintCondition.cpp
--- a/Engine/Source/Runtime/MovieSceneTracks/Private/Conditions/MovieSceneDirectorBlueprintCondition.cpp
+++ b/Engine/Source/Runtime/MovieSceneTracks/Private/Conditions/MovieSceneDirectorBlueprintCondition.cpp
@@ -77,7 +77,7 @@ bool FMovieSceneDirectorBlueprintConditionInvoker::InvokeDirectorBlueprintCondit
 	{
 		FProperty* LocalProp = *It;
 		checkSlow(LocalProp);
-		if (!LocalProp->HasAnyPropertyFlags(CPF_ZeroConstructor))
+		if (!LocalProp->HasAnyPropertyFlags(CPF_ZeroConstructor) && LocalProp->HasAllPropertyFlags(CPF_Parm))
 		{
 			LocalProp->InitializeValue_InContainer(Parameters);
 		}
@@ -115,7 +115,13 @@ bool FMovieSceneDirectorBlueprintConditionInvoker::InvokeDirectorBlueprintCondit
 	// Destroy parameters.
 	for (TFieldIterator<FProperty> It(ConditionFunc); It; ++It)
 	{
-		It->DestroyValue_InContainer(Parameters);
+		FProperty* LocalProp = *It;
+		checkSlow(LocalProp);
+
+		if (LocalProp->HasAllPropertyFlags(CPF_Parm))
+		{
+			It->DestroyValue_InContainer(Parameters);
+		}
 	}
 
 	return Result;
