Sequencer Conditions- Fix the 'Call in Editor' checkbox so it actually affects whether the condition is called. Now if call in editor is unchecked, the condition will return false always in editor contexts.

For some reason the call in editor metadata is not being respected on UFUNCTIONS that don't turn into event nodes in the blueprint, so I check it higher up in the condition specific code.

diff --git a/Engine/Source/Runtime/MovieSceneTracks/Private/Conditions/MovieSceneDirectorBlueprintCondition.cpp b/Engine/Source/Runtime/MovieSceneTracks/Private/Conditions/MovieSceneDirectorBlueprintCondition.cpp
--- a/Engine/Source/Runtime/MovieSceneTracks/Private/Conditions/MovieSceneDirectorBlueprintCondition.cpp
+++ b/Engine/Source/Runtime/MovieSceneTracks/Private/Conditions/MovieSceneDirectorBlueprintCondition.cpp
@@ -9,6 +9,9 @@
 #include "MovieSceneSequence.h"
 #include "UObject/UnrealType.h"
 #include "MovieScene.h"
+#if WITH_EDITOR
+#include "EdGraphSchema_K2.h"
+#endif
 
 bool FMovieSceneDirectorBlueprintConditionInvoker::EvaluateDirectorBlueprintCondition(FGuid BindingGuid, FMovieSceneSequenceID SequenceID, TSharedRef<const UE::MovieScene::FSharedPlaybackState> SharedPlaybackState, const FMovieSceneDirectorBlueprintConditionData& DirectorBlueprintConditionData)
 {
@@ -56,7 +59,7 @@ bool FMovieSceneDirectorBlueprintConditionInvoker::EvaluateDirectorBlueprintCond
 
 bool FMovieSceneDirectorBlueprintConditionInvoker::InvokeDirectorBlueprintCondition(UObject* DirectorInstance, const FMovieSceneDirectorBlueprintConditionData& DirectorBlueprintConditionData, const FMovieSceneConditionContext& ConditionContext)
 {
-	bool Result = true;
+	bool Result = false;
 
 	// Do some basic checks.
 	UFunction* ConditionFunc = DirectorBlueprintConditionData.Function.Get();
@@ -65,6 +68,17 @@ bool FMovieSceneDirectorBlueprintConditionInvoker::InvokeDirectorBlueprintCondit
 		return Result;
 	}
 
+#if WITH_EDITOR
+	// Not sure why this isn't being checked further down, but check manually here
+	if (!ConditionFunc->HasMetaData(FBlueprintMetadata::MD_CallInEditor) || ConditionFunc->GetMetaData(FBlueprintMetadata::MD_CallInEditor) != TEXT("true"))
+	{
+		if (ConditionContext.WorldContext && ConditionContext.WorldContext->GetWorld()->IsEditorWorld())
+		{
+			return Result;
+		}
+	}
+#endif
+
 	// Parse all function parameters.
 	uint8* Parameters = (uint8*)FMemory_Alloca(ConditionFunc->ParmsSize + ConditionFunc->MinAlignment);
 	Parameters = Align(Parameters, ConditionFunc->MinAlignment);
