Sequencer- Scalability Condition.

Add a scability condition that lets the user pick a scalability group, and what levels are valid for that group for the condition to pass successfully.

diff --git a/Engine/Source/Runtime/MovieSceneTracks/Private/Conditions/MovieSceneScalabilityCondition.cpp b/Engine/Source/Runtime/MovieSceneTracks/Private/Conditions/MovieSceneScalabilityCondition.cpp
new file mode 100644
--- /dev/null
+++ b/Engine/Source/Runtime/MovieSceneTracks/Private/Conditions/MovieSceneScalabilityCondition.cpp
@@ -0,0 +1,68 @@
+// Copyright Epic Games, Inc. All Rights Reserved.
+
+#include "Conditions/MovieSceneScalabilityCondition.h"
+#include "CoreMinimal.h"
+#include "Scalability.h"
+
+bool UMovieSceneScalabilityCondition::EvaluateConditionInternal(FGuid BindingGuid, FMovieSceneSequenceID SequenceID, TSharedRef<const UE::MovieScene::FSharedPlaybackState> SharedPlaybackState) const
+{
+	using namespace Scalability;
+	const FQualityLevels& QualityLevels = GetQualityLevels();
+	
+	int32 CurrentLevel = 0;
+
+	switch (Group)
+	{
+		 case EMovieSceneScalabilityConditionGroup::ViewDistance:
+			CurrentLevel = QualityLevels.ViewDistanceQuality;
+			break;
+		 case EMovieSceneScalabilityConditionGroup::AntiAliasing:
+			CurrentLevel = QualityLevels.AntiAliasingQuality;
+			break;
+		 case EMovieSceneScalabilityConditionGroup::Shadow:
+			CurrentLevel = QualityLevels.ShadowQuality;
+			break;
+		 case EMovieSceneScalabilityConditionGroup::GlobalIllumination:
+			CurrentLevel = QualityLevels.GlobalIlluminationQuality;
+			break;
+		 case EMovieSceneScalabilityConditionGroup::Reflection:
+			CurrentLevel = QualityLevels.ReflectionQuality;
+			break;
+		 case EMovieSceneScalabilityConditionGroup::PostProcess:
+			CurrentLevel = QualityLevels.PostProcessQuality;
+			break;
+		 case EMovieSceneScalabilityConditionGroup::Texture:
+			CurrentLevel = QualityLevels.TextureQuality;
+			break;
+		 case EMovieSceneScalabilityConditionGroup::Effects:
+			CurrentLevel = QualityLevels.EffectsQuality;
+			break;
+		 case EMovieSceneScalabilityConditionGroup::Foliage:
+			CurrentLevel = QualityLevels.FoliageQuality;
+			break;
+		 case EMovieSceneScalabilityConditionGroup::Shading:
+			CurrentLevel = QualityLevels.ShadingQuality;
+			break;
+		 case EMovieSceneScalabilityConditionGroup::Landscape:
+			CurrentLevel = QualityLevels.LandscapeQuality;
+			break;
+	}
+
+	int32 CompareValue = (int32)Level;
+
+	switch (Operator)
+	{
+		case EMovieSceneScalabilityConditionOperator::LessThan:
+			return CurrentLevel < CompareValue;
+		case EMovieSceneScalabilityConditionOperator::LessThanOrEqualTo:
+			return CurrentLevel <= CompareValue;
+		case EMovieSceneScalabilityConditionOperator::EqualTo:
+			return CurrentLevel == CompareValue;
+		case EMovieSceneScalabilityConditionOperator::GreaterThanOrEqualTo:
+			return CurrentLevel >= CompareValue;
+		case EMovieSceneScalabilityConditionOperator::GreaterThan:
+			return CurrentLevel > CompareValue;
+	}
+
+	return false;
+}
diff --git a/Engine/Source/Runtime/MovieSceneTracks/Public/Conditions/MovieSceneScalabilityCondition.h b/Engine/Source/Runtime/MovieSceneTracks/Public/Conditions/MovieSceneScalabilityCondition.h
new file mode 100644
--- /dev/null
+++ b/Engine/Source/Runtime/MovieSceneTracks/Public/Conditions/MovieSceneScalabilityCondition.h
@@ -0,0 +1,92 @@
+// Copyright Epic Games, Inc. All Rights Reserved.
+
+#pragma once
+
+#include "CoreTypes.h"
+#include "Conditions/MovieSceneCondition.h"
+#include "MovieSceneScalabilityCondition.generated.h"
+
+namespace UE
+{
+	namespace MovieScene
+	{
+		struct FSharedPlaybackState;
+	}
+}
+
+
+// The below is a bit hardcoded to try and match how scalability settings are set up in Scalability.h.
+// This is because unfortunately scalability settings at their core are not very extensible or data-driven 
+// and so it's difficult to do this in a data driven way. So I've made enums here to make the setup user friendly here,
+// and then do the mapping in code. If scalability gets re-architected, this will need be to be updated to match.
+
+UENUM()
+enum class EMovieSceneScalabilityConditionGroup : uint8
+{
+	ViewDistance,
+	AntiAliasing,
+	Shadow,
+	GlobalIllumination,
+	Reflection,
+	PostProcess,
+	Texture,
+	Effects,
+	Foliage,
+	Shading,
+	Landscape
+};
+
+
+UENUM()
+enum class EMovieSceneScalabilityConditionOperator : uint8
+{
+	LessThan,
+	LessThanOrEqualTo,
+	EqualTo,
+	GreaterThanOrEqualTo,
+	GreaterThan
+};
+
+
+UENUM()
+enum class EMovieSceneScalabilityConditionLevel : uint8
+{
+	Low,
+	Medium,
+	High,
+	Epic,
+	Cinematic
+};
+
+/**
+ * Condition on whether the current engine scalability settings fulfill a given constraint.
+ */
+UCLASS(DisplayName="Scalability Condition")
+class MOVIESCENETRACKS_API UMovieSceneScalabilityCondition
+	: public UMovieSceneCondition
+{
+	GENERATED_BODY()
+
+public: 
+
+	UPROPERTY(EditAnywhere, Category="Sequencer|Condition")
+	EMovieSceneScalabilityConditionGroup Group;
+
+	UPROPERTY(EditAnywhere, Category="Sequencer|Condition")
+	EMovieSceneScalabilityConditionOperator Operator;
+
+	UPROPERTY(EditAnywhere, Category="Sequencer|Condition")
+	EMovieSceneScalabilityConditionLevel Level;
+
+protected:
+
+	/*
+	* UMovieSceneCondition overrides 
+	*/
+	virtual bool EvaluateConditionInternal(FGuid BindingGuid, FMovieSceneSequenceID SequenceID, TSharedRef<const UE::MovieScene::FSharedPlaybackState> SharedPlaybackState) const override;
+	
+	virtual EMovieSceneConditionScope GetScopeInternal() const override { return EMovieSceneConditionScope::Global; }
+	
+	virtual EMovieSceneConditionCheckFrequency GetCheckFrequencyInternal() const override { return EMovieSceneConditionCheckFrequency::Once; }
+
+};
