Sequencer- rename 'Blueprint Condition' to 'Director Blueprint Condition' for clarity.

diff --git a/Engine/Source/Runtime/MovieSceneTracks/Public/Conditions/MovieSceneDirectorBlueprintCondition.h b/Engine/Source/Runtime/MovieSceneTracks/Public/Conditions/MovieSceneDirectorBlueprintCondition.h
--- a/Engine/Source/Runtime/MovieSceneTracks/Public/Conditions/MovieSceneDirectorBlueprintCondition.h
+++ b/Engine/Source/Runtime/MovieSceneTracks/Public/Conditions/MovieSceneDirectorBlueprintCondition.h
@@ -87,7 +87,7 @@ private:
 /**
  * Condition class allowing users to create a director blueprint endpoint in the Sequence to evaluate a condition.
  */
-UCLASS(DisplayName="Blueprint Condition")
+UCLASS(DisplayName="Director Blueprint Condition")
 class MOVIESCENETRACKS_API UMovieSceneDirectorBlueprintCondition
 	: public UMovieSceneCondition
 {
