Fixes for camera track in level sequence that grabbed the first player controller to fix splitscreen issues
- Using playbackcontext to get the level sequence actor and getting controller from the owner
- Using cached local player to get the controller when we restore previous values

diff --git a/Engine/Source/Runtime/MovieSceneTracks/Private/TrackInstances/MovieSceneCameraCutGameHandler.cpp b/Engine/Source/Runtime/MovieSceneTracks/Private/TrackInstances/MovieSceneCameraCutGameHandler.cpp
--- a/Engine/Source/Runtime/MovieSceneTracks/Private/TrackInstances/MovieSceneCameraCutGameHandler.cpp
+++ b/Engine/Source/Runtime/MovieSceneTracks/Private/TrackInstances/MovieSceneCameraCutGameHandler.cpp
@@ -23,6 +23,25 @@
 namespace UE::MovieScene
 {
 
+const APlayerController* GetPlaybackController(const UObject* PlaybackContext)
+{
+	UWorld* World = PlaybackContext ? PlaybackContext->GetWorld() : nullptr;
+	if (World == nullptr || World->GetGameInstance() == nullptr)
+	{
+		return nullptr;
+	}
+	
+	if (const AActor* PlaybackContextActor = Cast<AActor>(PlaybackContext))
+	{
+		if (const APlayerController* OwnerController = Cast<APlayerController>(PlaybackContextActor->GetOwner()))
+		{
+			return OwnerController;
+		}
+	}
+
+	return World->GetGameInstance()->GetFirstLocalPlayerController();
+}
+
 bool FPreAnimatedCameraCutTraits::ShouldHandleWorldCameraCuts(UWorld* World)
 {
 	return World &&
@@ -47,7 +66,7 @@ FPreAnimatedCameraCutState FPreAnimatedCameraCutTraits::CachePreAnimatedValue(
 	if (ShouldHandleWorldCameraCuts(World))
 	{
 		CA_SUPPRESS(6011);
-		APlayerController* PC = World->GetGameInstance()->GetFirstLocalPlayerController();
+		const APlayerController* PC = GetPlaybackController(PlaybackContext);
 
 		// Save previous view target.
 		APlayerCameraManager* CameraManager = (PC != nullptr) ? PC->PlayerCameraManager : nullptr;
@@ -77,8 +96,16 @@ void FPreAnimatedCameraCutTraits::RestorePreAnimatedValue(
 		return;
 	}
 
-	// TODO james.fleming really we should be getting the CameraManager from the sequence player or some other more reliable way that supports split screen.
-	APlayerController* PC = World->GetGameInstance()->GetFirstLocalPlayerController();
+
+	APlayerController* PC = nullptr;
+	if (ULocalPlayer* PreviousViewTarget = Cast<ULocalPlayer>(CachedValue.LastLocalPlayer.ResolveObjectPtr()))
+	{
+		if (APlayerController* OwnerController = Cast<APlayerController>(PreviousViewTarget->GetPlayerController(World)))
+		{
+			PC = OwnerController;
+		}
+	}
+
 	APlayerCameraManager* CameraManager = (PC != nullptr) ? PC->PlayerCameraManager : nullptr;
 
 	// Restore previous view target.
@@ -221,7 +248,7 @@ void FCameraCutGameHandler::SetCameraCut(
 	}
 
 	CA_SUPPRESS(6011);
-	APlayerController* PC = World->GetGameInstance()->GetFirstLocalPlayerController();
+	const APlayerController* PC = GetPlaybackController(PlaybackContext);
 	APlayerCameraManager* CameraManager = (PC != nullptr) ? PC->PlayerCameraManager : nullptr;
 
 	// If the player controller is missing, there is no camera manager for us to manage the view target
diff --git a/Engine/Source/Runtime/MovieSceneTracks/Private/TrackInstances/MovieSceneCameraCutGameHandler.h b/Engine/Source/Runtime/MovieSceneTracks/Private/TrackInstances/MovieSceneCameraCutGameHandler.h
--- a/Engine/Source/Runtime/MovieSceneTracks/Private/TrackInstances/MovieSceneCameraCutGameHandler.h
+++ b/Engine/Source/Runtime/MovieSceneTracks/Private/TrackInstances/MovieSceneCameraCutGameHandler.h
@@ -72,4 +72,5 @@ private:
 	const FSequenceInstance& SequenceInstance;
 };
 
+	const APlayerController* GetPlaybackController(const UObject* PlaybackContext);
 }  // namespace UE::MovieScene
