Sequencer: Get vector parameter channel names when adding color parameter keys

diff --git a/Engine/Source/Editor/MovieSceneTools/Private/TrackEditors/MaterialTrackEditor.cpp b/Engine/Source/Editor/MovieSceneTools/Private/TrackEditors/MaterialTrackEditor.cpp
--- a/Engine/Source/Editor/MovieSceneTools/Private/TrackEditors/MaterialTrackEditor.cpp
+++ b/Engine/Source/Editor/MovieSceneTools/Private/TrackEditors/MaterialTrackEditor.cpp
@@ -271,7 +271,9 @@ void FMaterialTrackEditor::AddColorParameter( FGuid ObjectBinding, UMovieSceneMa
 		FLinearColor ParameterValue;
 		Material->GetVectorParameterValue(ParameterInfo, ParameterValue );
 		MaterialTrack->Modify();
-		MaterialTrack->AddColorParameterKey(ParameterInfo, KeyTime, ParameterValue, InLayerName, InAssetName);
+		FParameterChannelNames ChannelNames;
+		Material->GetVectorParameterChannelNames(ParameterInfo, ChannelNames);
+		MaterialTrack->AddColorParameterKey(ParameterInfo, KeyTime, INDEX_NONE, ParameterValue, InLayerName, InAssetName, ChannelNames);
 	}
 	GetSequencer()->NotifyMovieSceneDataChanged( EMovieSceneDataChangeType::MovieSceneStructureItemAdded );
 }
diff --git a/Engine/Source/Runtime/Engine/Public/MaterialTypes.h b/Engine/Source/Runtime/Engine/Public/MaterialTypes.h
--- a/Engine/Source/Runtime/Engine/Public/MaterialTypes.h
+++ b/Engine/Source/Runtime/Engine/Public/MaterialTypes.h
@@ -252,7 +252,7 @@ enum class EMaterialSetParameterValueFlags : uint32
 };
 ENUM_CLASS_FLAGS(EMaterialSetParameterValueFlags);
 
-USTRUCT()
+USTRUCT(BlueprintType)
 struct FParameterChannelNames
 {
 	GENERATED_USTRUCT_BODY()
diff --git a/Engine/Source/Runtime/MovieSceneTracks/Private/Sections/MovieSceneComponentMaterialParameterSection.cpp b/Engine/Source/Runtime/MovieSceneTracks/Private/Sections/MovieSceneComponentMaterialParameterSection.cpp
--- a/Engine/Source/Runtime/MovieSceneTracks/Private/Sections/MovieSceneComponentMaterialParameterSection.cpp
+++ b/Engine/Source/Runtime/MovieSceneTracks/Private/Sections/MovieSceneComponentMaterialParameterSection.cpp
@@ -96,28 +96,33 @@ EMovieSceneChannelProxyType UMovieSceneComponentMaterialParameterSection::CacheC
 		FText ParameterDisplayName = GetMaterialParameterDisplayText(FText::FromName(Color.ParameterInfo.Name), Color.ParameterLayerName, Color.ParameterAssetName);
 		FText Group = FText::FromString(ParameterPath);
 
-		FMovieSceneChannelMetaData MetaData_R(*(ParameterPath + TEXT("R")), FCommonChannelData::ChannelR, Group);
+		FText RChannelDisplayName = !Color.ParameterChannelNames.R.IsEmpty() ? Color.ParameterChannelNames.R : FCommonChannelData::ChannelR;
+		FText GChannelDisplayName = !Color.ParameterChannelNames.G.IsEmpty() ? Color.ParameterChannelNames.G : FCommonChannelData::ChannelG;
+		FText BChannelDisplayName = !Color.ParameterChannelNames.B.IsEmpty() ? Color.ParameterChannelNames.B : FCommonChannelData::ChannelB;
+		FText AChannelDisplayName = !Color.ParameterChannelNames.A.IsEmpty() ? Color.ParameterChannelNames.A : FCommonChannelData::ChannelA;
+
+		FMovieSceneChannelMetaData MetaData_R(*(ParameterPath + TEXT("R")), RChannelDisplayName, Group);
 		MetaData_R.GetTooltipTextDelegate.BindLambda(GetMaterialParameterTooltipText, ParameterPath + TEXT(".R"));
 		MetaData_R.SortOrder = SortOrder++;
 		MetaData_R.Color = FCommonChannelData::RedChannelColor;
 		MetaData_R.PropertyMetaData.Add(FCommonChannelData::GroupDisplayName, ParameterDisplayName.ToString());
 		MetaData_R.GetGroupTooltipTextDelegate.BindLambda(GetMaterialParameterTooltipText, ParameterPath);
 
-		FMovieSceneChannelMetaData MetaData_G(*(ParameterPath + TEXT("G")), FCommonChannelData::ChannelG, Group);
+		FMovieSceneChannelMetaData MetaData_G(*(ParameterPath + TEXT("G")), GChannelDisplayName, Group);
 		MetaData_G.GetTooltipTextDelegate.BindLambda(GetMaterialParameterTooltipText, ParameterPath + TEXT(".G"));
 		MetaData_G.SortOrder = SortOrder++;
 		MetaData_G.Color = FCommonChannelData::GreenChannelColor;
 		MetaData_G.PropertyMetaData.Add(FCommonChannelData::GroupDisplayName, ParameterDisplayName.ToString());
 		MetaData_G.GetGroupTooltipTextDelegate.BindLambda(GetMaterialParameterTooltipText, ParameterPath);
 
-		FMovieSceneChannelMetaData MetaData_B(*(ParameterPath + TEXT("B")), FCommonChannelData::ChannelB, Group);
+		FMovieSceneChannelMetaData MetaData_B(*(ParameterPath + TEXT("B")), BChannelDisplayName, Group);
 		MetaData_B.GetTooltipTextDelegate.BindLambda(GetMaterialParameterTooltipText, ParameterPath + TEXT(".B"));
 		MetaData_B.SortOrder = SortOrder++;
 		MetaData_B.Color = FCommonChannelData::BlueChannelColor;
 		MetaData_B.PropertyMetaData.Add(FCommonChannelData::GroupDisplayName, ParameterDisplayName.ToString());
 		MetaData_B.GetGroupTooltipTextDelegate.BindLambda(GetMaterialParameterTooltipText, ParameterPath);
 
-		FMovieSceneChannelMetaData MetaData_A(*(ParameterPath + TEXT("A")), FCommonChannelData::ChannelA, Group);
+		FMovieSceneChannelMetaData MetaData_A(*(ParameterPath + TEXT("A")), AChannelDisplayName, Group);
 		MetaData_A.GetTooltipTextDelegate.BindLambda(GetMaterialParameterTooltipText, ParameterPath + TEXT(".A"));
 		MetaData_A.SortOrder = SortOrder++;
 		MetaData_A.PropertyMetaData.Add(FCommonChannelData::GroupDisplayName, ParameterDisplayName.ToString());
@@ -193,7 +198,7 @@ void UMovieSceneComponentMaterialParameterSection::AddScalarParameterKey(const F
 	}
 }
 
-void UMovieSceneComponentMaterialParameterSection::AddColorParameterKey(const FMaterialParameterInfo& InParameterInfo, FFrameNumber InTime, FLinearColor InValue, const FString& InLayerName, const FString& InAssetName)
+void UMovieSceneComponentMaterialParameterSection::AddColorParameterKey(const FMaterialParameterInfo& InParameterInfo, FFrameNumber InTime, FLinearColor InValue, const FString& InLayerName, const FString& InAssetName, const FParameterChannelNames& InChannelNames)
 {
 	FColorMaterialParameterInfoAndCurves* ExistingCurves = nullptr;
 	for (FColorMaterialParameterInfoAndCurves& ColorParameterInfoAndCurve : ColorParameterInfosAndCurves)
@@ -211,6 +216,7 @@ void UMovieSceneComponentMaterialParameterSection::AddColorParameterKey(const FM
 #if WITH_EDITOR
 		ColorParameterInfosAndCurves[NewIndex].ParameterLayerName = InLayerName;
 		ColorParameterInfosAndCurves[NewIndex].ParameterAssetName = InAssetName;
+		ColorParameterInfosAndCurves[NewIndex].ParameterChannelNames = InChannelNames;
 #endif
 		CacheChannelProxy();
 	}
diff --git a/Engine/Source/Runtime/MovieSceneTracks/Private/Tracks/MovieSceneMaterialTrack.cpp b/Engine/Source/Runtime/MovieSceneTracks/Private/Tracks/MovieSceneMaterialTrack.cpp
--- a/Engine/Source/Runtime/MovieSceneTracks/Private/Tracks/MovieSceneMaterialTrack.cpp
+++ b/Engine/Source/Runtime/MovieSceneTracks/Private/Tracks/MovieSceneMaterialTrack.cpp
@@ -190,7 +190,13 @@ void UMovieSceneMaterialTrack::AddColorParameterKey(const FMaterialParameterInfo
 	AddColorParameterKey(ParameterInfo, Time, INDEX_NONE, Value, InLayerName, InAssetName);
 }
 
+
 void UMovieSceneMaterialTrack::AddColorParameterKey(const FMaterialParameterInfo& ParameterInfo, FFrameNumber Time, int32 RowIndex, FLinearColor Value, const FString& InLayerName, const FString& InAssetName)
+{
+	AddColorParameterKey(ParameterInfo, Time, RowIndex, Value, InLayerName, InAssetName, FParameterChannelNames());
+}
+
+void UMovieSceneMaterialTrack::AddColorParameterKey(const FMaterialParameterInfo& ParameterInfo, FFrameNumber Time, int32 RowIndex, FLinearColor Value, const FString& InLayerName, const FString& InAssetName, const FParameterChannelNames& InChannelNames)
 {
 	UMovieSceneSection* NearestSection = SectionToKey;
 	if (NearestSection == nullptr || (RowIndex != INDEX_NONE && NearestSection->GetRowIndex() != RowIndex))
@@ -216,7 +222,7 @@ void UMovieSceneMaterialTrack::AddColorParameterKey(const FMaterialParameterInfo
 		}
 		else if (UMovieSceneComponentMaterialParameterSection* NearestComponentMaterialParameterSection = Cast<UMovieSceneComponentMaterialParameterSection>(NearestSection))
 		{
-			NearestComponentMaterialParameterSection->AddColorParameterKey(ParameterInfo, Time, Value, InLayerName, InAssetName);
+			NearestComponentMaterialParameterSection->AddColorParameterKey(ParameterInfo, Time, Value, InLayerName, InAssetName, InChannelNames);
 		}
 	}
 }
diff --git a/Engine/Source/Runtime/MovieSceneTracks/Public/Sections/MovieSceneComponentMaterialParameterSection.h b/Engine/Source/Runtime/MovieSceneTracks/Public/Sections/MovieSceneComponentMaterialParameterSection.h
--- a/Engine/Source/Runtime/MovieSceneTracks/Public/Sections/MovieSceneComponentMaterialParameterSection.h
+++ b/Engine/Source/Runtime/MovieSceneTracks/Public/Sections/MovieSceneComponentMaterialParameterSection.h
@@ -87,6 +87,10 @@ struct FColorMaterialParameterInfoAndCurves
 	/** Name of parameter asset, if applicable */
 	UPROPERTY()
 	FString ParameterAssetName;
+
+	/* Optional UI names for parameter channels, if applicable */
+	UPROPERTY()
+	FParameterChannelNames ParameterChannelNames;
 #endif
 };
 
@@ -107,7 +111,7 @@ public:
 
 	/** Adds a a key for a specific color parameter at the specified time with the specified value. */
 	UFUNCTION(BlueprintCallable, Category = "Sequencer|Section")
-	MOVIESCENETRACKS_API void AddColorParameterKey(const FMaterialParameterInfo& InParameterInfo, FFrameNumber InTime, FLinearColor InValue, const FString& InLayerName, const FString& InAssetName);
+	MOVIESCENETRACKS_API void AddColorParameterKey(const FMaterialParameterInfo& InParameterInfo, FFrameNumber InTime, FLinearColor InValue, const FString& InLayerName, const FString& InAssetName, const FParameterChannelNames& InChannelNames=FParameterChannelNames());
 
 	/** 
 	 * Removes a scalar parameter from this section. 
diff --git a/Engine/Source/Runtime/MovieSceneTracks/Public/Tracks/MovieSceneMaterialTrack.h b/Engine/Source/Runtime/MovieSceneTracks/Public/Tracks/MovieSceneMaterialTrack.h
--- a/Engine/Source/Runtime/MovieSceneTracks/Public/Tracks/MovieSceneMaterialTrack.h
+++ b/Engine/Source/Runtime/MovieSceneTracks/Public/Tracks/MovieSceneMaterialTrack.h
@@ -185,6 +185,18 @@ public:
 	*/
 	MOVIESCENETRACKS_API void AddColorParameterKey(const FMaterialParameterInfo& ParameterInfo, FFrameNumber Position, int32 RowIndex, FLinearColor Value, const FString& InLayerName, const FString& InAssetName);
 
+	/**
+	* Adds a color parameter key to the track.
+	* @param ParameterInfo The material parameter info for the parameter you want to add a key for.
+	* @param Time The time to add the new key.
+	* @param RowIndex The preferred row index on which to look for sections.
+	* @param The value for the new key.
+	* @param InLayerName Optional layer name for use in UI.
+	* @param InAssetName Optional asset name for use in UI.
+	* @param InChannelNames Optional channel names for the 4 vector channels for use in the UI.
+	*/
+	MOVIESCENETRACKS_API void AddColorParameterKey(const FMaterialParameterInfo& ParameterInfo, FFrameNumber Position, int32 RowIndex, FLinearColor Value, const FString& InLayerName, const FString& InAssetName, const FParameterChannelNames& InChannelNames);
+
 private:
 
 	/** The sections owned by this track .*/
