Sequencer: Gracefully handle pending kill objects when updating attachments

diff --git a/Engine/Source/Runtime/MovieSceneTracks/Private/Systems/MovieSceneComponentAttachmentSystem.cpp b/Engine/Source/Runtime/MovieSceneTracks/Private/Systems/MovieSceneComponentAttachmentSystem.cpp
--- a/Engine/Source/Runtime/MovieSceneTracks/Private/Systems/MovieSceneComponentAttachmentSystem.cpp
+++ b/Engine/Source/Runtime/MovieSceneTracks/Private/Systems/MovieSceneComponentAttachmentSystem.cpp
@@ -75,7 +75,11 @@ struct FAttachmentHandler
 
 	void UpdateOutput(FObjectKey Object, TArrayView<const FMovieSceneEntityID> Inputs, FPreAnimAttachment* Output, FEntityOutputAggregate Aggregate)
 	{
-		USceneComponent* AttachChild = CastChecked<USceneComponent>(Object.ResolveObjectPtr());
+		USceneComponent* AttachChild = Cast<USceneComponent>(Object.ResolveObjectPtr());
+		if (AttachChild == nullptr)
+		{
+			return;
+		}
 
 		for (FMovieSceneEntityID Entity : Inputs)
 		{
