New Override blend type for sections. Implemented for both ECS(for transform/double/floats) and template(Control Rig) sytems.  This is an order dependent blending algorithm that blends an additive or override section with the current value set on the blending stack. Without an override section present it will use the normal order independent absolute + additive blending algorithm that blends them together seperately. If there is an override present we accumalate the absolutes as before, but then in order blend additives with the overrides.  We need to do it in order since an override, can, ahem, override the exiting value at any point. It's common for animators to have many overrides/additives all together in one stack, muting them as needed to try out different animations.

To handle ordering we have added the idea of BlendOrder to sections/ecs/templates which when present, triggers this new blending algorithm. By default the BlendOrder is the index of the Section in the Track Array. This can change eventually, maybe have it match with the Track.RowIndex, but for the current animation layer workflows we don't need re-sort or change the BlendOrder yet.

diff --git a/Engine/Plugins/Animation/ControlRig/Source/ControlRig/Private/Sequencer/MovieSceneControlRigParameterSection.cpp b/Engine/Plugins/Animation/ControlRig/Source/ControlRig/Private/Sequencer/MovieSceneControlRigParameterSection.cpp
--- a/Engine/Plugins/Animation/ControlRig/Source/ControlRig/Private/Sequencer/MovieSceneControlRigParameterSection.cpp
+++ b/Engine/Plugins/Animation/ControlRig/Source/ControlRig/Private/Sequencer/MovieSceneControlRigParameterSection.cpp
@@ -900,61 +900,144 @@ bool UMovieSceneControlRigParameterSection::RenameParameterName(const FName& Old
 	return bWasReplaced;
 }
 
+//make sure to zero out Scale values if getting to Additive or use the current values if getting set to Override
 void UMovieSceneControlRigParameterSection::SetBlendType(EMovieSceneBlendType InBlendType)
 {
 	if (GetSupportedBlendTypes().Contains(InBlendType))
 	{
+		Modify();
 		BlendType = InBlendType;
 		if (ControlRig)
 		{
-			const FChannelMapInfo* ChannelInfo = nullptr;
-
 			// Set Defaults based upon Type
 			TArrayView<FMovieSceneFloatChannel*> FloatChannels = GetChannelProxy().GetChannels<FMovieSceneFloatChannel>();
 			TArray<FRigControlElement*> Controls = ControlRig->AvailableControls();
 
 			for (FRigControlElement* ControlElement : Controls)
 			{
+				if (!ControlRig->GetHierarchy()->IsAnimatable(ControlElement))
+				{
+					continue;
+				}
+				FChannelMapInfo* pChannelIndex = ControlChannelMap.Find(ControlElement->GetFName());
+				if (pChannelIndex == nullptr)
+				{
+					continue;
+				}
+				int32 ChannelIndex = pChannelIndex->ChannelIndex;
+
 				switch (ControlElement->Settings.ControlType)
 				{
+				case ERigControlType::Float:
+				case ERigControlType::ScaleFloat:
+				{
+					if (InBlendType == EMovieSceneBlendType::Override)
+					{
+						float Val = ControlRig->GetControlValue(ControlElement, ERigControlValueType::Current).Get<float>();
+						FloatChannels[ChannelIndex]->SetDefault(Val);
+					}
+					break;
+				}
+				case ERigControlType::Vector2D:
+				{
+					if (InBlendType == EMovieSceneBlendType::Override)
+					{
+						FVector3f Val = ControlRig->GetControlValue(ControlElement, ERigControlValueType::Current).Get<FVector3f>();
+						FloatChannels[ChannelIndex]->SetDefault(Val.X);
+						FloatChannels[ChannelIndex + 1]->SetDefault(Val.Y);
+					}
+					break;
+				}
+				case ERigControlType::Position:
+				case ERigControlType::Rotator:
+				{
+					if (InBlendType == EMovieSceneBlendType::Override)
+					{
+						FVector3f Val = (ControlElement->Settings.ControlType == ERigControlType::Rotator)
+							? FVector3f(ControlRig->GetHierarchy()->GetControlSpecifiedEulerAngle(ControlElement)) : ControlRig->GetControlValue(ControlElement, ERigControlValueType::Current).Get<FVector3f>();
 
+						FloatChannels[ChannelIndex]->SetDefault(Val.X);
+						FloatChannels[ChannelIndex + 1]->SetDefault(Val.Y);
+						FloatChannels[ChannelIndex + 2]->SetDefault(Val.Z);
+					}
+					break;
+				}
 				case ERigControlType::Scale:
 				{
-					ChannelInfo = ControlChannelMap.Find(ControlElement->GetFName());
-					if (ChannelInfo)
+					if (InBlendType == EMovieSceneBlendType::Absolute)
 					{
-						if (InBlendType == EMovieSceneBlendType::Absolute)
-						{
-							FloatChannels[ChannelInfo->ChannelIndex]->SetDefault(1.0f);
-							FloatChannels[ChannelInfo->ChannelIndex+1]->SetDefault(1.0f);
-							FloatChannels[ChannelInfo->ChannelIndex+2]->SetDefault(1.0f);
-						}
-						else
-						{
-							FloatChannels[ChannelInfo->ChannelIndex]->SetDefault(0.0f);
-							FloatChannels[ChannelInfo->ChannelIndex + 1]->SetDefault(0.0f);
-							FloatChannels[ChannelInfo->ChannelIndex + 2]->SetDefault(0.0f);
-						}
+						FloatChannels[ChannelIndex]->SetDefault(1.0f);
+						FloatChannels[ChannelIndex + 1]->SetDefault(1.0f);
+						FloatChannels[ChannelIndex + 2]->SetDefault(1.0f);
+					}
+					else if (InBlendType == EMovieSceneBlendType::Additive)
+					{
+						FloatChannels[ChannelIndex]->SetDefault(0.0f);
+						FloatChannels[ChannelIndex + 1]->SetDefault(0.0f);
+						FloatChannels[ChannelIndex + 2]->SetDefault(0.0f);
+					}
+					else if (InBlendType == EMovieSceneBlendType::Override)
+					{
+						FVector3f Val = ControlRig->GetControlValue(ControlElement, ERigControlValueType::Current).Get<FVector3f>();
+						FloatChannels[ChannelIndex]->SetDefault(Val.X);
+						FloatChannels[ChannelIndex + 1]->SetDefault(Val.Y);
+						FloatChannels[ChannelIndex + 2]->SetDefault(Val.Z);
 					}
 				}
 				break;
 				case ERigControlType::Transform:
 				case ERigControlType::EulerTransform:
+				case ERigControlType::TransformNoScale:
 				{
-					ChannelInfo = ControlChannelMap.Find(ControlElement->GetFName());
-					if (ChannelInfo)
+					FTransform Val = FTransform::Identity;
+					if (ControlElement->Settings.ControlType == ERigControlType::TransformNoScale)
+					{
+						FTransformNoScale NoScale =
+							ControlRig->GetControlValue(ControlElement, ERigControlValueType::Current).Get<FRigControlValue::FTransformNoScale_Float>().ToTransform();
+						Val = NoScale;
+					}
+					else if (ControlElement->Settings.ControlType == ERigControlType::EulerTransform)
+					{
+						FEulerTransform Euler =
+							ControlRig->GetControlValue(ControlElement, ERigControlValueType::Current).Get<FRigControlValue::FEulerTransform_Float>().ToTransform();
+						Val = Euler.ToFTransform();
+					}
+					else
+					{
+						Val = ControlRig->GetControlValue(ControlElement, ERigControlValueType::Current).Get<FRigControlValue::FTransform_Float>().ToTransform();
+					}
+					if (InBlendType == EMovieSceneBlendType::Override)
+					{
+						FVector CurrentVector = Val.GetTranslation();
+						FloatChannels[ChannelIndex]->SetDefault(CurrentVector.X);
+						FloatChannels[ChannelIndex + 1]->SetDefault(CurrentVector.Y);
+						FloatChannels[ChannelIndex + 2]->SetDefault(CurrentVector.Z);
+
+						CurrentVector = ControlRig->GetHierarchy()->GetControlSpecifiedEulerAngle(ControlElement);
+						FloatChannels[ChannelIndex + 3]->SetDefault(CurrentVector.X);
+						FloatChannels[ChannelIndex + 4]->SetDefault(CurrentVector.Y);
+						FloatChannels[ChannelIndex + 5]->SetDefault(CurrentVector.Z);
+					}
+					if (ControlElement->Settings.ControlType != ERigControlType::TransformNoScale)
 					{
 						if (InBlendType == EMovieSceneBlendType::Absolute)
 						{
-							FloatChannels[ChannelInfo->ChannelIndex + 6]->SetDefault(1.0f);
-							FloatChannels[ChannelInfo->ChannelIndex + 7]->SetDefault(1.0f);
-							FloatChannels[ChannelInfo->ChannelIndex + 8]->SetDefault(1.0f);
+							FloatChannels[ChannelIndex + 6]->SetDefault(1.0f);
+							FloatChannels[ChannelIndex + 7]->SetDefault(1.0f);
+							FloatChannels[ChannelIndex + 8]->SetDefault(1.0f);
 						}
-						else
+						else if (InBlendType == EMovieSceneBlendType::Additive)
+						{
+							FloatChannels[ChannelIndex + 6]->SetDefault(0.0f);
+							FloatChannels[ChannelIndex + 7]->SetDefault(0.0f);
+							FloatChannels[ChannelIndex + 8]->SetDefault(0.0f);
+						}
+						else if (InBlendType == EMovieSceneBlendType::Override)
 						{
-							FloatChannels[ChannelInfo->ChannelIndex + 6]->SetDefault(0.0f);
-							FloatChannels[ChannelInfo->ChannelIndex + 7]->SetDefault(0.0f);
-							FloatChannels[ChannelInfo->ChannelIndex + 8]->SetDefault(0.0f);
+							FVector CurrentVector = Val.GetScale3D();
+							FloatChannels[ChannelIndex + 6]->SetDefault(CurrentVector.X);
+							FloatChannels[ChannelIndex + 7]->SetDefault(CurrentVector.Y);
+							FloatChannels[ChannelIndex + 8]->SetDefault(CurrentVector.Z);
 						}
 					}
 				}
diff --git a/Engine/Plugins/Animation/ControlRig/Source/ControlRig/Private/Sequencer/MovieSceneControlRigParameterTemplate.cpp b/Engine/Plugins/Animation/ControlRig/Source/ControlRig/Private/Sequencer/MovieSceneControlRigParameterTemplate.cpp
--- a/Engine/Plugins/Animation/ControlRig/Source/ControlRig/Private/Sequencer/MovieSceneControlRigParameterTemplate.cpp
+++ b/Engine/Plugins/Animation/ControlRig/Source/ControlRig/Private/Sequencer/MovieSceneControlRigParameterTemplate.cpp
@@ -1869,6 +1869,7 @@ void FMovieSceneControlRigParameterTemplate::Evaluate(const FMovieSceneEvaluatio
 			return;
 		}
 
+		int32 BlendingOrder = Section->GetBlendingOrder();
 
 		//Do blended tokens
 		FEvaluatedControlRigParameterSectionValues Values;
@@ -1900,7 +1901,7 @@ void FMovieSceneControlRigParameterTemplate::Evaluate(const FMovieSceneEvaluatio
 			{
 				ExecutionTokens.GetBlendingAccumulator().DefineActuator(ActuatorTypeID, MakeShared <TControlRigParameterActuatorFloat>(AnimTypeID, ScalarNameAndValue.ParameterName, Section));
 			}
-			ExecutionTokens.BlendToken(ActuatorTypeID, TBlendableToken<FControlRigTrackTokenFloat>(ScalarNameAndValue.Value, BlendType, Weight));
+			ExecutionTokens.BlendToken(ActuatorTypeID, TBlendableToken<FControlRigTrackTokenFloat>(ScalarNameAndValue.Value, BlendType, Weight, BlendingOrder));
 		}
 
 		UE::MovieScene::TMultiChannelValue<float, 3> VectorData;
@@ -1917,7 +1918,7 @@ void FMovieSceneControlRigParameterTemplate::Evaluate(const FMovieSceneEvaluatio
 			VectorData.Set(1, VectorNameAndValue.Value.Y);
 			VectorData.Set(2, VectorNameAndValue.Value.Z);
 
-			ExecutionTokens.BlendToken(ActuatorTypeID, TBlendableToken<FControlRigTrackTokenVector>(VectorData, BlendType, Weight));
+			ExecutionTokens.BlendToken(ActuatorTypeID, TBlendableToken<FControlRigTrackTokenVector>(VectorData, BlendType, Weight, BlendingOrder));
 		}
 
 		UE::MovieScene::TMultiChannelValue<float, 2> Vector2DData;
@@ -1933,7 +1934,7 @@ void FMovieSceneControlRigParameterTemplate::Evaluate(const FMovieSceneEvaluatio
 			Vector2DData.Set(0, Vector2DNameAndValue.Value.X);
 			Vector2DData.Set(1, Vector2DNameAndValue.Value.Y);
 
-			ExecutionTokens.BlendToken(ActuatorTypeID, TBlendableToken<FControlRigTrackTokenVector2D>(Vector2DData, BlendType, Weight));
+			ExecutionTokens.BlendToken(ActuatorTypeID, TBlendableToken<FControlRigTrackTokenVector2D>(Vector2DData, BlendType, Weight, BlendingOrder));
 		}
 
 		UE::MovieScene::TMultiChannelValue<float, 9> TransformData;
@@ -1960,7 +1961,7 @@ void FMovieSceneControlRigParameterTemplate::Evaluate(const FMovieSceneEvaluatio
 			TransformData.Set(6, Transform.Scale.X);
 			TransformData.Set(7, Transform.Scale.Y);
 			TransformData.Set(8, Transform.Scale.Z);
-			ExecutionTokens.BlendToken(ActuatorTypeID, TBlendableToken<FControlRigTrackTokenTransform>(TransformData, BlendType, Weight));
+			ExecutionTokens.BlendToken(ActuatorTypeID, TBlendableToken<FControlRigTrackTokenTransform>(TransformData, BlendType, Weight, BlendingOrder));
 		}
 
 	}
@@ -1986,6 +1987,7 @@ void FMovieSceneControlRigParameterTemplate::EvaluateCurvesWithMasks(const FMovi
 		Values.ConstraintsValues.Reserve(Constraints.Num());
 
 		const bool bIsAdditive = (Section->GetBlendType().IsValid() && Section->GetBlendType().Get() == EMovieSceneBlendType::Additive);
+		const bool bIsAbsolute = (Section->GetBlendType().IsValid() && Section->GetBlendType().Get() == EMovieSceneBlendType::Absolute);
 
 		// Populate each of the output arrays in turn
 		for (int32 Index = 0; Index < Scalars.Num(); ++Index)
@@ -1999,7 +2001,7 @@ void FMovieSceneControlRigParameterTemplate::EvaluateCurvesWithMasks(const FMovi
 			}
 			else
 			{
-				Value = (bIsAdditive
+				Value = (!bIsAbsolute
 					|| Scalar.ParameterCurve.GetDefault().IsSet() == false) ? 0.0f :
 					Scalar.ParameterCurve.GetDefault().GetValue();
 			}
@@ -2040,7 +2042,7 @@ void FMovieSceneControlRigParameterTemplate::EvaluateCurvesWithMasks(const FMovi
 			}
 			else
 			{
-				Value = (bIsAdditive
+				Value = (!bIsAbsolute
 					|| Bool.ParameterCurve.GetDefault().IsSet() == false) ? false :
 					Bool.ParameterCurve.GetDefault().GetValue();
 			}
@@ -2057,7 +2059,7 @@ void FMovieSceneControlRigParameterTemplate::EvaluateCurvesWithMasks(const FMovi
 			}
 			else
 			{
-				Value = (bIsAdditive
+				Value = (!bIsAbsolute
 					|| Integer.ParameterCurve.GetDefault().IsSet() == false) ? 0 :
 					Integer.ParameterCurve.GetDefault().GetValue();
 			}
@@ -2074,7 +2076,7 @@ void FMovieSceneControlRigParameterTemplate::EvaluateCurvesWithMasks(const FMovi
 			}
 			else
 			{
-				Value = (bIsAdditive
+				Value = (!bIsAbsolute
 					|| Enum.ParameterCurve.GetDefault().IsSet() == false) ? 0 :
 					Enum.ParameterCurve.GetDefault().GetValue();
 
@@ -2094,7 +2096,7 @@ void FMovieSceneControlRigParameterTemplate::EvaluateCurvesWithMasks(const FMovi
 			}
 			else
 			{
-				if (!bIsAdditive)
+				if (bIsAbsolute)
 				{
 					if (Vector2D.XCurve.GetDefault().IsSet())
 					{
@@ -2123,7 +2125,7 @@ void FMovieSceneControlRigParameterTemplate::EvaluateCurvesWithMasks(const FMovi
 			}
 			else
 			{
-				if (!bIsAdditive)
+				if (bIsAbsolute)
 				{
 					if (Vector.XCurve.GetDefault().IsSet())
 					{
@@ -2155,7 +2157,7 @@ void FMovieSceneControlRigParameterTemplate::EvaluateCurvesWithMasks(const FMovi
 			}
 			else
 			{
-				if (!bIsAdditive)
+				if (bIsAbsolute)
 				{
 					if (Color.RedCurve.GetDefault().IsSet())
 					{
@@ -2198,7 +2200,7 @@ void FMovieSceneControlRigParameterTemplate::EvaluateCurvesWithMasks(const FMovi
 				}
 				else
 				{
-					if (!bIsAdditive && Transform.Translation[0].GetDefault().IsSet())
+					if (bIsAbsolute && Transform.Translation[0].GetDefault().IsSet())
 					{
 						Translation[0] = Transform.Translation[0].GetDefault().GetValue();
 					}
@@ -2209,7 +2211,7 @@ void FMovieSceneControlRigParameterTemplate::EvaluateCurvesWithMasks(const FMovi
 				}
 				else
 				{
-					if (!bIsAdditive && Transform.Translation[1].GetDefault().IsSet())
+					if (bIsAbsolute && Transform.Translation[1].GetDefault().IsSet())
 					{
 						Translation[1] = Transform.Translation[1].GetDefault().GetValue();
 					}
@@ -2220,7 +2222,7 @@ void FMovieSceneControlRigParameterTemplate::EvaluateCurvesWithMasks(const FMovi
 				}
 				else
 				{
-					if (!bIsAdditive && Transform.Translation[2].GetDefault().IsSet())
+					if (bIsAbsolute && Transform.Translation[2].GetDefault().IsSet())
 					{
 						Translation[2] = Transform.Translation[2].GetDefault().GetValue();
 					}
@@ -2231,7 +2233,7 @@ void FMovieSceneControlRigParameterTemplate::EvaluateCurvesWithMasks(const FMovi
 				}
 				else
 				{
-					if (!bIsAdditive && Transform.Rotation[0].GetDefault().IsSet())
+					if (bIsAbsolute && Transform.Rotation[0].GetDefault().IsSet())
 					{
 						Rotator.Roll = Transform.Rotation[0].GetDefault().GetValue();
 					}
@@ -2242,7 +2244,7 @@ void FMovieSceneControlRigParameterTemplate::EvaluateCurvesWithMasks(const FMovi
 				}
 				else
 				{
-					if (!bIsAdditive && Transform.Rotation[1].GetDefault().IsSet())
+					if (bIsAbsolute && Transform.Rotation[1].GetDefault().IsSet())
 					{
 						Rotator.Pitch = Transform.Rotation[1].GetDefault().GetValue();
 					}
@@ -2253,7 +2255,7 @@ void FMovieSceneControlRigParameterTemplate::EvaluateCurvesWithMasks(const FMovi
 				}
 				else
 				{
-					if (!bIsAdditive && Transform.Rotation[2].GetDefault().IsSet())
+					if (bIsAbsolute && Transform.Rotation[2].GetDefault().IsSet())
 					{
 						Rotator.Yaw = Transform.Rotation[2].GetDefault().GetValue();
 					}
@@ -2265,7 +2267,7 @@ void FMovieSceneControlRigParameterTemplate::EvaluateCurvesWithMasks(const FMovi
 				}
 				else
 				{
-					if (!bIsAdditive && Transform.Scale[0].GetDefault().IsSet())
+					if (bIsAbsolute && Transform.Scale[0].GetDefault().IsSet())
 					{
 						Scale[0] = Transform.Scale[0].GetDefault().GetValue();
 					}
@@ -2276,7 +2278,7 @@ void FMovieSceneControlRigParameterTemplate::EvaluateCurvesWithMasks(const FMovi
 				}
 				else
 				{
-					if (!bIsAdditive && Transform.Scale[1].GetDefault().IsSet())
+					if (bIsAbsolute && Transform.Scale[1].GetDefault().IsSet())
 					{
 						Scale[1] = Transform.Scale[1].GetDefault().GetValue();
 					}
@@ -2287,7 +2289,7 @@ void FMovieSceneControlRigParameterTemplate::EvaluateCurvesWithMasks(const FMovi
 				}
 				else
 				{
-					if (!bIsAdditive && Transform.Scale[2].GetDefault().IsSet())
+					if (bIsAbsolute && Transform.Scale[2].GetDefault().IsSet())
 					{
 						Scale[2] = Transform.Scale[2].GetDefault().GetValue();
 					}
@@ -2295,7 +2297,7 @@ void FMovieSceneControlRigParameterTemplate::EvaluateCurvesWithMasks(const FMovi
 			}
 			else //completely masked use default or zeroed, which is already set if additive
 			{
-				if (!bIsAdditive)
+				if (bIsAbsolute)
 				{
 
 					if (Transform.Translation[0].GetDefault().IsSet())
@@ -2374,6 +2376,7 @@ void FMovieSceneControlRigParameterTemplate::Interrogate(const FMovieSceneContex
 
 		FControlRigAnimTypeIDsPtr TypeIDs = FControlRigAnimTypeIDs::Get(Section->GetControlRig());
 		EMovieSceneBlendType BlendType = Section->GetBlendType().IsValid() ? Section->GetBlendType().Get() : EMovieSceneBlendType::Absolute;
+		int32 BlendingOrder = Section->GetBlendingOrder();
 
 		float Weight = EvaluateEasing(Context.GetTime());
 		if (EnumHasAllFlags(Section->TransformMask.GetChannels(), EMovieSceneTransformChannel::Weight))
@@ -2393,7 +2396,7 @@ void FMovieSceneControlRigParameterTemplate::Interrogate(const FMovieSceneContex
 			{
 				Container.GetAccumulator().DefineActuator(ActuatorTypeID, MakeShared <TControlRigParameterActuatorFloat>(AnimTypeID, ScalarNameAndValue.ParameterName, Section));
 			}
-			Container.GetAccumulator().BlendToken(FMovieSceneEvaluationOperand(), ActuatorTypeID, FMovieSceneEvaluationScope(), Context,TBlendableToken<FControlRigTrackTokenFloat>(ScalarNameAndValue.Value, BlendType, Weight));
+			Container.GetAccumulator().BlendToken(FMovieSceneEvaluationOperand(), ActuatorTypeID, FMovieSceneEvaluationScope(), Context,TBlendableToken<FControlRigTrackTokenFloat>(ScalarNameAndValue.Value, BlendType, Weight, BlendingOrder));
 		}
 
 		UE::MovieScene::TMultiChannelValue<float, 2> Vector2DData;
@@ -2409,7 +2412,7 @@ void FMovieSceneControlRigParameterTemplate::Interrogate(const FMovieSceneContex
 			Vector2DData.Set(0, Vector2DNameAndValue.Value.X);
 			Vector2DData.Set(1, Vector2DNameAndValue.Value.Y);
 
-			Container.GetAccumulator().BlendToken(FMovieSceneEvaluationOperand(), ActuatorTypeID, FMovieSceneEvaluationScope(), Context, TBlendableToken<FControlRigTrackTokenVector2D>(Vector2DData, BlendType, Weight));
+			Container.GetAccumulator().BlendToken(FMovieSceneEvaluationOperand(), ActuatorTypeID, FMovieSceneEvaluationScope(), Context, TBlendableToken<FControlRigTrackTokenVector2D>(Vector2DData, BlendType, Weight, BlendingOrder));
 		}
 
 		UE::MovieScene::TMultiChannelValue<float, 3> VectorData;
@@ -2426,7 +2429,7 @@ void FMovieSceneControlRigParameterTemplate::Interrogate(const FMovieSceneContex
 			VectorData.Set(1, VectorNameAndValue.Value.Y);
 			VectorData.Set(2, VectorNameAndValue.Value.Z);
 
-			Container.GetAccumulator().BlendToken(FMovieSceneEvaluationOperand(), ActuatorTypeID, FMovieSceneEvaluationScope(),Context,TBlendableToken<FControlRigTrackTokenVector>(VectorData, BlendType, Weight));
+			Container.GetAccumulator().BlendToken(FMovieSceneEvaluationOperand(), ActuatorTypeID, FMovieSceneEvaluationScope(),Context,TBlendableToken<FControlRigTrackTokenVector>(VectorData, BlendType, Weight, BlendingOrder));
 		}
 
 		UE::MovieScene::TMultiChannelValue<float, 9> TransformData;
@@ -2453,7 +2456,7 @@ void FMovieSceneControlRigParameterTemplate::Interrogate(const FMovieSceneContex
 			TransformData.Set(6, Transform.Scale.X);
 			TransformData.Set(7, Transform.Scale.Y);
 			TransformData.Set(8, Transform.Scale.Z);
-			Container.GetAccumulator().BlendToken(FMovieSceneEvaluationOperand(),ActuatorTypeID, FMovieSceneEvaluationScope(), Context, TBlendableToken<FControlRigTrackTokenTransform>(TransformData, BlendType, Weight));
+			Container.GetAccumulator().BlendToken(FMovieSceneEvaluationOperand(),ActuatorTypeID, FMovieSceneEvaluationScope(), Context, TBlendableToken<FControlRigTrackTokenTransform>(TransformData, BlendType, Weight, BlendingOrder));
 		}
 
 	}
diff --git a/Engine/Plugins/Animation/ControlRig/Source/ControlRig/Private/Sequencer/MovieSceneControlRigParameterTrack.cpp b/Engine/Plugins/Animation/ControlRig/Source/ControlRig/Private/Sequencer/MovieSceneControlRigParameterTrack.cpp
--- a/Engine/Plugins/Animation/ControlRig/Source/ControlRig/Private/Sequencer/MovieSceneControlRigParameterTrack.cpp
+++ b/Engine/Plugins/Animation/ControlRig/Source/ControlRig/Private/Sequencer/MovieSceneControlRigParameterTrack.cpp
@@ -44,7 +44,7 @@ UMovieSceneControlRigParameterTrack::UMovieSceneControlRigParameterTrack(const F
 	SupportedBlendTypes = FMovieSceneBlendTypeField::None();
 	SupportedBlendTypes.Add(EMovieSceneBlendType::Additive);
 	SupportedBlendTypes.Add(EMovieSceneBlendType::Absolute);
-
+	SupportedBlendTypes.Add(EMovieSceneBlendType::Override);
 }
 
 void UMovieSceneControlRigParameterTrack::BeginDestroy()
diff --git a/Engine/Source/Editor/Sequencer/Private/SequencerContextMenus.cpp b/Engine/Source/Editor/Sequencer/Private/SequencerContextMenus.cpp
--- a/Engine/Source/Editor/Sequencer/Private/SequencerContextMenus.cpp
+++ b/Engine/Source/Editor/Sequencer/Private/SequencerContextMenus.cpp
@@ -830,7 +830,9 @@ bool FSectionContextMenu::CanSetSectionToKey() const
 	for (UMovieSceneSection* Section : Sequencer->GetViewModel()->GetSelection()->GetSelectedSections())
 	{
 		UMovieSceneTrack* Track = Section->GetTypedOuter<UMovieSceneTrack>();
-		if (Track && Section->GetBlendType().IsValid() && (Section->GetBlendType().Get() == EMovieSceneBlendType::Absolute || Section->GetBlendType().Get() == EMovieSceneBlendType::Additive))
+		if (Track && Section->GetBlendType().IsValid() && (Section->GetBlendType().Get() == EMovieSceneBlendType::Absolute 
+			|| Section->GetBlendType().Get() == EMovieSceneBlendType::Additive 
+			|| Section->GetBlendType().Get() == EMovieSceneBlendType::Override))
 		{
 			return true;
 		}
diff --git a/Engine/Source/Editor/Sequencer/Private/SequencerUtilities.cpp b/Engine/Source/Editor/Sequencer/Private/SequencerUtilities.cpp
--- a/Engine/Source/Editor/Sequencer/Private/SequencerUtilities.cpp
+++ b/Engine/Source/Editor/Sequencer/Private/SequencerUtilities.cpp
@@ -318,17 +318,24 @@ void FSequencerUtilities::PopulateMenu_CreateNewSection(FMenuBuilder& MenuBuilde
 		FQualifiedFrameTime CurrentTime = Sequencer->GetLocalTime();
 		FFrameNumber PlaybackEnd = UE::MovieScene::DiscreteExclusiveUpper(Sequencer->GetFocusedMovieSceneSequence()->GetMovieScene()->GetPlaybackRange());
 
+		int32 SpecifiedRowIndex = RowIndex;
+
 		FScopedTransaction Transaction(LOCTEXT("AddSectionTransactionText", "Add Section"));
 		if (UMovieSceneSection* NewSection = Track->CreateNewSection())
 		{
 			int32 OverlapPriority = 0;
 			TMap<int32, int32> NewToOldRowIndices;
+			//if creating with an override force the row index to be last
+			if (Track->GetSupportedBlendTypes().Contains(EMovieSceneBlendType::Override))
+			{
+				SpecifiedRowIndex = Track->GetMaxRowIndex() + 1;
+			}
 			for (UMovieSceneSection* Section : Track->GetAllSections())
 			{
 				OverlapPriority = FMath::Max(Section->GetOverlapPriority() + 1, OverlapPriority);				
 
 				// Move existing sections on the same row or beyond so that they don't overlap with the new section
-				if (Section != NewSection && Section->GetRowIndex() >= RowIndex)
+				if (Section != NewSection && Section->GetRowIndex() >= SpecifiedRowIndex)
 				{
 					int32 OldRowIndex = Section->GetRowIndex();
 					int32 NewRowIndex = Section->GetRowIndex() + 1;
@@ -360,7 +367,7 @@ void FSequencerUtilities::PopulateMenu_CreateNewSection(FMenuBuilder& MenuBuilde
 			}
 
 			NewSection->SetOverlapPriority(OverlapPriority);
-			NewSection->SetRowIndex(RowIndex);			
+			NewSection->SetRowIndex(SpecifiedRowIndex);
 			NewSection->SetBlendType(BlendType);
 
 			Track->AddSection(*NewSection);
@@ -368,13 +375,10 @@ void FSequencerUtilities::PopulateMenu_CreateNewSection(FMenuBuilder& MenuBuilde
 
 			if (UMovieSceneNameableTrack* NameableTrack = Cast<UMovieSceneNameableTrack>(Track))
 			{
-				NameableTrack->SetTrackRowDisplayName(FText::GetEmpty(), RowIndex);
+				NameableTrack->SetTrackRowDisplayName(FText::GetEmpty(), SpecifiedRowIndex);
 			}
 
 			Sequencer->NotifyMovieSceneDataChanged(EMovieSceneDataChangeType::MovieSceneStructureItemAdded);
-			Sequencer->EmptySelection();
-			Sequencer->SelectSection(NewSection);
-			Sequencer->ThrobSectionSelection();
 		}
 		else
 		{
diff --git a/Engine/Source/Runtime/MovieScene/Private/EntitySystem/BuiltInComponentTypes.cpp b/Engine/Source/Runtime/MovieScene/Private/EntitySystem/BuiltInComponentTypes.cpp
--- a/Engine/Source/Runtime/MovieScene/Private/EntitySystem/BuiltInComponentTypes.cpp
+++ b/Engine/Source/Runtime/MovieScene/Private/EntitySystem/BuiltInComponentTypes.cpp
@@ -343,6 +343,7 @@ FBuiltInComponentTypes::FBuiltInComponentTypes()
 	ComponentRegistry->NewComponentType(&HierarchicalEasingProvider, TEXT("Hierarchical Easing Provider"));
 
 	ComponentRegistry->NewComponentType(&BlenderType,           TEXT("Blender System Type"), EComponentTypeFlags::CopyToChildren);
+	ComponentRegistry->NewComponentType(&BlendingOrder,			TEXT("Blending Order"));
 	ComponentRegistry->NewComponentType(&BlendChannelInput,     TEXT("Blend Channel Input"));
 	ComponentRegistry->NewComponentType(&HierarchicalBias,      TEXT("Hierarchical Bias"));
 	ComponentRegistry->NewComponentType(&BlendChannelOutput,    TEXT("Blend Channel Output"));
@@ -404,6 +405,7 @@ FBuiltInComponentTypes::FBuiltInComponentTypes()
 	Tags.RelativeBlend           = ComponentRegistry->NewTag(TEXT("Is Relative Blend"));
 	Tags.AdditiveBlend           = ComponentRegistry->NewTag(TEXT("Is Additive Blend"));
 	Tags.AdditiveFromBaseBlend   = ComponentRegistry->NewTag(TEXT("Is Additive From Base Blend"));
+	Tags.OverrideBlend			 = ComponentRegistry->NewTag(TEXT("Is Override Blend"));
 
 	Tags.NeedsLink               = ComponentRegistry->NewTag(TEXT("Needs Link"));
 	Tags.NeedsUnlink             = ComponentRegistry->NewTag(TEXT("Needs Unlink"));
@@ -443,6 +445,7 @@ FBuiltInComponentTypes::FBuiltInComponentTypes()
 	ComponentRegistry->Factories.DefineChildComponent(Tags.AbsoluteBlend, Tags.AbsoluteBlend);
 	ComponentRegistry->Factories.DefineChildComponent(Tags.RelativeBlend, Tags.RelativeBlend);
 	ComponentRegistry->Factories.DefineChildComponent(Tags.AdditiveBlend, Tags.AdditiveBlend);
+	ComponentRegistry->Factories.DefineChildComponent(Tags.OverrideBlend, Tags.OverrideBlend);
 	ComponentRegistry->Factories.DefineChildComponent(Tags.AdditiveFromBaseBlend, Tags.AdditiveFromBaseBlend);
 	ComponentRegistry->Factories.DefineChildComponent(Tags.FixedTime,     Tags.FixedTime);
 	ComponentRegistry->Factories.DefineChildComponent(Tags.PreRoll,       Tags.PreRoll);
@@ -461,6 +464,7 @@ FBuiltInComponentTypes::FBuiltInComponentTypes()
 	ComponentRegistry->Factories.DuplicateChildComponent(RootInstanceHandle);
 	ComponentRegistry->Factories.DuplicateChildComponent(PropertyBinding);
 	ComponentRegistry->Factories.DuplicateChildComponent(HierarchicalBias);
+	ComponentRegistry->Factories.DuplicateChildComponent(BlendingOrder);
 
 	// Children always need a Parent - these are initialized by the tasks that create them
 	{
diff --git a/Engine/Source/Runtime/MovieScene/Private/EntitySystem/MovieSceneDecompositionQuery.cpp b/Engine/Source/Runtime/MovieScene/Private/EntitySystem/MovieSceneDecompositionQuery.cpp
--- a/Engine/Source/Runtime/MovieScene/Private/EntitySystem/MovieSceneDecompositionQuery.cpp
+++ b/Engine/Source/Runtime/MovieScene/Private/EntitySystem/MovieSceneDecompositionQuery.cpp
@@ -42,18 +42,50 @@ void FDecomposedValue::Decompose(
 			Additives.AccumulateThis(Pair.Value);
 		}
 	}
+	for (TTuple<FMovieSceneEntityID, FWeightedValue> Pair : DecomposedOverrides)
+	{
+		if (Pair.Get<0>() == EntityID)
+		{
+			ThisValue = Pair.Value;
+			OutBlendType = EDecomposedValueBlendType::Override;
+		}
+	}
+}
+static void AddToValue(FWeightedValue& Data, double& OutCurrentValue)
+{
+	if (Data.bIsAdditive)
+	{
+		OutCurrentValue += (Data.Value * Data.Weight);
+	}
+	else //override do weighted blend
+	{
+		if (Data.Weight >= 1.0)
+		{
+			OutCurrentValue = Data.Value;
+		}
+		else
+		{
+			OutCurrentValue = (OutCurrentValue * (1.0 - Data.Weight)) +
+				(Data.Value * Data.Weight);
+		}
+	}
 }
-
 double FDecomposedValue::Recompose(FMovieSceneEntityID RecomposeEntity, double CurrentValue, const double* InitialValue) const
 {
-	// First, a little reminder... the formula for blending values is:
+	// First, a little reminder... the formula for blending values is either:
 	//
 	// Value = Absolutes + Additives
 	//
 	// Absolutes = (Abs1 * AbsWeight1 + ... + AbsN * AbsWeightN) / (AbsWeight1 + ... + AbsWeightN)
 	// Additives = (Add1 * AddWeight1 + ... AddN * AddWeightN)
 	//
-	//
+	// Or if an Override is present it's Absolutes  plus the Additives and Overrides in blending order
+	// Value = Absolutes + BlendOrder1 +  BlendOrder2
+	// 
+	// Where the BlendOrder is either
+	// AdditivesN = (AddN * AddWeightN) or
+	// OverrideN = (OverrideN * OverrideWeightN) + (Value * (1.0 -OverrideWeightN))
+	// 
 	// Sort through all the data we have so that the contribution of RecomposeEntity is set aside in Channel.  The
 	// contributions of other decomposed entities (if any) are combined into OtherAbsolute and OtherAdditive, which
 	// should also be combined with Result (which contains non-decomposed entities) to get the full picture.
@@ -62,11 +94,13 @@ double FDecomposedValue::Recompose(FMovieSceneEntityID RecomposeEntity, double C
 	FAccumulatedWeightedValue OtherAbsolute;
 	FAccumulatedWeightedValue OtherAdditive;
 	Decompose(RecomposeEntity, Channel, BlendType, OtherAbsolute, OtherAdditive);
-
+	
 	FAccumulatedWeightedValue ResultAbsolute = Result.Absolute;
-	const bool bIsAdditive = (BlendType != EDecomposedValueBlendType::Absolute);
+	const bool bIsAdditive = (BlendType == EDecomposedValueBlendType::Additive ||
+		BlendType == EDecomposedValueBlendType::AdditiveFromBase);
+	const bool bIsOverride = (BlendType == EDecomposedValueBlendType::Override);
 	float TotalAbsoluteWeight = ResultAbsolute.TotalWeight + OtherAbsolute.TotalWeight;
-	if (!bIsAdditive)
+	if (!bIsAdditive || !bIsOverride)
 	{
 		TotalAbsoluteWeight += Channel.Weight;
 	}
@@ -84,7 +118,39 @@ double FDecomposedValue::Recompose(FMovieSceneEntityID RecomposeEntity, double C
 	// between what we have and what the desired current value is.
 	if (OtherAbsolute.TotalWeight == 0.f && OtherAdditive.TotalWeight == 0.f)
 	{
-		if (bIsAdditive)
+		//for overrides we need to iterate over everthing since it's order dependent since 
+		//there can be an override track anywhere
+		if (bIsOverride && AllDecomposedOverrides.Num() > 0) //if override we go up to this one's BlendValue
+		{
+			AllDecomposedOverrides.Sort();
+			double OverrideValue = ResultAbsolute.Total; //current value that's set
+			double PreviousValue = OverrideValue;    //value set right before this one
+			double ExtraValue = 0.0;
+			for(FWeightedValue& Data: AllDecomposedOverrides)
+			{
+				if (Data.BlendingOrder < Channel.BlendingOrder)
+				{
+					AddToValue(Data, OverrideValue);
+					PreviousValue = OverrideValue;
+				}
+				else if(Data.BlendingOrder > Channel.BlendingOrder)
+				{
+					AddToValue(Data, ExtraValue);
+				}
+			}
+			if (Channel.bIsAdditive)
+			{
+				const double Difference = CurrentValue - ExtraValue - PreviousValue;
+				return (Channel.Weight == 0.f ? Difference : (Difference / Channel.Weight));
+			}
+			else
+			{
+				const double Difference = (CurrentValue - ExtraValue) - (PreviousValue * (1.0 - Channel.Weight));
+				return (Channel.Weight == 0.f ? Difference : (Difference / Channel.Weight));
+
+			}
+		}
+		else if (bIsAdditive)
 		{
 			// For an additive channel, we just put the missing difference on it, adjusted by its weight.
 			const double WeightedAdditiveResult = CurrentValue - ResultAbsolute.Normalize() - Result.Additive;
diff --git a/Engine/Source/Runtime/MovieScene/Private/MovieSceneSection.cpp b/Engine/Source/Runtime/MovieScene/Private/MovieSceneSection.cpp
--- a/Engine/Source/Runtime/MovieScene/Private/MovieSceneSection.cpp
+++ b/Engine/Source/Runtime/MovieScene/Private/MovieSceneSection.cpp
@@ -341,6 +341,41 @@ FMovieSceneBlendTypeField UMovieSceneSection::GetSupportedBlendTypes() const
 	return Track ? Track->GetSupportedBlendTypes() : FMovieSceneBlendTypeField::None();
 }
 
+int32 UMovieSceneSection::GetBlendingOrder() const
+{
+	//currently needs to support overrides and have one present
+	int32 BlendingOrder = INDEX_NONE;
+	bool bHasOverride = false;
+	if (GetBlendType().IsValid() && GetBlendType().BlendType == EMovieSceneBlendType::Absolute)
+	{
+		return BlendingOrder;
+	}
+
+	if (UMovieSceneTrack* Track = GetTypedOuter<UMovieSceneTrack>())
+	{
+		if (Track->GetSupportedBlendTypes().Contains(EMovieSceneBlendType::Override))
+		{
+			const TArray<UMovieSceneSection*>& Sections = Track->GetAllSections();
+			for (int32 Index = 0; Index < Sections.Num(); ++Index)
+			{
+				if (bHasOverride == false && Sections[Index]->GetBlendType().IsValid() && Sections[Index]->GetBlendType() == EMovieSceneBlendType::Override)
+				{
+					bHasOverride = true;
+				}
+				if (Sections[Index] == this)
+				{
+					BlendingOrder = Index;
+				}
+				if (bHasOverride && BlendingOrder != INDEX_NONE)
+				{
+					break;
+				}
+			}
+		}
+	}
+	return bHasOverride ? BlendingOrder : INDEX_NONE;
+}
+
 void UMovieSceneSection::BuildDefaultComponents(UMovieSceneEntitySystemLinker* EntityLinker, const UE::MovieScene::FEntityImportParams& Params, UE::MovieScene::FImportedEntity* OutImportedEntity)
 {
 	using namespace UE::MovieScene;
@@ -366,8 +401,14 @@ void UMovieSceneSection::BuildDefaultComponents(UMovieSceneEntitySystemLinker* E
 		{
 			BlendTag = Components->Tags.AdditiveFromBaseBlend;
 		}
+		else if (BlendType.Get() == EMovieSceneBlendType::Override)
+		{
+			BlendTag = Components->Tags.OverrideBlend;
+		}
 	}
 
+	int32 BlendingOrder = GetBlendingOrder();
+
 	const bool bHasEasing = (Easing.GetEaseInDuration() > 0 || Easing.GetEaseOutDuration() > 0);
 
 	// Should restore state if we're not forcing keep state and any one of the following:
@@ -407,6 +448,7 @@ void UMovieSceneSection::BuildDefaultComponents(UMovieSceneEntitySystemLinker* E
 	OutImportedEntity->AddBuilder(
 		FEntityBuilder()
 		.AddConditional(Components->BlenderType,                    BlenderSystemClass, BlenderSystemClass.Get() != nullptr)
+		.AddConditional(Components->BlendingOrder,					BlendingOrder, BlendingOrder != INDEX_NONE)
 		.AddConditional(Components->Easing,                         FEasingComponentData{ decltype(FEasingComponentData::Section)(this) }, bHasEasing)
 		.AddConditional(Components->HierarchicalBias,               Params.Sequence.HierarchicalBias, Params.Sequence.HierarchicalBias != 0)
 		.AddConditional(Components->Interrogation.InputKey,         Params.InterrogationKey, Params.InterrogationKey.IsValid())
diff --git a/Engine/Source/Runtime/MovieScene/Public/EntitySystem/BuiltInComponentTypes.h b/Engine/Source/Runtime/MovieScene/Public/EntitySystem/BuiltInComponentTypes.h
--- a/Engine/Source/Runtime/MovieScene/Public/EntitySystem/BuiltInComponentTypes.h
+++ b/Engine/Source/Runtime/MovieScene/Public/EntitySystem/BuiltInComponentTypes.h
@@ -710,6 +710,9 @@ public:
 	/** A blender type that should be used for blending this entity */
 	TComponentTypeID<TSubclassOf<UMovieSceneBlenderSystem>> BlenderType;
 
+	/** Value that is used to sort blending order */
+	TComponentTypeID<int32> BlendingOrder;
+
 	// An FMovieSceneTrackInstanceComponent that defines the track instance to use
 	TComponentTypeID<FMovieSceneTrackInstanceComponent> TrackInstance;
 
@@ -745,6 +748,7 @@ public:
 		FComponentTypeID RelativeBlend;
 		FComponentTypeID AdditiveBlend;
 		FComponentTypeID AdditiveFromBaseBlend;
+		FComponentTypeID OverrideBlend;
 
 		FComponentTypeID NeedsLink;
 		FComponentTypeID NeedsUnlink;
diff --git a/Engine/Source/Runtime/MovieScene/Public/EntitySystem/MovieSceneDecompositionQuery.h b/Engine/Source/Runtime/MovieScene/Public/EntitySystem/MovieSceneDecompositionQuery.h
--- a/Engine/Source/Runtime/MovieScene/Public/EntitySystem/MovieSceneDecompositionQuery.h
+++ b/Engine/Source/Runtime/MovieScene/Public/EntitySystem/MovieSceneDecompositionQuery.h
@@ -49,12 +49,19 @@ namespace MovieScene
 		float Weight = 0.f;
 		/** The base value, which should only be set when the value has additive-from-base blend type */
 		double BaseValue = 0.f;
-
+		/** Optional Blending Order*/
+		int32 BlendingOrder = INDEX_NONE;
+		/** if BlendingOrder set, whether additive or override*/
+		bool bIsAdditive = true;
 		/** Multiplies the value by the weight */
 		double Get() const
 		{
 			return (Value - BaseValue) * Weight;
 		}
+		bool operator<(const FWeightedValue& RHS) const
+		{
+			return BlendingOrder < RHS.BlendingOrder;
+		}
 	};
 
 	/** Holds an accumulation of weighted values */
@@ -117,12 +124,16 @@ namespace MovieScene
 		/** The absolute and additive values for all the non-decomposed channels */
 		FResult Result;
 
+		/** values for non-decomposed overides, we leave it decomposed */
+		mutable TArray <FWeightedValue> AllDecomposedOverrides;
 		/**
 		 * Decomposed values for channels we're interested in. Note that decomposed additives-from-base are lumped with
-		 * decomposed additives.
+		 * decomposed additives
 		 */
 		TArray<TTuple<FMovieSceneEntityID, FWeightedValue>> DecomposedAbsolutes;
 		TArray<TTuple<FMovieSceneEntityID, FWeightedValue>> DecomposedAdditives;
+		TArray<TTuple<FMovieSceneEntityID, FWeightedValue>> DecomposedOverrides;
+
 
 		/**
 		 * Get the value that the channel behind the given entity should have in order for the combined blended values
@@ -141,7 +152,8 @@ namespace MovieScene
 		{
 			Absolute,
 			Additive,
-			AdditiveFromBase
+			AdditiveFromBase,
+			Override,
 		};
 
 		/**
diff --git a/Engine/Source/Runtime/MovieScene/Public/Evaluation/Blending/BlendableToken.h b/Engine/Source/Runtime/MovieScene/Public/Evaluation/Blending/BlendableToken.h
--- a/Engine/Source/Runtime/MovieScene/Public/Evaluation/Blending/BlendableToken.h
+++ b/Engine/Source/Runtime/MovieScene/Public/Evaluation/Blending/BlendableToken.h
@@ -79,8 +79,13 @@ namespace MovieScene
 	template<typename InType, typename WorkingDataType, typename SourceDataType>
 	void BlendValue(WorkingDataType& OutBlend, InType InValue, float Weight, EMovieSceneBlendType BlendType, TMovieSceneInitialValueStore<SourceDataType>& InitialValueStore)
 	{
-		// Always assert on instantiation
-		static_assert(std::is_same_v<WorkingDataType, void>, "BlendValue must be implemented for the specified types in order to blend them with Sequencer.");
+		BlendValue(OutBlend, InValue, Weight, BlendType, INDEX_NONE, InitialValueStore);
+	}
+	//new optional blending priority that supports blending priorities
+	template<typename InType, typename WorkingDataType, typename SourceDataType>
+	void BlendValue(WorkingDataType& OutBlend, InType InValue, float Weight, EMovieSceneBlendType BlendType, int32 BlendingOrder, TMovieSceneInitialValueStore<SourceDataType>& InitialValueStore)
+	{
+	
 	}
 } // namespace MovieScene
 } // namespace UE
@@ -100,21 +105,25 @@ struct TBlendableToken
 
 	/** Construction from a value, blend method, and a weight. Scope and bias to be populated later */
 	template<typename T>
-	TBlendableToken(T&& InValue, EMovieSceneBlendType InBlendType, float InWeight = 1.f)
+	TBlendableToken(T&& InValue, EMovieSceneBlendType InBlendType, float InWeight = 1.f, int32 InBlendingOrder = INDEX_NONE)
 		: Value(TData<typename TDecay<T>::Type>(Forward<T>(InValue)))
 		, HierarchicalBias(0)
 		, Weight(InWeight)
 		, BlendType(InBlendType)
+		, BlendingOrder(InBlendingOrder)
 	{}
 
 	/** Construction from a value, scope, context, blend method, and a weight */
 	template<typename T>
-	TBlendableToken(T&& InValue, const FMovieSceneEvaluationScope& InCurrentScope, const FMovieSceneContext& InContext, EMovieSceneBlendType InBlendType, float InWeight = 1.f)
+	TBlendableToken(T&& InValue, const FMovieSceneEvaluationScope& InCurrentScope, const FMovieSceneContext& InContext, EMovieSceneBlendType InBlendType, float InWeight = 1.f,
+		int32 InBlendingOrder = INDEX_NONE)
 		: Value(TData<typename TDecay<T>::Type>(Forward<T>(InValue)))
 		, AnimatingScope(InCurrentScope)
 		, HierarchicalBias(InContext.GetHierarchicalBias())
 		, Weight(InWeight)
 		, BlendType(InBlendType)
+		, BlendingOrder(InBlendingOrder)
+
 	{}
 
 	/** Copying is disabled */
@@ -133,16 +142,20 @@ struct TBlendableToken
 	void AddTo(WorkingDataType& CumulativeBlend, TMovieSceneInitialValueStore<DataType>& InitialValueStore) const
 	{
 		check(Value.IsValid());
-		Value->AddTo(CumulativeBlend, Weight, BlendType, InitialValueStore);
+		Value->AddTo(CumulativeBlend, Weight, BlendType, BlendingOrder, InitialValueStore);
 	}
 
+	bool operator<(const TBlendableToken& RHS) const
+	{
+		return BlendingOrder < RHS.BlendingOrder;
+	}
 private:
 
 	/** Base class for all value types */
 	struct IData
 	{
 		virtual ~IData() {}
-		virtual void AddTo(WorkingDataType& CumulativeBlend, float Weight, EMovieSceneBlendType BlendType, TMovieSceneInitialValueStore<DataType>& InitialValueStore) const = 0;
+		virtual void AddTo(WorkingDataType& CumulativeBlend, float Weight, EMovieSceneBlendType BlendType, int32 BlendingOrder, TMovieSceneInitialValueStore<DataType>& InitialValueStore) const = 0;
 	};
 
 	/** Templated value data for any other type */
@@ -151,11 +164,18 @@ private:
 	{
 		TData(T In) : Data(MoveTemp(In)) {}
 
-		virtual void AddTo(WorkingDataType& CumulativeBlend, float InWeight, EMovieSceneBlendType InBlendType, TMovieSceneInitialValueStore<DataType>& InitialValueStore) const
+		virtual void AddTo(WorkingDataType& CumulativeBlend, float InWeight, EMovieSceneBlendType InBlendType, int32 InBlendingOrder, TMovieSceneInitialValueStore<DataType>& InitialValueStore) const
 		{
 			// Use the default BlendValue function, or any other BlendValue function found through ADL on WorkingDataType
 			using UE::MovieScene::BlendValue;
-			BlendValue(CumulativeBlend, Data, InWeight, InBlendType, InitialValueStore);
+			if (InBlendingOrder == INDEX_NONE)
+			{
+				BlendValue(CumulativeBlend, Data, InWeight, InBlendType, InitialValueStore);
+			}
+			else
+			{
+				BlendValue(CumulativeBlend, Data, InWeight, InBlendType, InBlendingOrder, InitialValueStore);
+			}
 		}
 
 		/** The actual value to blend */
@@ -178,4 +198,8 @@ public:
 
 	/** Enumeration specifying how this token should be blended */
 	EMovieSceneBlendType BlendType;
+
+	/** Opitonal Blending Order, if present will be used to sort additives and override blend types*/
+	int32 BlendingOrder = INDEX_NONE;
+
 };
diff --git a/Engine/Source/Runtime/MovieScene/Public/Evaluation/Blending/MovieSceneBlendingAccumulator.h b/Engine/Source/Runtime/MovieScene/Public/Evaluation/Blending/MovieSceneBlendingAccumulator.h
--- a/Engine/Source/Runtime/MovieScene/Public/Evaluation/Blending/MovieSceneBlendingAccumulator.h
+++ b/Engine/Source/Runtime/MovieScene/Public/Evaluation/Blending/MovieSceneBlendingAccumulator.h
@@ -44,10 +44,10 @@ struct FMovieSceneBlendingAccumulator
 		UnboundBlendState.Add(InActuatorTypeID, MoveTemp(InToken));
 	}
 	template<typename ActuatorDataType>
-	void BlendToken(FMovieSceneBlendingActuatorID InActuatorTypeID, const FMovieSceneEvaluationScope& InScope, const FMovieSceneContext& InContext, ActuatorDataType&& InInputValue, EMovieSceneBlendType InBlendType, float InWeight = 1.f)
+	void BlendToken(FMovieSceneBlendingActuatorID InActuatorTypeID, const FMovieSceneEvaluationScope& InScope, const FMovieSceneContext& InContext, ActuatorDataType&& InInputValue, EMovieSceneBlendType InBlendType, float InWeight = 1.f, int32 InBlendingPriority = INDEX_NONE)
 	{
 		check(FindActuator<ActuatorDataType>(InActuatorTypeID));
-		UnboundBlendState.Add(InActuatorTypeID, TBlendableToken<ActuatorDataType>(Forward<ActuatorDataType>(InInputValue), InScope, InContext, InBlendType, InWeight));
+		UnboundBlendState.Add(InActuatorTypeID, TBlendableToken<ActuatorDataType>(Forward<ActuatorDataType>(InInputValue), InScope, InContext, InBlendType, InWeight, InBlendingPriority));
 	}
 
 	/**
@@ -69,10 +69,10 @@ struct FMovieSceneBlendingAccumulator
 		OperandToBlendState.FindOrAdd(InOperand).Add(InActuatorTypeID, MoveTemp(InToken));
 	}
 	template<typename ActuatorDataType>
-	void BlendToken(const FMovieSceneEvaluationOperand& InOperand, FMovieSceneBlendingActuatorID InActuatorTypeID, const FMovieSceneEvaluationScope& InScope, const FMovieSceneContext& InContext, ActuatorDataType&& InInputValue, EMovieSceneBlendType InBlendType, float InWeight = 1.f)
+	void BlendToken(const FMovieSceneEvaluationOperand& InOperand, FMovieSceneBlendingActuatorID InActuatorTypeID, const FMovieSceneEvaluationScope& InScope, const FMovieSceneContext& InContext, ActuatorDataType&& InInputValue, EMovieSceneBlendType InBlendType, float InWeight = 1.f, int32 InBlendingPriority = INDEX_NONE)
 	{
 		check(FindActuator<ActuatorDataType>(InActuatorTypeID));
-		OperandToBlendState.FindOrAdd(InOperand).Add(InScope, InActuatorTypeID, TBlendableToken<ActuatorDataType>(Forward<ActuatorDataType>(InInputValue), InScope, InContext, InBlendType, InWeight));
+		OperandToBlendState.FindOrAdd(InOperand).Add(InScope, InActuatorTypeID, TBlendableToken<ActuatorDataType>(Forward<ActuatorDataType>(InInputValue), InScope, InContext, InBlendType, InWeight, InBlendingPriority));
 	}
 
 	/**
diff --git a/Engine/Source/Runtime/MovieScene/Public/Evaluation/Blending/MovieSceneBlendingAccumulator.inl b/Engine/Source/Runtime/MovieScene/Public/Evaluation/Blending/MovieSceneBlendingAccumulator.inl
--- a/Engine/Source/Runtime/MovieScene/Public/Evaluation/Blending/MovieSceneBlendingAccumulator.inl
+++ b/Engine/Source/Runtime/MovieScene/Public/Evaluation/Blending/MovieSceneBlendingAccumulator.inl
@@ -17,6 +17,7 @@ void TBlendableTokenStack<DataType>::ComputeAndActuate(UObject* InObject, FMovie
 	TMovieSceneInitialValueStore<DataType> InitialValues(*Actuator, *this, InObject, &Player);
 
 	typename TBlendableTokenTraits<DataType>::WorkingDataType WorkingTotal{};
+	Tokens.Sort();
 	for (const TBlendableToken<DataType>* Token : Tokens)
 	{
 		Token->AddTo(WorkingTotal, InitialValues);
diff --git a/Engine/Source/Runtime/MovieScene/Public/Evaluation/Blending/MovieSceneMultiChannelBlending.h b/Engine/Source/Runtime/MovieScene/Public/Evaluation/Blending/MovieSceneMultiChannelBlending.h
--- a/Engine/Source/Runtime/MovieScene/Public/Evaluation/Blending/MovieSceneMultiChannelBlending.h
+++ b/Engine/Source/Runtime/MovieScene/Public/Evaluation/Blending/MovieSceneMultiChannelBlending.h
@@ -138,6 +138,15 @@ namespace MovieScene
 			Mask |= (1 << Index);
 		}
 
+		/** Do weighted blend with current value, this is an override*/
+		void WeightedBlend(uint8 Index, T Value, float Weight)
+		{
+			check(Index < N);
+
+			Channels[Index] = IsSet(Index) ? ((1.0 - Weight)*Channels[Index]) + (Value * Weight) : Value;
+			Mask |= (1 << Index);
+		}
+
 	private:
 
 		/** Channel data */
@@ -182,6 +191,9 @@ namespace MovieScene
 		/** Per-channel additive values to apply, pre-multiplied by their weight */
 		TMultiChannelValue<DataType, N> Additive;
 
+		/** Per-channel override values to apply, when set this will override the default absolute/additive set */
+		TMultiChannelValue<DataType, N> Override;
+
 		/** Cached initial value for this blendable in multi-channel form */
 		TOptional<TMultiChannelValue<DataType, N>> InitialValue;
 
@@ -197,27 +209,34 @@ namespace MovieScene
 			for (uint8 Channel = 0; Channel < N; ++Channel)
 			{
 				// Any animated channels with a weight of 0 should match the object's *initial* position. Exclusively additive channels are based implicitly off the initial value
-				const bool bUseInitialValue = (Absolute.IsSet(Channel) && AbsoluteWeights[Channel] == 0.f) || (!Absolute.IsSet(Channel) && Additive.IsSet(Channel));
-				if (bUseInitialValue)
+				if (Override.IsSet(Channel)) //if override then we override
 				{
-					if (!InitialValue.IsSet())
-					{
-						InitialValue.Emplace();
-						MultiChannelFromData(InitialValueStore.GetInitialValue(), InitialValue.GetValue());
-					}
-
-					Result.Set(Channel, InitialValue.GetValue()[Channel]);
+					Result.Set(Channel, Override[Channel]);
 				}
-				else if (Absolute.IsSet(Channel))
+				else
 				{
-					// If it has a non-zero weight, divide by it, and apply the absolute total to the result
-					Result.Set(Channel, Absolute[Channel] / AbsoluteWeights[Channel]);
-				}
+					const bool bUseInitialValue = (Absolute.IsSet(Channel) && AbsoluteWeights[Channel] == 0.f) || (!Absolute.IsSet(Channel) && Additive.IsSet(Channel));
+					if (bUseInitialValue)
+					{
+						if (!InitialValue.IsSet())
+						{
+							InitialValue.Emplace();
+							MultiChannelFromData(InitialValueStore.GetInitialValue(), InitialValue.GetValue());
+						}
 
-				// If it has any additive values in the channel, add those on
-				if (Additive.IsSet(Channel))
-				{
-					Result.Increment(Channel, Additive[Channel]);
+						Result.Set(Channel, InitialValue.GetValue()[Channel]);
+					}
+					else if (Absolute.IsSet(Channel))
+					{
+						// If it has a non-zero weight, divide by it, and apply the absolute total to the result
+						Result.Set(Channel, Absolute[Channel] / AbsoluteWeights[Channel]);
+					}
+
+					// If it has any additive values in the channel, add those on
+					if (Additive.IsSet(Channel))
+					{
+						Result.Increment(Channel, Additive[Channel]);
+					}
 				}
 
 				// If the channel has not been animated at all, set it to the *current* value
@@ -242,7 +261,7 @@ namespace MovieScene
 	};
 
 	template<typename OutputType, typename InputType, typename ActualValueType, uint8 ChannelSize>
-	void BlendValue(TMaskedBlendable<OutputType, ChannelSize>& OutBlend, InputType InValue, int32 ChannelIndex, float Weight, EMovieSceneBlendType BlendType, TMovieSceneInitialValueStore<ActualValueType>& InitialValueStore)
+	void BlendValue(TMaskedBlendable<OutputType, ChannelSize>& OutBlend, InputType InValue, int32 ChannelIndex, float Weight, EMovieSceneBlendType BlendType, int32 BlendingOrder, TMovieSceneInitialValueStore<ActualValueType>& InitialValueStore)
 	{
 		if (BlendType == EMovieSceneBlendType::Absolute || BlendType == EMovieSceneBlendType::Relative)
 		{
@@ -266,6 +285,29 @@ namespace MovieScene
 			// Accumulate total weights
 			OutBlend.AbsoluteWeights[ChannelIndex] += Weight;
 		}
+		else if (BlendingOrder != INDEX_NONE)
+		{
+			//If there is a blending order set we use the additive to contain
+			//the full blended value consiting of absolute + overrides + additives
+			//This will already be sorted so we just need to accumulate the values
+			if (OutBlend.Override.IsSet(ChannelIndex) == false)
+			{
+				if (OutBlend.Absolute.IsSet(ChannelIndex))
+				{
+					// If it has a non-zero weight, divide by it, and apply the absolute total to the result
+					OutputType AbsoluteWeight = (OutBlend.AbsoluteWeights[ChannelIndex] != 0.f) ? OutBlend.AbsoluteWeights[ChannelIndex] : 1.0;
+					OutBlend.Override.Set(ChannelIndex, OutBlend.Absolute[ChannelIndex] / AbsoluteWeight);
+				}
+			}
+			if (BlendType == EMovieSceneBlendType::Additive)
+			{
+				OutBlend.Override.Increment(ChannelIndex, OutputType(InValue) * Weight);
+			}
+			else if (BlendType == EMovieSceneBlendType::Override)
+			{
+				OutBlend.Override.WeightedBlend(ChannelIndex, OutputType(InValue), Weight);
+			}
+		}
 		else if (BlendType == EMovieSceneBlendType::Additive)
 		{
 			// Additive animation just increments the additive channel
@@ -276,19 +318,19 @@ namespace MovieScene
 	}
 
 	template<typename OutputType, typename InputType, typename ActualValueType>
-	void BlendValue(TMaskedBlendable<OutputType, 1>& OutBlend, InputType InValue, float Weight, EMovieSceneBlendType BlendType, TMovieSceneInitialValueStore<ActualValueType>& InitialValueStore)
+	void BlendValue(TMaskedBlendable<OutputType, 1>& OutBlend, InputType InValue, float Weight, EMovieSceneBlendType BlendType, int32 BlendingOrder, TMovieSceneInitialValueStore<ActualValueType>& InitialValueStore)
 	{
-		BlendValue(OutBlend, InValue, 0, Weight, BlendType, InitialValueStore);
+		BlendValue(OutBlend, InValue, 0, Weight, BlendType, BlendingOrder, InitialValueStore);
 	}
 
 	template<typename OutputType, typename ActualValueType, uint8 ChannelSize>
-	void BlendValue(TMaskedBlendable<OutputType, ChannelSize>& OutBlend, const TMultiChannelValue<OutputType, ChannelSize>& InValue, float Weight, EMovieSceneBlendType BlendType, TMovieSceneInitialValueStore<ActualValueType>& InitialValueStore)
+	void BlendValue(TMaskedBlendable<OutputType, ChannelSize>& OutBlend, const TMultiChannelValue<OutputType, ChannelSize>& InValue, float Weight, EMovieSceneBlendType BlendType, int32 BlendingOrder, TMovieSceneInitialValueStore<ActualValueType>& InitialValueStore)
 	{
 		for (int32 Index = 0; Index < ChannelSize; ++Index)
 		{
 			if (InValue.IsSet(Index))
 			{
-				BlendValue(OutBlend, InValue[Index], Index, Weight, BlendType, InitialValueStore);
+				BlendValue(OutBlend, InValue[Index], Index, Weight, BlendType, BlendingOrder, InitialValueStore);
 			}
 		}
 	}
diff --git a/Engine/Source/Runtime/MovieScene/Public/MovieSceneSection.h b/Engine/Source/Runtime/MovieScene/Public/MovieSceneSection.h
--- a/Engine/Source/Runtime/MovieScene/Public/MovieSceneSection.h
+++ b/Engine/Source/Runtime/MovieScene/Public/MovieSceneSection.h
@@ -442,6 +442,12 @@ public:
 	 */
 	MOVIESCENE_API FMovieSceneBlendTypeField GetSupportedBlendTypes() const;
 
+	/**
+	 * Get the optional order for the blending of this section that may or may not be used by the blending system
+	 * Lower values blend first, value of INDEX_NONE means no order is set
+	 */
+	MOVIESCENE_API int32 GetBlendingOrder() const;
+
 	/**
 	 * Moves the section by a specific amount of time
 	 *
diff --git a/Engine/Source/Runtime/MovieSceneTracks/Private/EntitySystem/Interrogation/MovieSceneInterrogatedPropertyInstantiator.cpp b/Engine/Source/Runtime/MovieSceneTracks/Private/EntitySystem/Interrogation/MovieSceneInterrogatedPropertyInstantiator.cpp
--- a/Engine/Source/Runtime/MovieSceneTracks/Private/EntitySystem/Interrogation/MovieSceneInterrogatedPropertyInstantiator.cpp
+++ b/Engine/Source/Runtime/MovieSceneTracks/Private/EntitySystem/Interrogation/MovieSceneInterrogatedPropertyInstantiator.cpp
@@ -111,6 +111,7 @@ bool UMovieSceneInterrogatedPropertyInstantiatorSystem::PropertySupportsFastPath
 		FComponentMask Type = Linker->EntityManager.GetEntityType(Input);
 		if (Type.Contains(BuiltInComponents->Tags.RelativeBlend) ||
 			Type.Contains(BuiltInComponents->Tags.AdditiveBlend) ||
+			Type.Contains(BuiltInComponents->Tags.OverrideBlend) ||
 			Type.Contains(BuiltInComponents->Tags.AdditiveFromBaseBlend) ||
 			Type.Contains(BuiltInComponents->WeightAndEasingResult))
 		{
diff --git a/Engine/Source/Runtime/MovieSceneTracks/Private/Systems/MovieScenePiecewiseDoubleBlenderSystem.cpp b/Engine/Source/Runtime/MovieSceneTracks/Private/Systems/MovieScenePiecewiseDoubleBlenderSystem.cpp
--- a/Engine/Source/Runtime/MovieSceneTracks/Private/Systems/MovieScenePiecewiseDoubleBlenderSystem.cpp
+++ b/Engine/Source/Runtime/MovieSceneTracks/Private/Systems/MovieScenePiecewiseDoubleBlenderSystem.cpp
@@ -204,12 +204,108 @@ private:
 	}
 };
 
+/** Task to gather all of the additives and overrides that need to get blended in order.
+ *
+ *  Only used by entities with the specified blending order.
+ */
+struct FAdditiveAndOverridesTask
+{
+	TSortedMap<FComponentTypeID, FAdditiveAndOverrideBuffer>* AccumulationBuffers;
+
+	void ForEachAllocation(FEntityAllocationIteratorItem InItem, TRead<FMovieSceneBlendChannelID> BlendIDs, TRead<int32> BlendingOrder, TReadOptional<double> EasingAndWeightResults) const
+	{
+		FEntityAllocation* Allocation = InItem;
+		const FComponentMask& AllocationType = InItem;
+
+		for (const FComponentHeader& ComponentHeader : Allocation->GetComponentHeaders())
+		{
+			if (FAdditiveAndOverrideBuffer* Buffer = AccumulationBuffers->Find(ComponentHeader.ComponentType))
+			{
+				TComponentReader<double> Results(&ComponentHeader, Allocation->GetCurrentLockMode());
+
+				AccumulateResults(Allocation, Results.AsPtr(), BlendingOrder, BlendIDs, EasingAndWeightResults, Buffer->Values);
+			}
+		}
+	}
+
+
+private:
+
+	void AccumulateResults(const FEntityAllocation* InAllocation, const double* Results, const int32* BlendingOrder, const FMovieSceneBlendChannelID* BlendIDs, const double* OptionalEasingAndWeights, TArray<FAdditiveAndOverrides>& OutOverrides) const
+	{
+		static const FMovieSceneBlenderSystemID BlenderSystemID = UMovieSceneBlenderSystem::GetBlenderSystemID<UMovieScenePiecewiseDoubleBlenderSystem>();
+
+		const int32 Num = InAllocation->Num();
+		const bool bAdditive = InAllocation->HasComponent(FBuiltInComponentTypes::Get()->Tags.AdditiveBlend);
+		const bool bOverride = InAllocation->HasComponent(FBuiltInComponentTypes::Get()->Tags.OverrideBlend);
+		if (!bAdditive && !bOverride)
+		{
+			return;
+		}
+		if (OptionalEasingAndWeights)
+		{
+			for (int32 Index = 0; Index < Num; ++Index)
+			{
+				const FMovieSceneBlendChannelID& BlendID(BlendIDs[Index]);
+				ensureMsgf(BlendID.SystemID == BlenderSystemID, TEXT("Overriding the standard blender system of standard types isn't supported."));
+
+				FAdditiveAndOverrides& Result = OutOverrides[BlendID.ChannelID];
+				FAdditveAndOverrideData Data;
+				Data.bIsAdditive = bAdditive;
+				Data.Value = Results[Index];
+				Data.Weight = OptionalEasingAndWeights[Index];
+				Data.BlendingOrder = BlendingOrder[Index];
+				Result.Data.Add(Data.BlendingOrder, Data);
+			}
+		}
+		else
+		{
+			// Faster path for when there's no weight to multiply values with.
+			for (int32 Index = 0; Index < Num; ++Index)
+			{
+				const FMovieSceneBlendChannelID& BlendID(BlendIDs[Index]);
+				ensureMsgf(BlendID.SystemID == BlenderSystemID, TEXT("Overriding the standard blender system of standard types isn't supported."));
+				FAdditiveAndOverrides& Result = OutOverrides[BlendID.ChannelID];
+				FAdditveAndOverrideData Data;
+				Data.bIsAdditive = bAdditive;
+				Data.Value = Results[Index];
+				Data.Weight = 1.0;
+				Data.BlendingOrder = BlendingOrder[Index];
+				Result.Data.Add(Data.BlendingOrder, Data);
+			}
+		}
+	}
+};
+
+void AddOverrideResults(const FAdditiveAndOverrides& Overrides, double& OutFinalBlendValue)
+{
+	for (const TPair<int32,FAdditveAndOverrideData>& Data : Overrides.Data)
+	{
+		if (Data.Value.bIsAdditive)
+		{
+			OutFinalBlendValue += (Data.Value.Value * Data.Value.Weight);
+		}
+		else //override do weighted blend
+		{
+			if (Data.Value.Weight >= 1.0)
+			{
+				OutFinalBlendValue = Data.Value.Value;
+			}
+			else
+			{
+				OutFinalBlendValue =(OutFinalBlendValue *(1.0 - Data.Value.Weight)) +
+					(Data.Value.Value * Data.Value.Weight);
+			}
+		}
+	}
+}
+
 void BlendResults(const FAccumulationResult& Results, uint16 BlendID, double& OutFinalBlendResult)
 {
 	FBlendResult AbsoluteResult = Results.GetAbsoluteResult(BlendID);
 	FBlendResult AdditiveResult = Results.GetAdditiveResult(BlendID);
 	FBlendResult AdditiveFromBaseResult = Results.GetAdditiveFromBaseResult(BlendID);
-
+	FAdditiveAndOverrides Overrides = Results.GetAdditiveAndOverrideResult(BlendID);
 #if DO_GUARD_SLOW
 	ensureMsgf(AbsoluteResult.Weight != 0.f, TEXT("Default blend combine being used for an entity that has no absolute weight. This should have an initial value and should be handled by each system, and excluded by default with UMovieSceneBlenderSystem::FinalCombineExclusionFilter ."));
 #endif
@@ -220,6 +316,10 @@ void BlendResults(const FAccumulationResult& Results, uint16 BlendID, double& Ou
 		const double Value = AbsoluteResult.Total / AbsoluteResult.Weight + AdditiveResult.Total + AdditiveFromBaseResult.Total;
 		OutFinalBlendResult = Value;
 	}
+	if (Overrides.Data.Num() > 0)
+	{
+		AddOverrideResults(Overrides, OutFinalBlendResult);
+	}
 }
 
 void BlendResultsWithInitial(const FAccumulationResult& Results, uint16 BlendID, const double InitialValue, double& OutFinalBlendResult)
@@ -228,6 +328,7 @@ void BlendResultsWithInitial(const FAccumulationResult& Results, uint16 BlendID,
 	FBlendResult RelativeResult = Results.GetRelativeResult(BlendID);
 	FBlendResult AdditiveResult = Results.GetAdditiveResult(BlendID);
 	FBlendResult AdditiveFromBaseResult = Results.GetAdditiveFromBaseResult(BlendID);
+	FAdditiveAndOverrides Overrides = Results.GetAdditiveAndOverrideResult(BlendID);
 
 	if (RelativeResult.Weight != 0)
 	{
@@ -235,7 +336,6 @@ void BlendResultsWithInitial(const FAccumulationResult& Results, uint16 BlendID,
 	}
 
 	FBlendResult TotalAdditiveResult = { AdditiveResult.Total + AdditiveFromBaseResult.Total, AdditiveResult.Weight + AdditiveFromBaseResult.Weight };
-
 	const float TotalWeight = AbsoluteResult.Weight + RelativeResult.Weight;
 	if (TotalWeight != 0)
 	{
@@ -262,6 +362,10 @@ void BlendResultsWithInitial(const FAccumulationResult& Results, uint16 BlendID,
 	{
 		OutFinalBlendResult = InitialValue;
 	}
+	if (Overrides.Data.Num() > 0)
+	{
+		AddOverrideResults(Overrides, OutFinalBlendResult);
+	}
 }
 
 /** Task that combines all accumulated blends for any tracked property type that has blend inputs/outputs */
@@ -458,7 +562,8 @@ private:
 
 bool FAccumulationBuffers::IsEmpty() const
 {
-	return Absolute.Num() == 0 && Relative.Num() == 0 && Additive.Num() == 0 && AdditiveFromBase.Num() == 0;
+	return Absolute.Num() == 0 && Relative.Num() == 0 && Additive.Num() == 0 && AdditiveFromBase.Num() == 0
+		&& AdditiveAndOverrides.Num() == 0;
 }
 
 void FAccumulationBuffers::Reset()
@@ -467,6 +572,7 @@ void FAccumulationBuffers::Reset()
 	Relative.Empty();
 	Additive.Empty();
 	AdditiveFromBase.Empty();
+	AdditiveAndOverrides.Empty();
 }
 
 FAccumulationResult FAccumulationBuffers::FindResults(FComponentTypeID InComponentType) const
@@ -488,6 +594,10 @@ FAccumulationResult FAccumulationBuffers::FindResults(FComponentTypeID InCompone
 	{
 		Result.AdditivesFromBase = AdditivesFromBase->Buffer.GetData();
 	}
+	if (const FAdditiveAndOverrideBuffer* AandOs = AdditiveAndOverrides.Find(InComponentType))
+	{
+		Result.Overrides = AandOs->Values.GetData();
+	}
 	return Result;
 }
 
@@ -617,6 +727,23 @@ void UMovieScenePiecewiseDoubleBlenderSystem::OnSchedulePersistentTasks(UE::Movi
 		TaskScheduler->AddPrerequisite(AdditiveFromBaseTask, SyncTask);
 	}
 
+	if (AccumulationBuffers.AdditiveAndOverrides.Num() != 0)
+	{
+		FTaskID OverrideTask = FEntityTaskBuilder()
+			.Read(BuiltInComponents->BlendChannelInput)
+			.Read(BuiltInComponents->BlendingOrder)
+			.ReadOptional(BuiltInComponents->WeightAndEasingResult)
+			.FilterAll({ GetBlenderTypeTag() })
+			.FilterAny({ BlendedResultMask | BuiltInComponents->Tags.AdditiveBlend |
+					BuiltInComponents->Tags.OverrideBlend }).FilterNone({ BuiltInComponents->Tags.Ignored })
+			.AddDynamicReadDependency(BlendedResultMask)
+			.SetStat(BlendValuesStatId)
+			.Schedule_PerAllocation<FAdditiveAndOverridesTask>(&EntityManager, TaskScheduler, &AccumulationBuffers.AdditiveAndOverrides);
+
+		TaskScheduler->AddPrerequisite(ResetWeightsTask, OverrideTask);
+		TaskScheduler->AddPrerequisite(OverrideTask, SyncTask);
+	}
+
 	// Combine blends for all blend outputs on properties
 	if (BlendedPropertyMask.Find(true) != INDEX_NONE)
 	{
@@ -756,6 +883,29 @@ void UMovieScenePiecewiseDoubleBlenderSystem::OnRun(FSystemTaskPrerequisites& In
 		}
 	}
 
+	if (AccumulationBuffers.AdditiveAndOverrides.Num() != 0)
+	{
+		for (TPair<FComponentTypeID, FAdditiveAndOverrideBuffer>& Pair : AccumulationBuffers.AdditiveAndOverrides)
+		{
+			FGraphEventRef Task = FEntityTaskBuilder()
+				.Read(BuiltInComponents->BlendChannelInput)
+				.Read(BuiltInComponents->BlendingOrder)
+				.ReadOptional(BuiltInComponents->WeightAndEasingResult)
+				.FilterAll({ GetBlenderTypeTag() })
+				.FilterAny({ BlendedResultMask | BuiltInComponents->Tags.AdditiveBlend |
+					 BuiltInComponents->Tags.OverrideBlend })
+				.FilterNone({ BuiltInComponents->Tags.Ignored })
+				.AddDynamicReadDependency(BlendedResultMask)
+				.SetStat(BlendValuesStatId)
+				.Dispatch_PerAllocation<FAdditiveAndOverridesTask>(&EntityManager, InPrerequisites, nullptr, &AccumulationBuffers.AdditiveAndOverrides);
+
+			if (Task)
+			{
+				Prereqs.AddRootTask(Task);
+			}
+		}
+	}
+
 	if (AccumulationBuffers.AdditiveFromBase.Num() != 0)
 	{
 		FGraphEventRef Task = FEntityTaskBuilder()
@@ -808,6 +958,9 @@ void UMovieScenePiecewiseDoubleBlenderSystem::ReinitializeAccumulationBuffers()
 
 	TArrayView<TComponentTypeID<double>> BaseComponents(FBuiltInComponentTypes::Get()->BaseDouble);
 	TArrayView<TComponentTypeID<double>> ResultComponents(FBuiltInComponentTypes::Get()->DoubleResult);
+
+	TComponentTypeID<int32> BlendingOrder(FBuiltInComponentTypes::Get()->BlendingOrder);
+
 	check(BaseComponents.Num() == ResultComponents.Num());
 
 	// Recompute which result types are blended
@@ -828,6 +981,7 @@ void UMovieScenePiecewiseDoubleBlenderSystem::ReinitializeAccumulationBuffers()
 				BuiltInComponents->Tags.AbsoluteBlend, 
 				BuiltInComponents->Tags.RelativeBlend, 
 				BuiltInComponents->Tags.AdditiveBlend, 
+				BuiltInComponents->Tags.OverrideBlend,
 				BuiltInComponents->Tags.AdditiveFromBaseBlend });
 
 		AnyCompositeAllocationsMask.Reset();
@@ -836,6 +990,7 @@ void UMovieScenePiecewiseDoubleBlenderSystem::ReinitializeAccumulationBuffers()
 		const bool bHasAbsolutes         = AnyCompositeAllocationsMask.Contains(BuiltInComponents->Tags.AbsoluteBlend);
 		const bool bHasRelatives         = AnyCompositeAllocationsMask.Contains(BuiltInComponents->Tags.RelativeBlend);
 		const bool bHasAdditives         = AnyCompositeAllocationsMask.Contains(BuiltInComponents->Tags.AdditiveBlend);
+		const bool bHasOverrides		 = AnyCompositeAllocationsMask.Contains(BuiltInComponents->Tags.OverrideBlend);
 		const bool bHasAdditivesFromBase = AnyCompositeAllocationsMask.Contains(BuiltInComponents->Tags.AdditiveFromBaseBlend);
 
 		if (!(bHasAbsolutes || bHasRelatives || bHasAdditives || bHasAdditivesFromBase))
@@ -855,7 +1010,7 @@ void UMovieScenePiecewiseDoubleBlenderSystem::ReinitializeAccumulationBuffers()
 			TArray<FBlendResult>& Buffer = AccumulationBuffers.Relative.Add(Component);
 			Buffer.SetNum(MaximumNumBlends);
 		}
-		if (bHasAdditives)
+		if (bHasAdditives && !bHasOverrides) //if it has overrides we use additives with them
 		{
 			TArray<FBlendResult>& Buffer = AccumulationBuffers.Additive.Add(Component);
 			Buffer.SetNum(MaximumNumBlends);
@@ -866,8 +1021,13 @@ void UMovieScenePiecewiseDoubleBlenderSystem::ReinitializeAccumulationBuffers()
 			Buffer.Buffer.SetNum(MaximumNumBlends);
 			Buffer.BaseComponent = BaseComponents[Index];
 		}
+		if (bHasOverrides)
+		{
+			FAdditiveAndOverrideBuffer& Buffer = AccumulationBuffers.AdditiveAndOverrides.Add(Component);
+			Buffer.Values.SetNum(MaximumNumBlends);
+		}
 	}
-
+	
 	// Update property relevancy
 	CachedRelevantProperties.Empty();
 
@@ -951,12 +1111,14 @@ FGraphEventRef UMovieScenePiecewiseDoubleBlenderSystem::DispatchDecomposeTask(co
 		uint16 DecomposeBlendChannel;
 		FComponentTypeID AdditiveBlendTag;
 		FComponentTypeID AdditiveFromBaseBlendTag;
+		FComponentTypeID OverrideTag;
 
 		explicit FChannelResultTask(const FValueDecompositionParams& Params, FAlignedDecomposedValue* InResult)
 			: Result(InResult)
 			, DecomposeBlendChannel(Params.DecomposeBlendChannel)
 			, AdditiveBlendTag(FBuiltInComponentTypes::Get()->Tags.AdditiveBlend)
 			, AdditiveFromBaseBlendTag(FBuiltInComponentTypes::Get()->Tags.AdditiveFromBaseBlend)
+			, OverrideTag(FBuiltInComponentTypes::Get()->Tags.OverrideBlend)
 		{
 			EntitiesToDecompose.Append(Params.Query.Entities.GetData(), Params.Query.Entities.Num());
 		}
@@ -964,14 +1126,16 @@ FGraphEventRef UMovieScenePiecewiseDoubleBlenderSystem::DispatchDecomposeTask(co
 		void ForEachAllocation(
 				const FEntityAllocation* Allocation, TRead<FMovieSceneEntityID> EntityToDecomposeIDs, 
 				TRead<FMovieSceneBlendChannelID> BlendChannels, TRead<double> ValueResultComponent, 
-				TReadOptional<double> OptionalBaseValueComponent, TReadOptional<double> OptionalWeightComponent)
+				TReadOptional<int32> OptionalBlendingOrder, TReadOptional<double> OptionalBaseValueComponent, TReadOptional<double> OptionalWeightComponent)
 		{
 			static const FMovieSceneBlenderSystemID BlenderSystemID = UMovieSceneBlenderSystem::GetBlenderSystemID<UMovieScenePiecewiseDoubleBlenderSystem>();
 
+			const bool bOverride = Allocation->HasComponent(OverrideTag);
 			const bool bAdditive = Allocation->HasComponent(AdditiveBlendTag);
 			const bool bAdditiveFromBase = Allocation->HasComponent(AdditiveFromBaseBlendTag);
-
+			
 			const int32 Num = Allocation->Num();
+			
 			for (int32 EntityIndex = 0; EntityIndex < Num; ++EntityIndex)
 			{
 				const FMovieSceneBlendChannelID& BlendChannel(BlendChannels[EntityIndex]);
@@ -987,10 +1151,19 @@ FGraphEventRef UMovieScenePiecewiseDoubleBlenderSystem::DispatchDecomposeTask(co
 				const double              ValueResult       = ValueResultComponent[EntityIndex];
 				const double              BaseValue         = OptionalBaseValueComponent ? OptionalBaseValueComponent[EntityIndex] : 0.f;
 				const float               Weight            = OptionalWeightComponent ? OptionalWeightComponent[EntityIndex] : 1.f;
-
+				const int32				  BlendingOrder		= OptionalBlendingOrder ? OptionalBlendingOrder[EntityIndex] : INDEX_NONE;
+				
+				if (BlendingOrder != INDEX_NONE && (bAdditive || bOverride))
+				{
+					Result->Value.AllDecomposedOverrides.Add(FWeightedValue{ ValueResult, Weight, 0.0, BlendingOrder, bAdditive });
+				}
 				if (EntitiesToDecompose.Contains(EntityToDecompose))
 				{
-					if (bAdditive)
+					if (BlendingOrder != INDEX_NONE && (bAdditive || bOverride))
+					{
+						Result->Value.DecomposedOverrides.Add(MakeTuple(EntityToDecompose, FWeightedValue{ ValueResult, Weight, 0.0, BlendingOrder, bAdditive }));
+					}
+					else if (bAdditive)
 					{
 						Result->Value.DecomposedAdditives.Add(MakeTuple(EntityToDecompose, FWeightedValue{ ValueResult, Weight }));
 					}
@@ -1026,6 +1199,7 @@ FGraphEventRef UMovieScenePiecewiseDoubleBlenderSystem::DispatchDecomposeTask(co
 			.Read(BuiltInComponents->ParentEntity)
 			.Read(BuiltInComponents->BlendChannelInput)
 			.Read(ResultComponentType)
+			.ReadOptional(BuiltInComponents->BlendingOrder)
 			.ReadOptional(BaseValueComponentType)
 			.ReadOptional(BuiltInComponents->WeightAndEasingResult)
 			.FilterAll({ Params.PropertyTag, GetBlenderTypeTag() })
@@ -1037,6 +1211,7 @@ FGraphEventRef UMovieScenePiecewiseDoubleBlenderSystem::DispatchDecomposeTask(co
 			.ReadEntityIDs()
 			.Read(BuiltInComponents->BlendChannelInput)
 			.Read(ResultComponentType)
+			.ReadOptional(BuiltInComponents->BlendingOrder)
 			.ReadOptional(BaseValueComponentType)
 			.ReadOptional(BuiltInComponents->WeightAndEasingResult)
 			.FilterAll({ Params.PropertyTag, GetBlenderTypeTag() })
diff --git a/Engine/Source/Runtime/MovieSceneTracks/Private/Systems/MovieScenePropertyInstantiator.cpp b/Engine/Source/Runtime/MovieSceneTracks/Private/Systems/MovieScenePropertyInstantiator.cpp
--- a/Engine/Source/Runtime/MovieSceneTracks/Private/Systems/MovieScenePropertyInstantiator.cpp
+++ b/Engine/Source/Runtime/MovieSceneTracks/Private/Systems/MovieScenePropertyInstantiator.cpp
@@ -690,6 +690,7 @@ void UMovieScenePropertyInstantiatorSystem::UpdatePropertyInfo(const FPropertyPa
 			{
 				if (Type.Contains(BuiltInComponents->Tags.RelativeBlend) || 
 						Type.Contains(BuiltInComponents->Tags.AdditiveBlend) || 
+						Type.Contains(BuiltInComponents->Tags.OverrideBlend) ||
 						Type.Contains(BuiltInComponents->Tags.AdditiveFromBaseBlend) || 
 						Type.Contains(BuiltInComponents->WeightAndEasingResult))
 				{
diff --git a/Engine/Source/Runtime/MovieSceneTracks/Public/Systems/MovieScenePiecewiseDoubleBlenderSystem.h b/Engine/Source/Runtime/MovieSceneTracks/Public/Systems/MovieScenePiecewiseDoubleBlenderSystem.h
--- a/Engine/Source/Runtime/MovieSceneTracks/Public/Systems/MovieScenePiecewiseDoubleBlenderSystem.h
+++ b/Engine/Source/Runtime/MovieSceneTracks/Public/Systems/MovieScenePiecewiseDoubleBlenderSystem.h
@@ -24,6 +24,28 @@ struct FBlendResult
 	float Weight = 0.f;
 };
 
+/* Data for each additive/override section, we don't accumulate until we recieve all of the data*/
+struct FAdditveAndOverrideData
+{
+	/** If Additive, if not then override*/
+	bool bIsAdditive;
+	/** Section Value*/
+	double Value;
+	/** Section Weight*/
+	float Weight;
+	/** Section Blending Order*/
+	int32 BlendingOrder;
+	bool operator<(const FAdditveAndOverrideData& RHS) const
+	{
+		return BlendingOrder < RHS.BlendingOrder;
+	}
+};
+/** Structure for holding the sorted blending order data for the additive and override data*/
+struct FAdditiveAndOverrides
+{
+	mutable TSortedMap<int32, FAdditveAndOverrideData> Data;
+};
+
 /** Structure for holding the blend results of each blend type */
 struct FAccumulationResult
 {
@@ -31,10 +53,11 @@ struct FAccumulationResult
 	const FBlendResult* Relatives = nullptr;
 	const FBlendResult* Additives = nullptr;
 	const FBlendResult* AdditivesFromBase = nullptr;
-
+	const FAdditiveAndOverrides* Overrides = nullptr;
 	bool IsValid() const
 	{
-		return Absolutes || Relatives || Additives || AdditivesFromBase;
+		return Absolutes || Relatives || Additives || AdditivesFromBase
+			|| Overrides;
 	}
 
 	FBlendResult GetAbsoluteResult(uint16 BlendID) const
@@ -53,6 +76,16 @@ struct FAccumulationResult
 	{
 		return AdditivesFromBase ? AdditivesFromBase[BlendID] : FBlendResult{};
 	}
+	FAdditiveAndOverrides GetAdditiveAndOverrideResult(uint16 BlendID)  const
+	{
+		return Overrides? Overrides[BlendID] : FAdditiveAndOverrides{};
+	}
+};
+
+/** Bufffer used for additive/override values*/
+struct FAdditiveAndOverrideBuffer
+{
+	TArray<FAdditiveAndOverrides> Values;
 };
 
 /** Buffer used for accumulating additive-from-base values */
@@ -79,6 +112,9 @@ struct FAccumulationBuffers
 	TSortedMap<FComponentTypeID, TArray<FBlendResult>> Additive;
 	/** Map from value result component type -> Additive From Base blend accumulation buffer for that channel type */
 	TSortedMap<FComponentTypeID, FAdditiveFromBaseBuffer> AdditiveFromBase;
+	/** Map from value result component type -> Additive and Override blending order blneds for that channel type */
+	TSortedMap<FComponentTypeID, FAdditiveAndOverrideBuffer> AdditiveAndOverrides;
+
 };
 
 } // namespace MovieScene
