Sequencer: Check for valid storage index in camera cut handler which can be null if the camera is no longer loaded

diff --git a/Engine/Source/Runtime/MovieSceneTracks/Private/TrackInstances/MovieSceneCameraCutGameHandler.cpp b/Engine/Source/Runtime/MovieSceneTracks/Private/TrackInstances/MovieSceneCameraCutGameHandler.cpp
--- a/Engine/Source/Runtime/MovieSceneTracks/Private/TrackInstances/MovieSceneCameraCutGameHandler.cpp
+++ b/Engine/Source/Runtime/MovieSceneTracks/Private/TrackInstances/MovieSceneCameraCutGameHandler.cpp
@@ -198,11 +198,16 @@ void FCameraCutGameHandler::ForcePreAnimatedValueRestore(
 	FRestoreStateParams Params;
 	Params.Linker = Linker;
 	Params.TerminalInstanceHandle = SequenceInstance.GetRootInstanceHandle();
-	PreAnimatedStorage->RestorePreAnimatedStateStorage(
-			(uint8)0,  // See comment below
-			EPreAnimatedStorageRequirement::Transient, 
-			EPreAnimatedStorageRequirement::Persistent,
-			Params);
+
+	FPreAnimatedStorageIndex StorageIndex = PreAnimatedStorage->FindStorageIndex(0);
+	if (StorageIndex.IsValid())
+	{		
+		PreAnimatedStorage->RestorePreAnimatedStateStorage(
+				(uint8)0,  // See comment below
+				EPreAnimatedStorageRequirement::Transient, 
+				EPreAnimatedStorageRequirement::Persistent,
+				Params);
+	}
 }
 
 void FCameraCutGameHandler::CachePreAnimatedValue(
