[Amber] Fix intermintent silient audio in sequence due to stale parent pointer on SceneComponent during playback evaluation.

#rnx

diff --git a/Engine/Source/Runtime/MovieSceneTracks/Private/Systems/MovieSceneAudioSystem.cpp b/Engine/Source/Runtime/MovieSceneTracks/Private/Systems/MovieSceneAudioSystem.cpp
--- a/Engine/Source/Runtime/MovieSceneTracks/Private/Systems/MovieSceneAudioSystem.cpp
+++ b/Engine/Source/Runtime/MovieSceneTracks/Private/Systems/MovieSceneAudioSystem.cpp
@@ -397,10 +397,16 @@ private:
 			{
 				UAudioComponent* AudioComponent = EvaluationData->AudioComponent.Get();
 
-				if (AudioComponent && AttachComponent && 
-					(AudioComponent->GetAttachParent() != AttachComponent || AudioComponent->GetAttachSocketName() != AttachKey.SocketName))
+				if (AudioComponent)
 				{
-					AudioComponent->AttachToComponent(AttachComponent, FAttachmentTransformRules::KeepRelativeTransform, AttachKey.SocketName);
+					if (AttachComponent && (AudioComponent->GetAttachParent() != AttachComponent || AudioComponent->GetAttachSocketName() != AttachKey.SocketName))
+					{
+						AudioComponent->AttachToComponent(AttachComponent, FAttachmentTransformRules::KeepRelativeTransform, AttachKey.SocketName);
+					}
+					else if (!AttachComponent && (AudioComponent->GetAttachParent() != AttachComponent || AudioComponent->GetAttachSocketName() != AttachKey.SocketName))
+					{
+						AudioComponent->DetachFromComponent(FDetachmentTransformRules::KeepRelativeTransform);
+					}
 				}
 
 				EvaluationData->VolumeMultiplier = VolumeMultiplier * AudioSection->EvaluateEasing(Context.GetTime());
