Sequencer: Fixed missing dependency between byte propertyies and byte blenders

diff --git a/Engine/Source/Runtime/MovieSceneTracks/Private/Systems/MovieSceneBytePropertySystem.cpp b/Engine/Source/Runtime/MovieSceneTracks/Private/Systems/MovieSceneBytePropertySystem.cpp
--- a/Engine/Source/Runtime/MovieSceneTracks/Private/Systems/MovieSceneBytePropertySystem.cpp
+++ b/Engine/Source/Runtime/MovieSceneTracks/Private/Systems/MovieSceneBytePropertySystem.cpp
@@ -2,6 +2,7 @@
 
 #include "Systems/MovieSceneBytePropertySystem.h"
 #include "Systems/ByteChannelEvaluatorSystem.h"
+#include "Systems/MovieScenePiecewiseByteBlenderSystem.h"
 #include "MovieSceneTracksComponentTypes.h"
 
 #include UE_INLINE_GENERATED_CPP_BY_NAME(MovieSceneBytePropertySystem)
@@ -14,6 +15,7 @@ UMovieSceneBytePropertySystem::UMovieSceneBytePropertySystem(const FObjectInitia
 	if (HasAnyFlags(RF_ClassDefaultObject))
 	{
 		DefineImplicitPrerequisite(UByteChannelEvaluatorSystem::StaticClass(), GetClass());
+		DefineImplicitPrerequisite(UMovieScenePiecewiseByteBlenderSystem::StaticClass(), GetClass());
 
 		DefineComponentConsumer(GetClass(), UE::MovieScene::FMovieSceneTracksComponentTypes::Get()->Byte.PropertyTag);
 	}
