Sequencer: Override Blending:
Add decompose unit tests for override eval/decomposition

diff --git a/Engine/Source/Runtime/MovieSceneTracks/Private/Tests/MovieSceneDecomposerTests.cpp b/Engine/Source/Runtime/MovieSceneTracks/Private/Tests/MovieSceneDecomposerTests.cpp
--- a/Engine/Source/Runtime/MovieSceneTracks/Private/Tests/MovieSceneDecomposerTests.cpp
+++ b/Engine/Source/Runtime/MovieSceneTracks/Private/Tests/MovieSceneDecomposerTests.cpp
@@ -336,6 +336,72 @@ bool FMovieSceneDecomposerAdditiveWithEasingTest::RunTest(const FString& Paramet
 	return true;
 }
 
+IMPLEMENT_SIMPLE_AUTOMATION_TEST(
+	FMovieSceneDecomposerOverrideWithEasingTest,
+	"System.Engine.Sequencer.Decomposer.OverrideWithEasing",
+	EAutomationTestFlags::EditorContext | EAutomationTestFlags::EngineFilter)
+bool FMovieSceneDecomposerOverrideWithEasingTest::RunTest(const FString& Parameters)
+{
+	using namespace UE::MovieScene;
+	using namespace UE::MovieScene::Test;
+
+	UMovieSceneSection* TestObjectFloatSection1;
+	UMovieSceneSection* TestObjectFloatSection2;
+	UMovieSceneSection* TestObjectFloatSection3;
+	UMovieSceneSection* TestObjectFloatSection4;
+
+	UMovieSceneDecomposerTestObject* TestObject = NewObject<UMovieSceneDecomposerTestObject>();
+
+	// Absolute section setting the value to 100, with an additive on top worth +100.
+	// There is ease-in on the additive.  There is then an Override with an ease-in with a value of 0.
+	// On top of that is another additive with a value of 25.
+	// With weight of 0.5 the total value is (100 + (0.5) *100) *0.5  +(0.0) *0.5  + 25 = 100
+	// With weight of 1.0 the total value is (100 + 100) *0.0 + 0.0 + 25 = 25
+	FSequenceBuilder()
+		.AddObjectBinding(TestObject)
+		.AddPropertyTrack<UMovieSceneFloatTrack>(GET_MEMBER_NAME_CHECKED(UMovieSceneDecomposerTestObject, FloatProperty))
+		.AddSection(0, 5000)
+		.Assign(TestObjectFloatSection1)
+		.AddKey<FMovieSceneFloatChannel, float>(0, 0, 100.f)
+		.Pop()
+		.AddSection(0, 5000, 1)
+		.Assign(TestObjectFloatSection2)
+		.SetBlendType(EMovieSceneBlendType::Additive)
+		.SetEaseIn(2000)
+		.AddKey<FMovieSceneFloatChannel, float>(0, 0, 100.f)
+		.Pop()
+		.AddSection(0, 5000, 1)
+		.Assign(TestObjectFloatSection3)
+		.SetBlendType(EMovieSceneBlendType::Override)
+		.SetEaseIn(2000)
+		.AddKey<FMovieSceneFloatChannel, float>(0, 0, 0.0f)
+		.Pop()
+		.AddSection(0, 5000, 1)
+		.Assign(TestObjectFloatSection4)
+		.SetBlendType(EMovieSceneBlendType::Additive)
+		.AddKey<FMovieSceneFloatChannel, float>(0, 0, 25.0f)
+		.Pop()
+		.Pop();
+
+	// Figure out what it would take for the final additive to give us a final value of 125. It should make the additive 50,
+	// since it only has a weight of 50% at the evaluation time.
+	FMovieSceneTracksComponentTypes* ComponentTypes = FMovieSceneTracksComponentTypes::Get();
+	TRecompositionResult<double> Result = RecomposeBlendOperational(TestObject, TestObjectFloatSection4, ComponentTypes->Float, 125.f, 1000);
+
+	UTEST_EQUAL("Recomposed values", Result.Values.Num(), 1);
+	UTEST_EQUAL("Value to key", Result.Values[0], 50.0);
+
+	// Run at time 2000, where the weight is 100%. With a total of 100.0,  the Override should be 75, which is 100 - 25(additive) value.
+	// The absolute and previous additive values don't matter, they have been overriden.
+	Result = RecomposeBlendOperational(TestObject, TestObjectFloatSection3, ComponentTypes->Float, 100.0f, 2000);
+
+	UTEST_EQUAL("Recomposed values", Result.Values.Num(), 1);
+	UTEST_EQUAL("Value to key", Result.Values[0], 75.0);
+
+	return true;
+}
+
+
 IMPLEMENT_SIMPLE_AUTOMATION_TEST(
 		FMovieSceneDecomposerAdditiveWithMultipleEasingTest,
 		"System.Engine.Sequencer.Decomposer.AdditiveWithMultipleEasing", 
