Remove calls to CacheChannelProxy

diff --git a/Engine/Source/Runtime/MovieSceneTracks/Private/Sections/MovieSceneAudioSection.cpp b/Engine/Source/Runtime/MovieSceneTracks/Private/Sections/MovieSceneAudioSection.cpp
--- a/Engine/Source/Runtime/MovieSceneTracks/Private/Sections/MovieSceneAudioSection.cpp
+++ b/Engine/Source/Runtime/MovieSceneTracks/Private/Sections/MovieSceneAudioSection.cpp
@@ -105,13 +105,6 @@ namespace MovieSceneAudioSectionPrivate
 	}
 }
 
-void UMovieSceneAudioSection::PostEditImport()
-{
-	Super::PostEditImport();
-
-	CacheChannelProxy();
-}
-
 EMovieSceneChannelProxyType  UMovieSceneAudioSection::CacheChannelProxy()
 {
 	// Set up the channel proxy
@@ -265,8 +258,6 @@ void UMovieSceneAudioSection::PostLoad()
 {
 	Super::PostLoad();
 
-	CacheChannelProxy();
-
 	if (AudioDilationFactor_DEPRECATED != AudioDeprecatedMagicNumber)
 	{
 		PitchMultiplier.SetDefault(AudioDilationFactor_DEPRECATED);
@@ -374,7 +365,7 @@ UMovieSceneSection* UMovieSceneAudioSection::SplitSection(FQualifiedFrameTime Sp
 void UMovieSceneAudioSection::SetSound(USoundBase* InSound)
 {
 	Sound = InSound;
-	CacheChannelProxy();
+	InvalidateChannelProxy();
 }
 
 USceneComponent* UMovieSceneAudioSection::GetAttachComponent(const AActor* InParentActor, const FMovieSceneActorReferenceKey& Key) const
@@ -598,14 +589,3 @@ void UMovieSceneAudioSection::ImportEntityImpl(UMovieSceneEntitySystemLinker* En
 		);
 	}
 }
-
-#if WITH_EDITOR
-
-void UMovieSceneAudioSection::PostEditChangeProperty(FPropertyChangedEvent& PropertyChangedEvent)
-{
-	Super::PostEditChangeProperty(PropertyChangedEvent);
-
-	CacheChannelProxy();
-}
-
-#endif
diff --git a/Engine/Source/Runtime/MovieSceneTracks/Public/Sections/MovieSceneAudioSection.h b/Engine/Source/Runtime/MovieSceneTracks/Public/Sections/MovieSceneAudioSection.h
--- a/Engine/Source/Runtime/MovieSceneTracks/Public/Sections/MovieSceneAudioSection.h
+++ b/Engine/Source/Runtime/MovieSceneTracks/Public/Sections/MovieSceneAudioSection.h
@@ -189,17 +189,12 @@ public:
 	MOVIESCENETRACKS_API virtual UMovieSceneSection* SplitSection(FQualifiedFrameTime SplitTime, bool bDeleteKeys) override;
 	MOVIESCENETRACKS_API virtual TOptional<FFrameTime> GetOffsetTime() const override;
 	MOVIESCENETRACKS_API virtual void MigrateFrameTimes(FFrameRate SourceRate, FFrameRate DestinationRate) override;
-	MOVIESCENETRACKS_API virtual void PostEditImport() override;
 	MOVIESCENETRACKS_API virtual EMovieSceneChannelProxyType CacheChannelProxy() override;
 
 	//~ IMovieSceneEntityProvider interface
 	MOVIESCENETRACKS_API virtual bool PopulateEvaluationFieldImpl(const TRange<FFrameNumber>& EffectiveRange, const FMovieSceneEvaluationFieldEntityMetaData& InMetaData, FMovieSceneEntityComponentFieldBuilder* OutFieldBuilder) override;
 	MOVIESCENETRACKS_API virtual void ImportEntityImpl(UMovieSceneEntitySystemLinker* EntityLinker, const FEntityImportParams& Params, FImportedEntity* OutImportedEntity) override;
 
-#if WITH_EDITOR
-	MOVIESCENETRACKS_API virtual void PostEditChangeProperty(FPropertyChangedEvent& PropertyChangedEvent) override;
-#endif
-
 private:
 
 	template<typename ChannelType, typename ForEachFunction>
