Sequencer: Remove the inline color picker for the fade track as it's causing problems with transactions

diff --git a/Engine/Source/Runtime/MovieSceneTracks/Public/Sections/MovieSceneFadeSection.h b/Engine/Source/Runtime/MovieSceneTracks/Public/Sections/MovieSceneFadeSection.h
--- a/Engine/Source/Runtime/MovieSceneTracks/Public/Sections/MovieSceneFadeSection.h
+++ b/Engine/Source/Runtime/MovieSceneTracks/Public/Sections/MovieSceneFadeSection.h
@@ -39,7 +39,7 @@ public:
 	FMovieSceneFloatChannel FloatCurve;
 
 	/** Fade color. */
-	UPROPERTY(EditAnywhere, Category="Fade", meta=(InlineColorPicker))
+	UPROPERTY(EditAnywhere, Category="Fade")
 	FLinearColor FadeColor;
 
 	/** Fade audio. */
