Sequencer: support reverse playback of camera cut track with shot blending

The main part of the change fixes the logic in FMovieSceneCameraCutTrackInstance:

- Reverse in/out blends when playing backwards. For instance, a blend out becomes a blend in.
- Adjust time comparisons when playing backwards. For instance, figuring out if we have entered a blend out requires checking if the current time is *lesser* than the blend out start time, instead of greater than.

This however surfaced another problem, which is that the evaluation context for the first update doesn't have the correct "Backwards" orientation. This is because when we start playing, we do a zero-width update, i.e. a forced update at the current/start time. When the evaluation range has no width, the evaluation range structure defaults to "forwards" direction. This change also adds a "preferred direction" to that constructor (and to the methods leading to it) so that a zero-width range gets the "backwards" direction when needed.

diff --git a/Engine/Source/Runtime/MovieScene/Private/Evaluation/MovieScenePlayback.cpp b/Engine/Source/Runtime/MovieScene/Private/Evaluation/MovieScenePlayback.cpp
--- a/Engine/Source/Runtime/MovieScene/Private/Evaluation/MovieScenePlayback.cpp
+++ b/Engine/Source/Runtime/MovieScene/Private/Evaluation/MovieScenePlayback.cpp
@@ -47,12 +47,16 @@ FMovieSceneEvaluationRange::FMovieSceneEvaluationRange(TRange<FFrameTime> InRang
 {
 }
 
-FMovieSceneEvaluationRange::FMovieSceneEvaluationRange(FFrameTime InCurrentTime, FFrameTime InPreviousTime, FFrameRate InFrameRate, bool bInclusivePreviousTime)
+FMovieSceneEvaluationRange::FMovieSceneEvaluationRange(FFrameTime InCurrentTime, FFrameTime InPreviousTime, FFrameRate InFrameRate, bool bInclusivePreviousTime, EPlayDirection PreferredDirection)
 	: EvaluationRange(CalculateEvaluationRange(InCurrentTime, InPreviousTime, bInclusivePreviousTime))
 	, CurrentFrameRate(InFrameRate)
-	, Direction((InCurrentTime - InPreviousTime >= FFrameTime()) ? EPlayDirection::Forwards : EPlayDirection::Backwards)
 	, TimeOverride(TNumericLimits<int32>::Lowest())
 {
+	const FFrameTime ZeroTime;
+	const FFrameTime RangeLength = (InCurrentTime - InPreviousTime);
+	Direction = (RangeLength > ZeroTime) ? EPlayDirection::Forwards :
+		((RangeLength < ZeroTime) ? EPlayDirection::Backwards :
+		 PreferredDirection);
 }
 
 void FMovieSceneEvaluationRange::ResetRange(const TRange<FFrameTime>& NewRange)
@@ -219,7 +223,7 @@ FMovieSceneEvaluationRange FMovieScenePlaybackPosition::GetCurrentPositionAsRang
 	return FMovieSceneEvaluationRange(OutputPosition, OutputRate);
 }
 
-FMovieSceneEvaluationRange FMovieScenePlaybackPosition::JumpTo(FFrameTime InputPosition)
+FMovieSceneEvaluationRange FMovieScenePlaybackPosition::JumpTo(FFrameTime InputPosition, EPlayDirection PreferredDirection)
 {
 	CheckInvariants();
 
@@ -237,11 +241,11 @@ FMovieSceneEvaluationRange FMovieScenePlaybackPosition::JumpTo(FFrameTime InputP
 	// Convert to output time-base
 	FFrameTime OutputPosition = ConvertFrameTime(InputPosition, InputRate, OutputRate);
 
-	LastRange = FMovieSceneEvaluationRange(TRange<FFrameTime>(OutputPosition), OutputRate, EPlayDirection::Forwards);
+	LastRange = FMovieSceneEvaluationRange(TRange<FFrameTime>(OutputPosition), OutputRate, PreferredDirection);
 	return LastRange.GetValue();
 }
 
-FMovieSceneEvaluationRange FMovieScenePlaybackPosition::PlayTo(FFrameTime InputPosition)
+FMovieSceneEvaluationRange FMovieScenePlaybackPosition::PlayTo(FFrameTime InputPosition, EPlayDirection PreferredDirection)
 {
 	CheckInvariants();
 
@@ -256,7 +260,7 @@ FMovieSceneEvaluationRange FMovieScenePlaybackPosition::PlayTo(FFrameTime InputP
 	FFrameTime OutputEvalPositionFrom = ConvertFrameTime(InputEvalPositionFrom, InputRate, OutputRate);
 	FFrameTime OutputEvalPositionTo   = ConvertFrameTime(InputPosition, InputRate, OutputRate);
 
-	LastRange = FMovieSceneEvaluationRange(OutputEvalPositionTo, OutputEvalPositionFrom, OutputRate, !PreviousPlayEvalPosition.IsSet());
+	LastRange = FMovieSceneEvaluationRange(OutputEvalPositionTo, OutputEvalPositionFrom, OutputRate, !PreviousPlayEvalPosition.IsSet(), PreferredDirection);
 
 	// Assign the cached input values
 	CurrentPosition          = InputPosition;
diff --git a/Engine/Source/Runtime/MovieScene/Private/MovieSceneSequencePlayer.cpp b/Engine/Source/Runtime/MovieScene/Private/MovieSceneSequencePlayer.cpp
--- a/Engine/Source/Runtime/MovieScene/Private/MovieSceneSequencePlayer.cpp
+++ b/Engine/Source/Runtime/MovieScene/Private/MovieSceneSequencePlayer.cpp
@@ -373,7 +373,8 @@ void UMovieSceneSequencePlayer::PlayInternal()
 
 		if (!PlayPosition.GetLastPlayEvalPostition().IsSet() || PlayPosition.GetLastPlayEvalPostition() != PlayPosition.GetCurrentPosition())
 		{
-			UpdateMovieSceneInstance(PlayPosition.PlayTo(PlayPosition.GetCurrentPosition()), EMovieScenePlayerStatus::Playing);
+			EPlayDirection PreferredDirection = bReversePlayback ? EPlayDirection::Backwards : EPlayDirection::Forwards;
+			UpdateMovieSceneInstance(PlayPosition.PlayTo(PlayPosition.GetCurrentPosition(), PreferredDirection), EMovieScenePlayerStatus::Playing);
 		}
 
 		RunLatentActions();
@@ -1164,14 +1165,15 @@ EMovieScenePlayerStatus::Type UpdateMethodToStatus(EUpdatePositionMethod Method)
 	}
 }
 
-FMovieSceneEvaluationRange UpdatePlayPosition(FMovieScenePlaybackPosition& InOutPlayPosition, FFrameTime NewTime, EUpdatePositionMethod Method)
+FMovieSceneEvaluationRange UpdatePlayPosition(FMovieScenePlaybackPosition& InOutPlayPosition, FFrameTime NewTime, EUpdatePositionMethod Method, bool bPreferReversePlayback)
 {
+	EPlayDirection PreferredDirection = bPreferReversePlayback ? EPlayDirection::Backwards : EPlayDirection::Forwards;
 	if (Method == EUpdatePositionMethod::Play)
 	{
-		return InOutPlayPosition.PlayTo(NewTime);
+		return InOutPlayPosition.PlayTo(NewTime, PreferredDirection);
 	}
 
-	return InOutPlayPosition.JumpTo(NewTime);
+	return InOutPlayPosition.JumpTo(NewTime, PreferredDirection);
 }
 
 void UMovieSceneSequencePlayer::UpdateTimeCursorPosition_Internal(FFrameTime NewPosition, EUpdatePositionMethod Method, bool bHasJumpedOverride)
@@ -1294,7 +1296,7 @@ void UMovieSceneSequencePlayer::UpdateTimeCursorPosition_Internal(FFrameTime New
 			// Clamp the position to the duration
 			NewPosition = FMath::Clamp(NewPosition, FFrameTime(StartTime), GetLastValidTime());
 
-			FMovieSceneEvaluationRange Range = UpdatePlayPosition(PlayPosition, NewPosition, Method);
+			FMovieSceneEvaluationRange Range = UpdatePlayPosition(PlayPosition, NewPosition, Method, bReversePlayback);
 			UpdateMovieSceneInstance(Range, StatusOverride);
 
 			// We have authority to finish playback if:
@@ -1334,7 +1336,7 @@ void UMovieSceneSequencePlayer::UpdateTimeCursorPosition_Internal(FFrameTime New
 		// If the desired evaluation will take us past where we want to go we need to use a clipped range provided by the PauseRange, otherwise use the normal one.
 		FMovieSceneEvaluationRange Range = PauseRange.IsSet() 
 			? FMovieSceneEvaluationRange(PauseRange.GetValue(), PlayPosition.GetOutputRate(), bReversePlayback ? EPlayDirection::Backwards : EPlayDirection::Forwards)
-			: UpdatePlayPosition(PlayPosition, NewPosition, Method);
+			: UpdatePlayPosition(PlayPosition, NewPosition, Method, bReversePlayback);
 
 		UMovieSceneSequence* MovieSceneSequence = RootTemplateInstance.GetSequence(MovieSceneSequenceID::Root);
 		const bool bIsSequenceBlocking = MovieSceneSequence ? EnumHasAnyFlags(MovieSceneSequence->GetFlags(), EMovieSceneSequenceFlags::BlockingEvaluation) : false;
@@ -1991,7 +1993,7 @@ void UMovieSceneSequencePlayer::PostNetReceive()
 			else if (Status == EMovieScenePlayerStatus::Stopped)
 			{
 				// Both client and server are stopped so just update our (client) position to match the server's.
-				UpdatePlayPosition(PlayPosition, NetSyncProps.LastKnownPosition, EUpdatePositionMethod::Jump);
+				UpdatePlayPosition(PlayPosition, NetSyncProps.LastKnownPosition, EUpdatePositionMethod::Jump, bReversePlayback);
 				TimeController->Reset(GetCurrentTime());
 			}
 
diff --git a/Engine/Source/Runtime/MovieScene/Public/Evaluation/MovieScenePlayback.h b/Engine/Source/Runtime/MovieScene/Public/Evaluation/MovieScenePlayback.h
--- a/Engine/Source/Runtime/MovieScene/Public/Evaluation/MovieScenePlayback.h
+++ b/Engine/Source/Runtime/MovieScene/Public/Evaluation/MovieScenePlayback.h
@@ -44,7 +44,7 @@ struct FMovieSceneEvaluationRange
 	/**
 	 * Construct this range from 2 times, and whether the range should include the previous time or not
 	 */
-	MOVIESCENE_API FMovieSceneEvaluationRange(FFrameTime InCurrentTime, FFrameTime InPreviousTime, FFrameRate InFrameRate, bool bInclusivePreviousTime = false);
+	MOVIESCENE_API FMovieSceneEvaluationRange(FFrameTime InCurrentTime, FFrameTime InPreviousTime, FFrameRate InFrameRate, bool bInclusivePreviousTime = false, EPlayDirection PreferredDirection = EPlayDirection::Forwards);
 
 	/**
 	 * Convert a frame time range to a frame number range comprising all the frame numbers traversed in the range
@@ -547,7 +547,7 @@ public:
 	 * @param NewPosition         The new frame time to set, in InputRate space
 	 * @return A range encompassing only the specified time, in OutputRate space.
 	 */
-	MOVIESCENE_API FMovieSceneEvaluationRange JumpTo(FFrameTime NewPosition);
+	MOVIESCENE_API FMovieSceneEvaluationRange JumpTo(FFrameTime NewPosition, EPlayDirection PreferredDirection = EPlayDirection::Forwards);
 
 	/**
 	 * Play from the previously evaluated play time, to the specified time
@@ -555,7 +555,7 @@ public:
 	 * @param NewPosition         The new frame time to set, in InputRate space
 	 * @return An evaluation range from the previously evaluated time to the specified time, in OutputRate space.
 	 */
-	MOVIESCENE_API FMovieSceneEvaluationRange PlayTo(FFrameTime NewPosition);
+	MOVIESCENE_API FMovieSceneEvaluationRange PlayTo(FFrameTime NewPosition, EPlayDirection PreferredDirection = EPlayDirection::Forwards);
 
 	/**
 	 * Get a range that encompasses the last evaluated range in OutputRate space.
diff --git a/Engine/Source/Runtime/MovieSceneTracks/Private/TrackInstances/MovieSceneCameraCutTrackInstance.cpp b/Engine/Source/Runtime/MovieSceneTracks/Private/TrackInstances/MovieSceneCameraCutTrackInstance.cpp
--- a/Engine/Source/Runtime/MovieSceneTracks/Private/TrackInstances/MovieSceneCameraCutTrackInstance.cpp
+++ b/Engine/Source/Runtime/MovieSceneTracks/Private/TrackInstances/MovieSceneCameraCutTrackInstance.cpp
@@ -117,6 +117,7 @@ struct FBlendedCameraCut
 	FFrameTime LocalContextTime;
 	FFrameNumber LocalEaseOutStartTime;
 	FFrameNumber LocalEndTime;
+	EPlayDirection LocalDirection = EPlayDirection::Forwards;
 
 	FBlendedCameraCutEasingInfo EaseIn;
 	FBlendedCameraCutEasingInfo EaseOut;
@@ -404,13 +405,15 @@ void UMovieSceneCameraCutTrackInstance::OnAnimate()
 		FTransform CutTransform = Section->InitialCameraCutTransform;
 		const bool bHasCutTransform = Section->bHasInitialCameraCutTransform;
 
+		const bool bIsBackwards = Context.GetDirection() == EPlayDirection::Backwards;
+
 		const int32 PreRollFrames = Section->GetPreRollFrames();
 		bool bIsSectionPreRoll = false;
 		if (PreRollFrames > 0 && Section->HasStartFrame())
 		{
 			const FFrameNumber SectionStartTime = Section->GetTrueRange().GetLowerBoundValue();
 			const TRange<FFrameNumber> PreRollRange(SectionStartTime - PreRollFrames, SectionStartTime);
-			bIsSectionPreRoll = PreRollRange.Contains(Context.GetTime().FloorToFrame());
+			bIsSectionPreRoll = !bIsBackwards && PreRollRange.Contains(Context.GetTime().FloorToFrame());
 		}
 
 		if (Context.IsPreRoll() || bIsSectionPreRoll)
@@ -424,6 +427,7 @@ void UMovieSceneCameraCutTrackInstance::OnAnimate()
 			const FMovieSceneInverseSequenceTransform SequenceToRootTransform = Context.GetSequenceToRootSequenceTransform();
 
 			FBlendedCameraCut Params(Input, CameraBindingID, SequenceInstance.GetSequenceID());
+			Params.LocalDirection = Context.GetDirection();
 			Params.bCanBlend = Track->bCanBlend;
 
 			// Get start/current/end time.
@@ -452,7 +456,6 @@ void UMovieSceneCameraCutTrackInstance::OnAnimate()
 			{
 				Params.LocalEaseOutStartTime = Params.LocalEndTime - Section->Easing.GetEaseOutDuration();
 
-
 				TOptional<FFrameTime> RootStart = SequenceToRootTransform.TryTransformTime(Params.LocalEaseOutStartTime);
 				TOptional<FFrameTime> RootEnd   = SequenceToRootTransform.TryTransformTime(Params.LocalEndTime);
 
@@ -463,6 +466,14 @@ void UMovieSceneCameraCutTrackInstance::OnAnimate()
 				}
 			}
 
+			// If we are evaluating backwards, swap all our start/end times.
+			if (bIsBackwards)
+			{
+				Swap(Params.LocalStartTime, Params.LocalEndTime);
+				Swap(Params.EaseIn, Params.EaseOut);
+				Swap(Params.LocalEaseInEndTime, Params.LocalEaseOutStartTime);
+			}
+
 			// Remember locking option.
 			Params.bLockPreviousCamera = Section->bLockPreviousCamera;
 
@@ -502,7 +513,11 @@ void UMovieSceneCameraCutTrackInstance::OnAnimate()
 		// We want to blend away from this priority cut only if we are currently in its blend-out. We'll
 		// determine that thanks to the various times we saved on the info struct.
 		const FBlendedCameraCut& PriorityCameraCut = CameraCutParams[0];
-		if (PriorityCameraCut.LocalContextTime < PriorityCameraCut.LocalEaseInEndTime)
+		const bool bIsBackwards = PriorityCameraCut.LocalDirection == EPlayDirection::Backwards;
+
+		if (
+				(!bIsBackwards && PriorityCameraCut.LocalContextTime < PriorityCameraCut.LocalEaseInEndTime) ||
+				(bIsBackwards && PriorityCameraCut.LocalContextTime > PriorityCameraCut.LocalEaseInEndTime))
 		{
 			// Blending in from the other cut.
 			const FBlendedCameraCut& PrevCameraCut = CameraCutParams[1];
@@ -512,7 +527,9 @@ void UMovieSceneCameraCutTrackInstance::OnAnimate()
 			FinalCameraCut.PreviousCameraBindingID = PrevCameraCut.CameraBindingID;
 			FinalCameraCut.PreviousOperandSequenceID = PrevCameraCut.OperandSequenceID;
 		}
-		else if (PriorityCameraCut.LocalContextTime > PriorityCameraCut.LocalEaseOutStartTime)
+		else if (
+				(!bIsBackwards && PriorityCameraCut.LocalContextTime > PriorityCameraCut.LocalEaseOutStartTime) ||
+				(bIsBackwards && PriorityCameraCut.LocalContextTime < PriorityCameraCut.LocalEaseOutStartTime))
 		{
 			// Blending out to the other cut.
 			const FBlendedCameraCut& PrevCameraCut = CameraCutParams[0];
@@ -552,7 +569,10 @@ void UMovieSceneCameraCutTrackInstance::OnAnimate()
 		FinalCameraCut = CameraCutParams[0];
 
 		// It may be blending out back to gameplay however.
-		if (FinalCameraCut.LocalContextTime > FinalCameraCut.LocalEaseOutStartTime)
+		const bool bIsBackwards = FinalCameraCut.LocalDirection == EPlayDirection::Backwards;
+		if (
+				(!bIsBackwards && FinalCameraCut.LocalContextTime > FinalCameraCut.LocalEaseOutStartTime) ||
+				(bIsBackwards && FinalCameraCut.LocalContextTime < FinalCameraCut.LocalEaseOutStartTime))
 		{
 			FinalCameraCut.PreviewBlendFactor = 1.0f - FinalCameraCut.PreviewBlendFactor;
 			FinalCameraCut.EaseIn = FinalCameraCut.EaseOut;
