Sequencer: Hid fix for UE-215502 behind a cvar due to unexpected but rare edge-case crash

The cvar Sequencer.Animation.UIFlickerFix (default: false) enables the fix for UI flicker. This temporarily work around an crash that exists when animating Anim Class along with a skel anim and control rig tracks (a very uncommon workflow).

diff --git a/Engine/Source/Runtime/MovieSceneTracks/Private/Systems/MovieSceneSkeletalAnimationSystem.cpp b/Engine/Source/Runtime/MovieSceneTracks/Private/Systems/MovieSceneSkeletalAnimationSystem.cpp
--- a/Engine/Source/Runtime/MovieSceneTracks/Private/Systems/MovieSceneSkeletalAnimationSystem.cpp
+++ b/Engine/Source/Runtime/MovieSceneTracks/Private/Systems/MovieSceneSkeletalAnimationSystem.cpp
@@ -52,6 +52,13 @@ DECLARE_CYCLE_STAT(TEXT("Evaluate skeletal animations"), MovieSceneEval_Evaluate
 namespace UE::MovieScene
 {
 
+bool GAnimationUIFlickerFix = false;
+FAutoConsoleVariableRef CVarAnimationUIFlickerFix(
+	TEXT("Sequencer.Animation.UIFlickerFix"),
+	GAnimationUIFlickerFix,
+	TEXT("(Default: true. Fixes pre-animated state ordering that was causing excessive UI flicker. Known to cause issues when animating Anim Class so should be disabled if a crash is encountered.")
+	);
+
 /** Helper function to get our sequencer animation node from a skeletal mesh component */
 UAnimSequencerInstance* GetAnimSequencerInstance(USkeletalMeshComponent* SkeletalMeshComponent)
 {
@@ -456,6 +463,7 @@ private:
 	UMovieSceneEntitySystemLinker* Linker;
 	FSkeletalAnimationSystemData* SystemData;
 
+	TSharedPtr<FPreAnimatedSkeletalAnimationStorage> PreAnimatedStorage;
 	TSharedPtr<FPreAnimatedSkeletalAnimationMontageStorage> PreAnimatedMontageStorage;
 	TSharedPtr<FPreAnimatedSkeletalAnimationAnimInstanceStorage> PreAnimatedAnimInstanceStorage;
 
@@ -465,6 +473,7 @@ public:
 		: Linker(InLinker)
 		, SystemData(InSystemData)
 	{
+		PreAnimatedStorage = InLinker->PreAnimatedState.GetOrCreateStorage<FPreAnimatedSkeletalAnimationStorage>();
 		PreAnimatedMontageStorage = InLinker->PreAnimatedState.GetOrCreateStorage<FPreAnimatedSkeletalAnimationMontageStorage>();
 		PreAnimatedAnimInstanceStorage = InLinker->PreAnimatedState.GetOrCreateStorage<FPreAnimatedSkeletalAnimationAnimInstanceStorage>();
 	}
@@ -492,6 +501,20 @@ private:
 			return;
 		}
 
+		if (GAnimationUIFlickerFix == false)
+		{
+			// Cache pre-animated state for this bound object before doing anything.
+			// We don't yet track what entities have already started animated vs. entities that just started this frame,
+			// so we just process all the currently active ones. If they are already tracked and have already had their
+			// pre-animated state saved, it these calls will just early return.
+			for (const FActiveSkeletalAnimation& SkeletalAnimation : InSkeletalAnimations.Animations)
+			{
+				PreAnimatedStorage->BeginTrackingEntity(SkeletalAnimation.EntityID, SkeletalAnimation.bWantsRestoreState, SkeletalAnimation.RootInstanceHandle, SkeletalMeshComponent);
+			}
+			FCachePreAnimatedValueParams CacheParams;
+			PreAnimatedStorage->CachePreAnimatedValue(CacheParams, SkeletalMeshComponent);
+		}
+
 		// Setup any needed animation nodes for sequencer playback.
 		UAnimInstance* ExistingAnimInstance = GetSourceAnimInstance(SkeletalMeshComponent);
 		bool bWasCreated = false;
@@ -1075,15 +1098,18 @@ void UMovieSceneSkeletalAnimationSystem::OnRun(FSystemTaskPrerequisites& InPrere
 	TSharedRef<FMovieSceneEntitySystemRunner> Runner = Linker->GetRunner();
 	if (Runner->GetCurrentPhase() == ESystemPhase::Instantiation)
 	{
-		// Begin tracking pre-animated state for all bound skel animation components
-		TSharedPtr<FPreAnimatedSkeletalAnimationStorage> PreAnimatedStorage = Linker->PreAnimatedState.GetOrCreateStorage<FPreAnimatedSkeletalAnimationStorage>();
-
-		struct FTask
+		if (GAnimationUIFlickerFix == true)
 		{
-			FEntityComponentFilter AdditionalFilter;
-		} Task;
-		Task.AdditionalFilter.All({ TrackComponents->SkeletalAnimation, BuiltInComponents->Tags.NeedsLink });
-		PreAnimatedStorage->BeginTrackingAndCachePreAnimatedValuesTask(Linker, Task, BuiltInComponents->BoundObject);
+			// Begin tracking pre-animated state for all bound skel animation components
+			TSharedPtr<FPreAnimatedSkeletalAnimationStorage> PreAnimatedStorage = Linker->PreAnimatedState.GetOrCreateStorage<FPreAnimatedSkeletalAnimationStorage>();
+
+			struct FTask
+			{
+				FEntityComponentFilter AdditionalFilter;
+			} Task;
+			Task.AdditionalFilter.All({ TrackComponents->SkeletalAnimation, BuiltInComponents->Tags.NeedsLink });
+			PreAnimatedStorage->BeginTrackingAndCachePreAnimatedValuesTask(Linker, Task, BuiltInComponents->BoundObject);
+		}
 
 		CleanSystemData();
 		return;
