Sequencer: Fixes to keying transform tracks with root motion 'swap actor bone' or 'swap component' enabled.

* Using the gizmo to move the character around while root motion was being applied didn't work- the character kept getting snapped back to its original spot by root motion code. This has been fixed by not re-applying root motion during actor movement in editor.
* The gizmo would not move along with the actor when in Sequencer. This was because character position is applied in 2 stages- first transform track, then later in the frame in PostPhysics the root motion would be applied. However, the gizmo position was updated before the root motion was applied so wasn't correctly updated to the actor position after root motion. This has been fixed by adding a tick function in LevelEditorSequencerIntegration to update the gizmo position later in the tick after root motion is applied.
* Keying the transform track was not working- if you keyed during the animation, it was doubly applying the current root motion position by adding it to the keyed value. This has been fixed by decomposing the root motion from the actor position before keying, which involved having the transform track editor speaking to the skeletal animation system. When this is generalized, this code will need to be updated again.

diff --git a/Engine/Source/Editor/MovieSceneTools/Private/TrackEditors/TransformTrackEditor.cpp b/Engine/Source/Editor/MovieSceneTools/Private/TrackEditors/TransformTrackEditor.cpp
--- a/Engine/Source/Editor/MovieSceneTools/Private/TrackEditors/TransformTrackEditor.cpp
+++ b/Engine/Source/Editor/MovieSceneTools/Private/TrackEditors/TransformTrackEditor.cpp
@@ -53,6 +53,7 @@
 #include "ISequencerObjectChangeListener.h"
 #include "ISequencerPropertyKeyedStatus.h"
 #include "Systems/MovieSceneTransformOriginSystem.h"
+#include "Systems/MovieSceneSkeletalAnimationSystem.h"
 
 #define LOCTEXT_NAMESPACE "MovieScene_TransformTrack"
 
@@ -1012,6 +1013,40 @@ FTransform F3DTransformTrackEditor::GetTransformOrigin() const
 	return TransformOrigin;
 }
 
+FTransform F3DTransformTrackEditor::GetRootMotionOffset(UObject* AnimatedObject) const
+{
+	FTransform RootMotionOffset;
+
+	const FMovieSceneRootEvaluationTemplateInstance& EvaluationTemplate = GetSequencer()->GetEvaluationTemplate();
+	const UMovieSceneEntitySystemLinker* EntityLinker = EvaluationTemplate.GetEntitySystemLinker();
+	if (!EntityLinker)
+	{
+		return RootMotionOffset;
+	}
+
+	if (const UMovieSceneSkeletalAnimationSystem* SkeletalAnimationSystem = EntityLinker->FindSystem<UMovieSceneSkeletalAnimationSystem>())
+	{
+		return SkeletalAnimationSystem->GetRootMotionOffset(AnimatedObject);
+	}
+
+	return RootMotionOffset;
+}
+
+void F3DTransformTrackEditor::UpdateRootMotionOffset(UObject* AnimatedObject)
+{
+	const FMovieSceneRootEvaluationTemplateInstance& EvaluationTemplate = GetSequencer()->GetEvaluationTemplate();
+	const UMovieSceneEntitySystemLinker* EntityLinker = EvaluationTemplate.GetEntitySystemLinker();
+	if (!EntityLinker)
+	{
+		return;
+	}
+
+	if (UMovieSceneSkeletalAnimationSystem* SkeletalAnimationSystem = EntityLinker->FindSystem<UMovieSceneSkeletalAnimationSystem>())
+	{
+		SkeletalAnimationSystem->UpdateRootMotionOffset(AnimatedObject);
+	}
+}
+
 void F3DTransformTrackEditor::AddTransformKeysForHandle(TArray<FGuid> ObjectHandles, EMovieSceneTransformChannel ChannelToKey, ESequencerKeyMode KeyMode)
 {
 	const FScopedTransaction Transaction(NSLOCTEXT("Sequencer", "AddTransformTrack", "Add Transform Track"));
@@ -1074,6 +1109,10 @@ void F3DTransformTrackEditor::AddTransformKeys( UObject* ObjectToKey, const TOpt
 				}
 			}
 		}
+
+		// Update root motion if necessary so it doesn't double-compensate in the case we don't force update the sequencer frame
+		UpdateRootMotionOffset(ObjectToKey);
+
 		return KeyPropertyResult;
 	};
 
@@ -1107,7 +1146,13 @@ FTransformData F3DTransformTrackEditor::RecomposeTransform(const FTransformData&
 	TArray<FMovieSceneEntityID> ImportedEntityIDs;
 	EvaluationTemplate.FindEntitiesFromOwner(Section, GetSequencer()->GetFocusedTemplateID(), ImportedEntityIDs);
 
-	FTransform CurrentTransform = SceneComponent->GetRelativeTransform();
+	FTransform CurrentTransform = FTransform(InTransformData.Rotation, InTransformData.Translation, InTransformData.Scale);
+
+	// Subtract any root motion first (both from the 'CurrentTransform' and the InTransformData)
+	if (!SceneComponent->GetAttachParent() && Section->GetBlendType() == EMovieSceneBlendType::Absolute)
+	{
+		CurrentTransform *= GetRootMotionOffset(AnimatedObject).Inverse();
+	}
 
 	if (ImportedEntityIDs.Num())
 	{
@@ -1142,7 +1187,7 @@ FTransformData F3DTransformTrackEditor::RecomposeTransform(const FTransformData&
 			Query.Object   = SceneComponent;
 			Query.bConvertFromSourceEntityIDs = false;  // We already pass the children entity IDs
 
-			FIntermediate3DTransform CurrentValue(InTransformData.Translation, InTransformData.Rotation, InTransformData.Scale);
+			FIntermediate3DTransform CurrentValue(CurrentTransform.GetTranslation(), CurrentTransform.GetRotation().Rotator(), CurrentTransform.GetScale3D());
 
 			TRecompositionResult<FIntermediate3DTransform> TransformData = System->RecomposeBlendOperational(TrackComponents->ComponentTransform, Query, CurrentValue);
 
diff --git a/Engine/Source/Editor/MovieSceneTools/Private/TrackEditors/TransformTrackEditor.h b/Engine/Source/Editor/MovieSceneTools/Private/TrackEditors/TransformTrackEditor.h
--- a/Engine/Source/Editor/MovieSceneTools/Private/TrackEditors/TransformTrackEditor.h
+++ b/Engine/Source/Editor/MovieSceneTools/Private/TrackEditors/TransformTrackEditor.h
@@ -145,6 +145,12 @@ private:
 	/** Transform origin which may be set for the current level sequence */
 	FTransform GetTransformOrigin() const;
 
+	/*In the case the object is being animated with root motion, return the offset from the keyed transform value that root motion is currently applying.*/
+	FTransform GetRootMotionOffset(UObject* AnimatedObject) const;
+
+	/*In the case the object is being animated with root motion, updates the offset based on the current object position.*/
+	void UpdateRootMotionOffset(UObject* AnimatedObject);
+
 	/** 
 	 * Adds transform keys to an object represented by a handle.
 
diff --git a/Engine/Source/Editor/Sequencer/Private/LevelEditorSequencerIntegration.cpp b/Engine/Source/Editor/Sequencer/Private/LevelEditorSequencerIntegration.cpp
--- a/Engine/Source/Editor/Sequencer/Private/LevelEditorSequencerIntegration.cpp
+++ b/Engine/Source/Editor/Sequencer/Private/LevelEditorSequencerIntegration.cpp
@@ -174,6 +174,25 @@ private:
 	TArray<TWeakPtr<ISequencer>> Sequencers;
 };
 
+void FLevelEditorSequencerUpdateGizmoTickFunction::ExecuteTick(float DeltaTime, enum ELevelTick TickType, ENamedThreads::Type CurrentThread, const FGraphEventRef& MyCompletionGraphEvent)
+{
+	if (bSequencerEvaluated)
+	{
+		GUnrealEd->UpdatePivotLocationForSelection();
+		bSequencerEvaluated = false;
+	}
+}
+
+FString FLevelEditorSequencerUpdateGizmoTickFunction::DiagnosticMessage()
+{
+	return TEXT("[SequencerUpdateGizmoPosition]");
+}
+
+FName FLevelEditorSequencerUpdateGizmoTickFunction::DiagnosticContext(bool bDetailed)
+{
+	return FName(TEXT("SequencerUpdateGizmoPosition"));
+}
+
 static const FName DetailsTabIdentifiers[] = { "LevelEditorSelectionDetails", "LevelEditorSelectionDetails2", "LevelEditorSelectionDetails3", "LevelEditorSelectionDetails4" };
 
 FLevelEditorSequencerIntegration::FLevelEditorSequencerIntegration()
@@ -544,7 +563,7 @@ void FLevelEditorSequencerIntegration::OnSequencerEvaluated()
 	}
 
 	// If realtime is off, this needs to be called to update the pivot location when scrubbing.
-	GUnrealEd->UpdatePivotLocationForSelection();
+	UpdateGizmoTickFunction.bSequencerEvaluated = true;
 }
 
 void FLevelEditorSequencerIntegration::OnBeginDeferUpdates()
@@ -1356,6 +1375,22 @@ void FLevelEditorSequencerIntegration::AddSequencer(TSharedRef<ISequencer> InSeq
 	{
 		Initialize(Options);
 	}
+	if (UObject* PlaybackContext = InSequencer->GetSharedPlaybackState()->GetPlaybackContext())
+	{
+		if (UWorld* World = PlaybackContext->GetWorld())
+		{
+			if (World->PersistentLevel)
+			{
+				if (!UpdateGizmoTickFunction.IsTickFunctionRegistered())
+				{
+					UpdateGizmoTickFunction.TickGroup = TG_PostUpdateWork;
+					UpdateGizmoTickFunction.bCanEverTick = true;
+					UpdateGizmoTickFunction.bStartWithTickEnabled = true;
+					UpdateGizmoTickFunction.RegisterTickFunction(World->PersistentLevel);
+				}
+			}
+		}
+	}
 
 	KeyFrameHandler->Add(InSequencer);
 
diff --git a/Engine/Source/Editor/Sequencer/Public/LevelEditorSequencerIntegration.h b/Engine/Source/Editor/Sequencer/Public/LevelEditorSequencerIntegration.h
--- a/Engine/Source/Editor/Sequencer/Public/LevelEditorSequencerIntegration.h
+++ b/Engine/Source/Editor/Sequencer/Public/LevelEditorSequencerIntegration.h
@@ -9,6 +9,8 @@
 #include "ISceneOutlinerColumn.h"
 #include "UObject/ObjectKey.h"
 #include "ISequencer.h"
+#include "Engine/EngineBaseTypes.h"
+#include "LevelEditorSequencerIntegration.generated.h"
 
 class AActor;
 class FExtender;
@@ -76,6 +78,39 @@ private:
 	FPropertyBindingsDataChanged PropertyBindingsDataChanged;
 };
 
+/**
+* Tick function whose sole purpose is to update the gizmo position after any root motion may have run on selected characters at the end of the frame.
+**/
+USTRUCT()
+struct FLevelEditorSequencerUpdateGizmoTickFunction : public FTickFunction
+{
+	GENERATED_USTRUCT_BODY()
+
+	bool bSequencerEvaluated = false;
+
+	/**
+	* Abstract function to execute the tick.
+	* @param DeltaTime - frame time to advance, in seconds.
+	* @param TickType - kind of tick for this frame.
+	* @param CurrentThread - thread we are executing on, useful to pass along as new tasks are created.
+	* @param MyCompletionGraphEvent - completion event for this task. Useful for holding the completetion of this task until certain child tasks are complete.
+	*/
+	virtual void ExecuteTick(float DeltaTime, enum ELevelTick TickType, ENamedThreads::Type CurrentThread, const FGraphEventRef& MyCompletionGraphEvent) override;
+	/** Abstract function to describe this tick. Used to print messages about illegal cycles in the dependency graph. */
+	virtual FString DiagnosticMessage() override;
+	/** Function used to describe this tick for active tick reporting. **/
+	virtual FName DiagnosticContext(bool bDetailed) override;
+};
+
+template<>
+struct TStructOpsTypeTraits<FLevelEditorSequencerUpdateGizmoTickFunction> : public TStructOpsTypeTraitsBase2<FLevelEditorSequencerUpdateGizmoTickFunction>
+{
+	enum
+	{
+		WithCopy = false
+	};
+};
+
 
 class SEQUENCER_API FLevelEditorSequencerIntegration
 {
@@ -227,4 +262,7 @@ private:
 	bool bDeferUpdates;
 
 	FOnSequencersChanged OnSequencersChanged;
+
+	FLevelEditorSequencerUpdateGizmoTickFunction UpdateGizmoTickFunction;
+
 };
diff --git a/Engine/Source/Runtime/MovieSceneTracks/Private/Systems/MovieSceneSkeletalAnimationSystem.cpp b/Engine/Source/Runtime/MovieSceneTracks/Private/Systems/MovieSceneSkeletalAnimationSystem.cpp
--- a/Engine/Source/Runtime/MovieSceneTracks/Private/Systems/MovieSceneSkeletalAnimationSystem.cpp
+++ b/Engine/Source/Runtime/MovieSceneTracks/Private/Systems/MovieSceneSkeletalAnimationSystem.cpp
@@ -44,6 +44,12 @@
 
 #include "SequencerAnimationOverride.h"
 
+#if WITH_EDITOR
+#include "Editor.h"
+#include "Editor/UnrealEdEngine.h"
+#include "UnrealEdGlobals.h"
+#endif
+
 #include UE_INLINE_GENERATED_CPP_BY_NAME(MovieSceneSkeletalAnimationSystem)
 
 DECLARE_CYCLE_STAT(TEXT("Gather skeletal animations"), MovieSceneEval_GatherSkeletalAnimations, STATGROUP_MovieSceneECS);
@@ -272,6 +278,24 @@ void FBoneTransformFinalizeData::Register(USkeletalMeshComponent* InSkeleletalMe
 					BoneTransformFinalized();
 				}));
 	}
+
+#if WITH_EDITOR
+	if (OnBeginActorMovementHandle.IsValid() == false)
+	{
+		OnBeginActorMovementHandle = GEditor->OnBeginObjectMovement().Add(UEditorEngine::FOnBeginTransformObject::FDelegate::CreateLambda([this](UObject&)
+				{
+					bActorBeingMoved = true;
+				}));
+	}
+
+	if (OnEndActorMovementHandle.IsValid() == false)
+	{
+		OnEndActorMovementHandle = GEditor->OnEndObjectMovement().Add(UEditorEngine::FOnEndTransformObject::FDelegate::CreateLambda([this](UObject&)
+			{
+				bActorBeingMoved = false;
+			}));
+	}
+#endif
 }
 
 void FBoneTransformFinalizeData::Unregister()
@@ -281,10 +305,24 @@ void FBoneTransformFinalizeData::Unregister()
 		SkeletalMeshComponent->UnregisterOnBoneTransformsFinalizedDelegate(OnBoneTransformsFinalizedHandle);
 	}
 	OnBoneTransformsFinalizedHandle.Reset();
+
+#if WITH_EDITOR
+	GEditor->OnBeginObjectMovement().Remove(OnBeginActorMovementHandle);
+	OnBeginActorMovementHandle.Reset();
+	GEditor->OnEndObjectMovement().Remove(OnEndActorMovementHandle);
+	OnEndActorMovementHandle.Reset();
+#endif
 }
 
 void FBoneTransformFinalizeData::BoneTransformFinalized()
 {
+#if WITH_EDITOR
+	if (bActorBeingMoved)
+	{
+		return;
+	}
+#endif
+
 	if (SkeletalMeshComponent && SwapRootBone != ESwapRootBone::SwapRootBone_None)
 	{
 		FTransform RelativeTransform = CurrentTransform;
@@ -1006,6 +1044,57 @@ UObject* UMovieSceneSkeletalAnimationSystem::ResolveSkeletalMeshComponentBinding
 	return nullptr;
 }
 
+FTransform UMovieSceneSkeletalAnimationSystem::GetRootMotionOffset(UObject* InObject) const
+{
+	using namespace UE::MovieScene;
+	
+	FTransform RootMotionOffset;
+	if (USkeletalMeshComponent* SkeletalMeshComponent = Cast<USkeletalMeshComponent>(ResolveSkeletalMeshComponentBinding(InObject)))
+	{
+		if (const FBoundObjectActiveSkeletalAnimations* BoundObjectAnimations = SystemData.SkeletalAnimations.Find(SkeletalMeshComponent))
+		{
+			RootMotionOffset = BoundObjectAnimations->BoneTransformFinalizeData.CurrentTransform;
+		}
+	}
+
+	return RootMotionOffset;
+}
+
+void UMovieSceneSkeletalAnimationSystem::UpdateRootMotionOffset(UObject* InObject)
+{
+	using namespace UE::MovieScene;
+
+	if (USkeletalMeshComponent* SkeletalMeshComponent = Cast<USkeletalMeshComponent>(ResolveSkeletalMeshComponentBinding(InObject)))
+	{
+		if (FBoundObjectActiveSkeletalAnimations* BoundObjectAnimations = SystemData.SkeletalAnimations.Find(SkeletalMeshComponent))
+		{
+			if (BoundObjectAnimations->BoneTransformFinalizeData.InitialTransform.IsSet())
+			{
+				FTransform RootMotionOffset = BoundObjectAnimations->BoneTransformFinalizeData.CurrentTransform;
+				FTransform CurrentTransform;
+				if (BoundObjectAnimations->BoneTransformFinalizeData.SwapRootBone == ESwapRootBone::SwapRootBone_Component)
+				{
+					CurrentTransform = SkeletalMeshComponent->GetRelativeTransform();
+				}
+				else if (BoundObjectAnimations->BoneTransformFinalizeData.SwapRootBone == ESwapRootBone::SwapRootBone_Actor)
+				{
+					if (AActor* Actor = SkeletalMeshComponent->GetOwner())
+					{
+						if (USceneComponent* RootComponent = Actor->GetRootComponent())
+						{
+							CurrentTransform = RootComponent->GetRelativeTransform();
+						}
+					}
+				}
+				// Subtract root motion off of this
+				CurrentTransform = CurrentTransform * RootMotionOffset.Inverse();
+				// Reset the initial transform based on this.
+				BoundObjectAnimations->BoneTransformFinalizeData.InitialTransform = CurrentTransform;
+			}
+		}
+	}
+}
+
 void UMovieSceneSkeletalAnimationSystem::OnSchedulePersistentTasks(UE::MovieScene::IEntitySystemScheduler* TaskScheduler)
 {
 	using namespace UE::MovieScene;
diff --git a/Engine/Source/Runtime/MovieSceneTracks/Public/Systems/MovieSceneSkeletalAnimationSystem.h b/Engine/Source/Runtime/MovieSceneTracks/Public/Systems/MovieSceneSkeletalAnimationSystem.h
--- a/Engine/Source/Runtime/MovieSceneTracks/Public/Systems/MovieSceneSkeletalAnimationSystem.h
+++ b/Engine/Source/Runtime/MovieSceneTracks/Public/Systems/MovieSceneSkeletalAnimationSystem.h
@@ -46,6 +46,12 @@ struct FBoneTransformFinalizeData
 	/** Delegate Handle for skel mesh bone transform finalized callback*/
 	FDelegateHandle OnBoneTransformsFinalizedHandle;
 
+	#if WITH_EDITORONLY_DATA
+	FDelegateHandle OnBeginActorMovementHandle;
+	FDelegateHandle OnEndActorMovementHandle;
+	bool bActorBeingMoved = false;
+	#endif
+
 	void Register(USkeletalMeshComponent* InSkeleletalMeshCompononent, ESwapRootBone InSwapRootBone, FTransform& InCurrentTransform, TOptional<FTransform> InTransform);
 	void Unregister();
 	void BoneTransformFinalized();
@@ -97,6 +103,9 @@ public:
 
 	static UObject* ResolveSkeletalMeshComponentBinding(UObject* InObject);
 
+	MOVIESCENETRACKS_API FTransform GetRootMotionOffset(UObject* InObject) const;
+	MOVIESCENETRACKS_API void UpdateRootMotionOffset(UObject* InObject);
+
 private:
 
 	virtual void OnSchedulePersistentTasks(UE::MovieScene::IEntitySystemScheduler* TaskScheduler) override;
