Sequencer: Fix root motion and authoring of root motion in cases where the Skeletal Mesh Component is not the root component and has a rotational offset (e.g. 90 degrees).

diff --git a/Engine/Source/Editor/MovieSceneTools/Private/TrackEditors/TransformTrackEditor.cpp b/Engine/Source/Editor/MovieSceneTools/Private/TrackEditors/TransformTrackEditor.cpp
--- a/Engine/Source/Editor/MovieSceneTools/Private/TrackEditors/TransformTrackEditor.cpp
+++ b/Engine/Source/Editor/MovieSceneTools/Private/TrackEditors/TransformTrackEditor.cpp
@@ -1151,7 +1151,7 @@ FTransformData F3DTransformTrackEditor::RecomposeTransform(const FTransformData&
 	// Subtract any root motion first (both from the 'CurrentTransform' and the InTransformData)
 	if (!SceneComponent->GetAttachParent() && Section->GetBlendType() == EMovieSceneBlendType::Absolute)
 	{
-		CurrentTransform *= GetRootMotionOffset(AnimatedObject).Inverse();
+		CurrentTransform = GetRootMotionOffset(AnimatedObject).Inverse() * CurrentTransform;
 	}
 
 	if (ImportedEntityIDs.Num())
diff --git a/Engine/Source/Runtime/MovieSceneTracks/Private/Systems/MovieSceneSkeletalAnimationSystem.cpp b/Engine/Source/Runtime/MovieSceneTracks/Private/Systems/MovieSceneSkeletalAnimationSystem.cpp
--- a/Engine/Source/Runtime/MovieSceneTracks/Private/Systems/MovieSceneSkeletalAnimationSystem.cpp
+++ b/Engine/Source/Runtime/MovieSceneTracks/Private/Systems/MovieSceneSkeletalAnimationSystem.cpp
@@ -264,14 +264,20 @@ FBoneTransformFinalizeData::FBoneTransformFinalizeData() : SkeletalMeshComponent
 {
 }
 
-void FBoneTransformFinalizeData::Register(USkeletalMeshComponent* InSkeleletalMeshCompononent, ESwapRootBone InSwapRootBone, FTransform& InCurrentTransform,TOptional<FTransform> InInitialTransform)
+void FBoneTransformFinalizeData::Register(USkeletalMeshComponent* InSkeleletalMeshCompononent, ESwapRootBone InSwapRootBone, FTransform& InMeshRelativeRootMotionTransform,TOptional<FTransform> InInitialActorTransform)
 {
 	if (OnBoneTransformsFinalizedHandle.IsValid() == false)
 	{
 		SkeletalMeshComponent = InSkeleletalMeshCompononent;
 		SwapRootBone = InSwapRootBone;
-		CurrentTransform = InCurrentTransform;
-		InitialTransform = InInitialTransform;
+		MeshRelativeRootMotionTransform = InMeshRelativeRootMotionTransform;
+		InitialActorTransform = InInitialActorTransform;
+		// Also store the inverse relative rotation between SkeletalMeshComponent and the Actor if it's not the root component
+		if (SkeletalMeshComponent && SkeletalMeshComponent != SkeletalMeshComponent->GetOwner()->GetRootComponent())
+		{
+			InverseMeshToActorRotation = SkeletalMeshComponent->GetOwner()->GetRootComponent()->GetComponentTransform().GetRelativeTransformReverse(SkeletalMeshComponent->GetComponentTransform()).GetRotation();
+		}
+
 		OnBoneTransformsFinalizedHandle = SkeletalMeshComponent->RegisterOnBoneTransformsFinalizedDelegate(
 			FOnBoneTransformsFinalizedMultiCast::FDelegate::CreateLambda([this]()
 				{
@@ -305,6 +311,8 @@ void FBoneTransformFinalizeData::Unregister()
 		SkeletalMeshComponent->UnregisterOnBoneTransformsFinalizedDelegate(OnBoneTransformsFinalizedHandle);
 	}
 	OnBoneTransformsFinalizedHandle.Reset();
+	InitialActorTransform.Reset();
+	InverseMeshToActorRotation.Reset();
 
 #if WITH_EDITOR
 	GEditor->OnBeginObjectMovement().Remove(OnBeginActorMovementHandle);
@@ -325,11 +333,11 @@ void FBoneTransformFinalizeData::BoneTransformFinalized()
 
 	if (SkeletalMeshComponent && SwapRootBone != ESwapRootBone::SwapRootBone_None)
 	{
-		FTransform RelativeTransform = CurrentTransform;
-		
-		if (InitialTransform.IsSet())
+		FTransform RelativeTransform = MeshRelativeRootMotionTransform;
+
+		if (InitialActorTransform.IsSet())
 		{
-			RelativeTransform = RelativeTransform * InitialTransform.GetValue();
+			RelativeTransform = RelativeTransform * InitialActorTransform.GetValue();
 		}
 
 		if (SwapRootBone == ESwapRootBone::SwapRootBone_Component)
@@ -341,6 +349,14 @@ void FBoneTransformFinalizeData::BoneTransformFinalized()
 			AActor* Actor = SkeletalMeshComponent->GetOwner();
 			if (Actor && Actor->GetRootComponent())
 			{
+				// Compensate for any mesh component rotation
+				if (InitialActorTransform.IsSet() && InverseMeshToActorRotation.IsSet())
+				{
+					RelativeTransform = MeshRelativeRootMotionTransform;
+					RelativeTransform.SetTranslation(InverseMeshToActorRotation.GetValue() * RelativeTransform.GetTranslation());
+					RelativeTransform = RelativeTransform * InitialActorTransform.GetValue();
+				}
+
 				Actor->GetRootComponent()->SetRelativeLocationAndRotation(RelativeTransform.GetLocation(), RelativeTransform.GetRotation().Rotator());
 			}
 		}
@@ -697,10 +713,10 @@ private:
 			if (AnimSection->Params.SwapRootBone != ESwapRootBone::SwapRootBone_None)
 			{
 				FTransform Transform = RootMotionParams.Transform.IsSet() ? RootMotionParams.Transform.GetValue() : FTransform::Identity;
-				TOptional<FTransform> InitialTransform = GetCurrentTransform(AnimSection->Params.SwapRootBone, SkeletalMeshComponent);
+				TOptional<FTransform> InitialActorTransform = GetCurrentTransform(AnimSection->Params.SwapRootBone, SkeletalMeshComponent);
 				if (FBoundObjectActiveSkeletalAnimations* BoundObjectAnimations = SystemData->SkeletalAnimations.Find(SkeletalMeshComponent))
 				{
-					BoundObjectAnimations->BoneTransformFinalizeData.Register(SkeletalMeshComponent, AnimSection->Params.SwapRootBone,Transform, InitialTransform);
+					BoundObjectAnimations->BoneTransformFinalizeData.Register(SkeletalMeshComponent, AnimSection->Params.SwapRootBone,Transform, InitialActorTransform);
 				}
 			}
 
@@ -1053,7 +1069,12 @@ FTransform UMovieSceneSkeletalAnimationSystem::GetRootMotionOffset(UObject* InOb
 	{
 		if (const FBoundObjectActiveSkeletalAnimations* BoundObjectAnimations = SystemData.SkeletalAnimations.Find(SkeletalMeshComponent))
 		{
-			RootMotionOffset = BoundObjectAnimations->BoneTransformFinalizeData.CurrentTransform;
+			RootMotionOffset = BoundObjectAnimations->BoneTransformFinalizeData.MeshRelativeRootMotionTransform;
+			if (BoundObjectAnimations->BoneTransformFinalizeData.SwapRootBone == ESwapRootBone::SwapRootBone_Actor
+				&& BoundObjectAnimations->BoneTransformFinalizeData.InverseMeshToActorRotation.IsSet())
+			{
+				RootMotionOffset.SetTranslation(BoundObjectAnimations->BoneTransformFinalizeData.InverseMeshToActorRotation.GetValue() * RootMotionOffset.GetTranslation());
+			}
 		}
 	}
 
@@ -1068,9 +1089,9 @@ void UMovieSceneSkeletalAnimationSystem::UpdateRootMotionOffset(UObject* InObjec
 	{
 		if (FBoundObjectActiveSkeletalAnimations* BoundObjectAnimations = SystemData.SkeletalAnimations.Find(SkeletalMeshComponent))
 		{
-			if (BoundObjectAnimations->BoneTransformFinalizeData.InitialTransform.IsSet())
+			if (BoundObjectAnimations->BoneTransformFinalizeData.InitialActorTransform.IsSet())
 			{
-				FTransform RootMotionOffset = BoundObjectAnimations->BoneTransformFinalizeData.CurrentTransform;
+				FTransform RootMotionOffset = GetRootMotionOffset(InObject);
 				FTransform CurrentTransform;
 				if (BoundObjectAnimations->BoneTransformFinalizeData.SwapRootBone == ESwapRootBone::SwapRootBone_Component)
 				{
@@ -1087,14 +1108,45 @@ void UMovieSceneSkeletalAnimationSystem::UpdateRootMotionOffset(UObject* InObjec
 					}
 				}
 				// Subtract root motion off of this
-				CurrentTransform = CurrentTransform * RootMotionOffset.Inverse();
+				CurrentTransform = RootMotionOffset.Inverse() * CurrentTransform;
+				
 				// Reset the initial transform based on this.
-				BoundObjectAnimations->BoneTransformFinalizeData.InitialTransform = CurrentTransform;
+				BoundObjectAnimations->BoneTransformFinalizeData.InitialActorTransform = CurrentTransform;
 			}
 		}
 	}
 }
 
+TOptional<FTransform> UMovieSceneSkeletalAnimationSystem::GetInitialActorTransform(UObject* InObject) const
+{
+	using namespace UE::MovieScene;
+
+	TOptional<FTransform> InitialActorTransform;
+	if (USkeletalMeshComponent* SkeletalMeshComponent = Cast<USkeletalMeshComponent>(ResolveSkeletalMeshComponentBinding(InObject)))
+	{
+		if (const FBoundObjectActiveSkeletalAnimations* BoundObjectAnimations = SystemData.SkeletalAnimations.Find(SkeletalMeshComponent))
+		{
+			InitialActorTransform = BoundObjectAnimations->BoneTransformFinalizeData.InitialActorTransform;
+		}
+	}
+	return InitialActorTransform;
+}
+
+TOptional<FQuat> UMovieSceneSkeletalAnimationSystem::GetInverseMeshToActorRotation(UObject* InObject) const
+{
+	using namespace UE::MovieScene;
+
+	TOptional<FQuat> InverseMeshToActorRotation;
+	if (USkeletalMeshComponent* SkeletalMeshComponent = Cast<USkeletalMeshComponent>(ResolveSkeletalMeshComponentBinding(InObject)))
+	{
+		if (const FBoundObjectActiveSkeletalAnimations* BoundObjectAnimations = SystemData.SkeletalAnimations.Find(SkeletalMeshComponent))
+		{
+			InverseMeshToActorRotation = BoundObjectAnimations->BoneTransformFinalizeData.InverseMeshToActorRotation;
+		}
+	}
+	return InverseMeshToActorRotation;
+}
+
 void UMovieSceneSkeletalAnimationSystem::OnSchedulePersistentTasks(UE::MovieScene::IEntitySystemScheduler* TaskScheduler)
 {
 	using namespace UE::MovieScene;
diff --git a/Engine/Source/Runtime/MovieSceneTracks/Public/Systems/MovieSceneSkeletalAnimationSystem.h b/Engine/Source/Runtime/MovieSceneTracks/Public/Systems/MovieSceneSkeletalAnimationSystem.h
--- a/Engine/Source/Runtime/MovieSceneTracks/Public/Systems/MovieSceneSkeletalAnimationSystem.h
+++ b/Engine/Source/Runtime/MovieSceneTracks/Public/Systems/MovieSceneSkeletalAnimationSystem.h
@@ -41,8 +41,9 @@ struct FBoneTransformFinalizeData
 	FBoneTransformFinalizeData();
 	USkeletalMeshComponent* SkeletalMeshComponent;
 	ESwapRootBone SwapRootBone;
-	FTransform CurrentTransform;
-	TOptional<FTransform> InitialTransform;
+	FTransform MeshRelativeRootMotionTransform;
+	TOptional<FTransform> InitialActorTransform;
+	TOptional<FQuat> InverseMeshToActorRotation;
 	/** Delegate Handle for skel mesh bone transform finalized callback*/
 	FDelegateHandle OnBoneTransformsFinalizedHandle;
 
@@ -52,7 +53,7 @@ struct FBoneTransformFinalizeData
 	bool bActorBeingMoved = false;
 	#endif
 
-	void Register(USkeletalMeshComponent* InSkeleletalMeshCompononent, ESwapRootBone InSwapRootBone, FTransform& InCurrentTransform, TOptional<FTransform> InTransform);
+	void Register(USkeletalMeshComponent* InSkeleletalMeshCompononent, ESwapRootBone InSwapRootBone, FTransform& InMeshRelativeRootMotionTransform, TOptional<FTransform> InInitialActorTransform);
 	void Unregister();
 	void BoneTransformFinalized();
 };
@@ -106,6 +107,9 @@ public:
 	MOVIESCENETRACKS_API FTransform GetRootMotionOffset(UObject* InObject) const;
 	MOVIESCENETRACKS_API void UpdateRootMotionOffset(UObject* InObject);
 
+	MOVIESCENETRACKS_API TOptional<FTransform> GetInitialActorTransform(UObject* InObject) const;
+	MOVIESCENETRACKS_API TOptional<FQuat> GetInverseMeshToActorRotation(UObject* InObject) const;
+
 private:
 
 	virtual void OnSchedulePersistentTasks(UE::MovieScene::IEntitySystemScheduler* TaskScheduler) override;
