Overriding existing parent transform when gathering subsequence transforms, instead of transforming it by the parent.

Parent transforms will already be applied in the PostTask, so transforming here would apply it once extra for each level in the hierarchy.

#rnx

diff --git a/Engine/Source/Runtime/MovieSceneTracks/Private/Systems/MovieSceneTransformOriginSystem.cpp b/Engine/Source/Runtime/MovieSceneTracks/Private/Systems/MovieSceneTransformOriginSystem.cpp
--- a/Engine/Source/Runtime/MovieSceneTracks/Private/Systems/MovieSceneTransformOriginSystem.cpp
+++ b/Engine/Source/Runtime/MovieSceneTracks/Private/Systems/MovieSceneTransformOriginSystem.cpp
@@ -95,11 +95,10 @@ struct FGatherTransformOriginsFromSubscenes
 				TransformOriginsByInstanceID->Insert(SubInstanceHandle.InstanceID, TransformOrigin);
 			}
 			// Transforms for the root sequence are controlled on the level sequence actor, and are gathered in FGatherTransformOrigin before this task runs
-			// If the root sequence transform origin exists, transform our offset with it, to apply additively.
+			// If the root sequence transform origin exists, it will be applied later when we accumulate parent transforms, so override it here.
 			else
 			{
-				// @todo: transform order is non-deterministic here - should we try and handle that??
-				(*TransformOriginsByInstanceID)[SubInstanceHandle.InstanceID] *= TransformOrigin;
+				(*TransformOriginsByInstanceID)[SubInstanceHandle.InstanceID] = TransformOrigin;
 			}
 		}
 	}
