Sequencer: Reset component velocity in pre animated state

This fixes a bug where FIntermediate3DTransform::ApplyTransformTo sets component velocity when animating a transform but never resets it when done. This can lead to unexpected side effects. If you animate the transform of an object, close Sequencer when the component velocity has been set (ie. nonzero), it will retain that component velocity after closing Sequencer.

Unfortunately, these Component Velocity values have already been saved and the only way to fix these would be to open and close the sequence that did animate those objects after this fix.

Resetting the component velocity is enabled by default but can be disabled to revert to the previous behavior with the cvar MovieScene.ResetComponentVelocity 0

diff --git a/Engine/Source/Runtime/MovieSceneTracks/Private/PreAnimatedState/MovieScenePreAnimatedComponentTransformStorage.cpp b/Engine/Source/Runtime/MovieSceneTracks/Private/PreAnimatedState/MovieScenePreAnimatedComponentTransformStorage.cpp
--- a/Engine/Source/Runtime/MovieSceneTracks/Private/PreAnimatedState/MovieScenePreAnimatedComponentTransformStorage.cpp
+++ b/Engine/Source/Runtime/MovieSceneTracks/Private/PreAnimatedState/MovieScenePreAnimatedComponentTransformStorage.cpp
@@ -12,6 +12,12 @@ namespace MovieScene
 
 TAutoRegisterPreAnimatedStorageID<FPreAnimatedComponentTransformStorage> FPreAnimatedComponentTransformStorage::StorageID;
 
+int32 GResetComponentVelocity = 1;
+static FAutoConsoleVariableRef CVarResetComponentVelocity(
+	TEXT("MovieScene.ResetComponentVelocity"),
+	GResetComponentVelocity,
+	TEXT("If 1 set component velocity to 0 on restore state, otherwise don't do anything (previous behavior)."),
+	ECVF_Default);
 
 struct FTemporaryMobilityScope
 {
@@ -36,6 +42,12 @@ struct FTemporaryMobilityScope
 		{
 			SceneComponent->SetMobility(PreviousMobility);
 		}
+
+		// Forcibly reset the component velocity which is set in FIntermediate3DTransform::ApplyTransformTo
+		if (CVarResetComponentVelocity->GetInt() > 0)
+		{
+			SceneComponent->ComponentVelocity = FVector(0.f);
+		}
 	}
 
 	USceneComponent* SceneComponent;
