Fix Subsections without transform origins having root override applied twice

diff --git a/Engine/Source/Runtime/MovieSceneTracks/Private/Systems/MovieSceneTransformOriginSystem.cpp b/Engine/Source/Runtime/MovieSceneTracks/Private/Systems/MovieSceneTransformOriginSystem.cpp
--- a/Engine/Source/Runtime/MovieSceneTracks/Private/Systems/MovieSceneTransformOriginSystem.cpp
+++ b/Engine/Source/Runtime/MovieSceneTracks/Private/Systems/MovieSceneTransformOriginSystem.cpp
@@ -41,19 +41,23 @@ struct FGatherTransformOrigin
 			}
 
 			const FSequenceInstance& Instance = SparseInstances[Index];
-			TSharedRef<const FSharedPlaybackState> SharedPlaybackState = Instance.GetSharedPlaybackState();
 
-			const IMovieScenePlaybackClient*  Client       = SharedPlaybackState->FindCapability<IMovieScenePlaybackClient>();
-			const UObject*                    InstanceData = Client ? Client->GetInstanceData() : nullptr;
-			const IMovieSceneTransformOrigin* RawInterface = Cast<const IMovieSceneTransformOrigin>(InstanceData);
-
-			const bool bHasInterface = RawInterface || (InstanceData && InstanceData->GetClass()->ImplementsInterface(UMovieSceneTransformOrigin::StaticClass()));
-			if (bHasInterface)
+			if(Instance.IsRootSequence())
 			{
-				// Retrieve the current origin
-				FTransform TransformOrigin = RawInterface ? RawInterface->GetTransformOrigin() : IMovieSceneTransformOrigin::Execute_BP_GetTransformOrigin(InstanceData);
+				TSharedRef<const FSharedPlaybackState> SharedPlaybackState = Instance.GetSharedPlaybackState();
+
+				const IMovieScenePlaybackClient*  Client       = SharedPlaybackState->FindCapability<IMovieScenePlaybackClient>();
+				const UObject*                    InstanceData = Client ? Client->GetInstanceData() : nullptr;
+				const IMovieSceneTransformOrigin* RawInterface = Cast<const IMovieSceneTransformOrigin>(InstanceData);
+
+				const bool bHasInterface = RawInterface || (InstanceData && InstanceData->GetClass()->ImplementsInterface(UMovieSceneTransformOrigin::StaticClass()));
+				if (bHasInterface)
+				{
+					// Retrieve the current origin
+					FTransform TransformOrigin = RawInterface ? RawInterface->GetTransformOrigin() : IMovieSceneTransformOrigin::Execute_BP_GetTransformOrigin(InstanceData);
 
-				TransformOriginsByInstanceID->Insert(Index, TransformOrigin);
+					TransformOriginsByInstanceID->Insert(Index, TransformOrigin);
+				}
 			}
 		}
 	}
@@ -89,21 +93,13 @@ struct FGatherTransformOriginsFromSubscenes
 			const FRotator   Rotation(RotationY ? RotationY[Index] : 0.f, RotationZ ? RotationZ[Index] : 0.f, RotationX ? RotationX[Index] : 0.f);
 			const FTransform TransformOrigin(Rotation, Translation);
 			
-			// Set the entry for this transform origin if no prior transform is set
-			if (!TransformOriginsByInstanceID->IsValidIndex(SubInstanceHandle.InstanceID))
-			{
-				TransformOriginsByInstanceID->Insert(SubInstanceHandle.InstanceID, TransformOrigin);
-			}
-			// Transforms for the root sequence are controlled on the level sequence actor, and are gathered in FGatherTransformOrigin before this task runs
-			// If the root sequence transform origin exists, it will be applied later when we accumulate parent transforms, so override it here.
-			else
-			{
-				(*TransformOriginsByInstanceID)[SubInstanceHandle.InstanceID] = TransformOrigin;
-			}
+			// Set the entry for this transform origin.
+			TransformOriginsByInstanceID->Insert(SubInstanceHandle.InstanceID, TransformOrigin);
 		}
 	}
-
+	
 	// After all the base transforms for subsequences are gathered, multiply in their parent transform
+	// This is run even if ForEachAllocation is not run, which will be the case if there are no transform overrides on a given subscene, ensuring it still gets its parent's origin.
 	void PostTask()
 	{
 		// InstanceHandleToParentHandle mapping is sorted parent first, so each parent's transform will be valid by the time its child uses it.
