Fix for the Gather Subsequence post task not being run in all circumstances. We rely on the post task to apply the parent transform to subsequences, even if those subsequences don't have transform overrides, and wouldn't run the main task.

diff --git a/Engine/Source/Runtime/MovieSceneTracks/Private/Systems/MovieSceneTransformOriginSystem.cpp b/Engine/Source/Runtime/MovieSceneTracks/Private/Systems/MovieSceneTransformOriginSystem.cpp
--- a/Engine/Source/Runtime/MovieSceneTracks/Private/Systems/MovieSceneTransformOriginSystem.cpp
+++ b/Engine/Source/Runtime/MovieSceneTracks/Private/Systems/MovieSceneTransformOriginSystem.cpp
@@ -356,6 +356,7 @@ void UMovieSceneTransformOriginSystem::OnSchedulePersistentTasks(UE::MovieScene:
 	.FilterAll({BuiltInComponents->Tags.SubInstance})
 	.FilterNone({BuiltInComponents->Tags.ImportedEntity})
 	.CombineFilter(LocationAndRotationFilterResults)
+	.SetParams(FTaskParams(TEXT("Gather Transform Origins From Subscenes")).ForcePrePostTask())
 	.Fork_PerAllocation<FGatherTransformOriginsFromSubscenes>(&Linker->EntityManager, TaskScheduler, &InstanceHandleToParentHandle, &TransformOriginsByInstanceID, InstanceRegistry);
 
 	FTaskID AssignTask = FEntityTaskBuilder()
