Sequencer: Automatically add a new section when you create a new cvar track

Cvar section should default to infinite, like fade and slomo track sections

diff --git a/Engine/Source/Editor/MovieSceneTools/Private/TrackEditors/CVarTrackEditor.cpp b/Engine/Source/Editor/MovieSceneTools/Private/TrackEditors/CVarTrackEditor.cpp
--- a/Engine/Source/Editor/MovieSceneTools/Private/TrackEditors/CVarTrackEditor.cpp
+++ b/Engine/Source/Editor/MovieSceneTools/Private/TrackEditors/CVarTrackEditor.cpp
@@ -104,6 +104,12 @@ void FCVarTrackEditor::HandleAddCVarTrackMenuEntryExecute()
 
 	NewTrack->SetDisplayName(LOCTEXT("CVarTrackName", "Console Variable"));
 
+	UMovieSceneSection* NewSection = NewTrack->CreateNewSection();
+	check(NewSection);
+
+	NewTrack->Modify();
+	NewTrack->AddSection(*NewSection);
+
 	if (GetSequencer().IsValid())
 	{
 		GetSequencer()->OnAddTrack(NewTrack, FGuid());
diff --git a/Engine/Source/Runtime/MovieSceneTracks/Private/Sections/MovieSceneCVarSection.cpp b/Engine/Source/Runtime/MovieSceneTracks/Private/Sections/MovieSceneCVarSection.cpp
--- a/Engine/Source/Runtime/MovieSceneTracks/Private/Sections/MovieSceneCVarSection.cpp
+++ b/Engine/Source/Runtime/MovieSceneTracks/Private/Sections/MovieSceneCVarSection.cpp
@@ -69,9 +69,13 @@ FString FMovieSceneCVarOverrides::GetString() const
 UMovieSceneCVarSection::UMovieSceneCVarSection(const FObjectInitializer& Init)
 	: Super(Init)
 {
+	SetRange(TRange<FFrameNumber>::All());
+
 	EvalOptions.EnableAndSetCompletionMode(EMovieSceneCompletionMode::ProjectDefault);
 
 	SetBlendType(EMovieSceneBlendType::Absolute);
+
+	bSupportsInfiniteRange = true;
 }
 
 void UMovieSceneCVarSection::ImportEntityImpl(UMovieSceneEntitySystemLinker* EntityLinker, const FEntityImportParams& Params, FImportedEntity* OutImportedEntity)
