Accout for transform origin overrides when drawing transform trajectories.

Adds Instance Data and subsequence origin override interrogation support.

diff --git a/Engine/Source/Editor/Sequencer/Private/SequencerEdMode.cpp b/Engine/Source/Editor/Sequencer/Private/SequencerEdMode.cpp
--- a/Engine/Source/Editor/Sequencer/Private/SequencerEdMode.cpp
+++ b/Engine/Source/Editor/Sequencer/Private/SequencerEdMode.cpp
@@ -48,6 +48,7 @@
 #include "Editor/UnrealEdEngine.h"
 #include "TextureResource.h"
 #include "EditorModeManager.h"
+#include "IMovieScenePlaybackClient.h"
 #include "Components/PrimitiveComponent.h"
 #include "SequencerSelectabilityTool.h"
 
@@ -107,10 +108,27 @@ struct FTrackTransforms
 
 		FInterrogationChannel Channel = Interrogator.ImportTransformHierarchy(SceneComponent, Sequencer, Sequencer->GetFocusedTemplateID());
 
-		if (TrajectoryKeys.Num() > 0)
+		TArray<FMovieSceneSequenceID> SubsequenceHierarchy = Sequencer->GetSubSequenceHierarchy();
+		
+		const FMovieSceneRootEvaluationTemplateInstance& RootInstance = Sequencer->GetEvaluationTemplate();
+		const FMovieSceneSequenceHierarchy* Hierarchy = RootInstance.GetCompiledDataManager()->FindHierarchy(RootInstance.GetCompiledDataID());
+
+		if (Hierarchy)
 		{
-			Times.Add(TrajectoryKeys[0].Time);
-			Interrogator.AddInterrogation(TrajectoryKeys[0].Time);
+			Interrogator.SetHierarchy(const_cast<FMovieSceneSequenceHierarchy*>(Hierarchy));
+
+			FMovieSceneSequenceID OwningSequenceID = MovieSceneSequenceID::Root;
+			for (const FMovieSceneSequenceID SubSequenceID : SubsequenceHierarchy)
+			{
+				if (const UMovieSceneSubSection* SubSection = Sequencer->FindSubSection(SubSequenceID))
+				{
+					if (UMovieSceneTrack* SubTrack = Cast<UMovieSceneTrack>(SubSection->GetOuter()))
+					{
+						Interrogator.ImportTrack(SubTrack, FInterrogationChannel::FromIndex(SubSection->GetRowIndex()), OwningSequenceID);
+					}
+				}
+				OwningSequenceID = SubSequenceID;
+			}
 		}
 
 		const int32 NumTrajectoryKeys = TrajectoryKeys.Num();
@@ -145,8 +163,23 @@ struct FTrackTransforms
 			}
 		}
 
+		TArray<FTransform> OriginTransforms;
+
+		// Get the Instance Data override from the world context object.
+		const IMovieScenePlaybackClient*  Client = Sequencer->GetPlaybackClient();
+		const UObject* InstanceData = Client ? Client->GetInstanceData() : nullptr;
+
 		Interrogator.Update();
 		Interrogator.QueryWorldSpaceTransforms(Channel, Transforms);
+		Interrogator.QueryTransformOrigins(OriginTransforms, SubsequenceHierarchy, InstanceData);
+
+		for (int32 Index = 0; Index < Times.Num(); ++Index)
+		{
+			if (Transforms.IsValidIndex(Index))
+			{
+				Transforms[Index] *= OriginTransforms[Index];
+			}
+		}
 
 		check(Transforms.Num() == Times.Num());
 		Interrogator.Reset();
diff --git a/Engine/Source/Runtime/MovieScene/Private/Sections/MovieSceneSubSection.cpp b/Engine/Source/Runtime/MovieScene/Private/Sections/MovieSceneSubSection.cpp
--- a/Engine/Source/Runtime/MovieScene/Private/Sections/MovieSceneSubSection.cpp
+++ b/Engine/Source/Runtime/MovieScene/Private/Sections/MovieSceneSubSection.cpp
@@ -19,6 +19,7 @@
 #include "EntitySystem/MovieSceneInstanceRegistry.h"
 #include "EntitySystem/IMovieSceneEntityProvider.h"
 #include "Channels/MovieSceneChannelProxy.h"
+#include "EntitySystem/Interrogation/MovieSceneInterrogationLinker.h"
 #include "Tracks/MovieSceneTimeWarpTrack.h"
 #include "Sections/MovieSceneSectionTimingParameters.h"
 #include "Variants/MovieSceneTimeWarpGetter.h"
@@ -739,8 +740,26 @@ void UMovieSceneSubSection::BuildDefaultSubSectionComponents(UMovieSceneEntitySy
 
 	const bool bHasEasing = Easing.GetEaseInDuration() > 0 || Easing.GetEaseOutDuration() > 0;
 
-	const FSubSequencePath PathToRoot = EntityLinker->GetInstanceRegistry()->GetInstance(Params.Sequence.InstanceHandle).GetSubSequencePath();
-	FMovieSceneSequenceID ResolvedSequenceID = PathToRoot.ResolveChildSequenceID(this->GetSequenceID());
+
+	// When interrogating, the sequence ID is known, and set directly, so default to using this value.
+	FMovieSceneSequenceID ResolvedSequenceID = MovieSceneSequenceID::Invalid;
+	
+	const TEntitySystemLinkerExtensionID<IInterrogationExtension> ID = IInterrogationExtension::GetExtensionID();
+	
+	const IInterrogationExtension* Interrogation = EntityLinker->FindExtension<IInterrogationExtension>(ID);
+
+	if (Interrogation)
+	{
+		const FMovieSceneSequenceHierarchy* Hierarchy = Interrogation->GetHierarchy();
+		const FSubSequencePath PathToRoot = FSubSequencePath(Params.Sequence.SequenceID, Hierarchy);
+		ResolvedSequenceID = PathToRoot.ResolveChildSequenceID(this->GetSequenceID());
+	}
+	// During normal evaluation (i.e. not interrogating) the instance registry will have its instance populated, and the sequence ID can be resolved this way.
+	else if (EntityLinker->GetInstanceRegistry()->IsHandleValid(Params.Sequence.InstanceHandle))
+	{
+		const FSubSequencePath PathToRoot = EntityLinker->GetInstanceRegistry()->GetInstance(Params.Sequence.InstanceHandle).GetSubSequencePath();
+		ResolvedSequenceID = PathToRoot.ResolveChildSequenceID(this->GetSequenceID());
+	}
 
 	EMovieSceneTransformChannel Channels = OriginOverrideMask.GetChannels();
 
diff --git a/Engine/Source/Runtime/MovieScene/Public/EntitySystem/MovieSceneInterrogation.h b/Engine/Source/Runtime/MovieScene/Public/EntitySystem/MovieSceneInterrogation.h
--- a/Engine/Source/Runtime/MovieScene/Public/EntitySystem/MovieSceneInterrogation.h
+++ b/Engine/Source/Runtime/MovieScene/Public/EntitySystem/MovieSceneInterrogation.h
@@ -10,6 +10,8 @@
 #include "EntitySystem/MovieScenePropertyBinding.h"
 #include "EntitySystem/MovieSceneEntitySystemLinkerExtension.h"
 
+struct FMovieSceneSequenceHierarchy;
+
 namespace UE::MovieScene
 {
 
@@ -69,6 +71,8 @@ public:
 
 	virtual const FSparseInterrogationChannelInfo& GetSparseChannelInfo() const = 0;
 
+	MOVIESCENE_API virtual const FMovieSceneSequenceHierarchy* GetHierarchy() const = 0;
+
 protected:
 
 	virtual ~IInterrogationExtension() {}
diff --git a/Engine/Source/Runtime/MovieSceneTracks/Private/EntitySystem/Interrogation/MovieSceneInterrogationLinker.cpp b/Engine/Source/Runtime/MovieSceneTracks/Private/EntitySystem/Interrogation/MovieSceneInterrogationLinker.cpp
--- a/Engine/Source/Runtime/MovieSceneTracks/Private/EntitySystem/Interrogation/MovieSceneInterrogationLinker.cpp
+++ b/Engine/Source/Runtime/MovieSceneTracks/Private/EntitySystem/Interrogation/MovieSceneInterrogationLinker.cpp
@@ -17,6 +17,7 @@
 #include "Tracks/MovieScene3DTransformTrack.h"
 
 #include "GameFramework/Actor.h"
+#include "Tracks/IMovieSceneTransformOrigin.h"
 #include "UObject/Package.h"
 
 
@@ -612,6 +613,127 @@ void FInterrogationChannels::QueryLocalSpaceTransforms(UMovieSceneEntitySystemLi
 	}
 }
 
+void FInterrogationChannels::QueryTransformOrigins(UMovieSceneEntitySystemLinker* Linker, const FMovieSceneSequenceHierarchy* Hierarchy, TArray<FTransform>& OutTransformOrigins, TArray<FMovieSceneSequenceID>& FocusedSubsequenceHierarchy, const UObject* InstanceData) const
+{
+	const FBuiltInComponentTypes* Components = FBuiltInComponentTypes::Get();
+
+	TArray<TPair<FMovieSceneSequenceID, FMovieSceneSequenceID>> ChildToParentMapping;
+
+	// This array will be treated as an NxM, two-dimensional array that will hold the transform origins for all sequences in the focused hierarchy
+	// N is the number of interrogations, and M is the number of Sequences in the currently focused hierarchy in Sequencer
+	TArray<TOptional<FTransform>> TransformOrigins;
+	// Reserve room for every sequence in the focused hierarchy to have entries for every time being interrogated.
+	TransformOrigins.Init(TOptional<FTransform>(), Interrogations.Num() * FocusedSubsequenceHierarchy.Num());
+
+	const IMovieSceneTransformOrigin* RawInterface = Cast<const IMovieSceneTransformOrigin>(InstanceData);
+
+	// Read the world context actor's instance data.
+	const bool bHasInterface = RawInterface || (InstanceData && InstanceData->GetClass()->ImplementsInterface(UMovieSceneTransformOrigin::StaticClass()));
+	if (bHasInterface)
+	{
+		// Retrieve the current origin and use that for every sampled time for the root sequence transforms.
+		const FTransform TransformOrigin = RawInterface ? RawInterface->GetTransformOrigin() : IMovieSceneTransformOrigin::Execute_BP_GetTransformOrigin(InstanceData);
+
+		for (int32 OriginIndex = 0; OriginIndex < Interrogations.Num(); OriginIndex++)
+		{
+			TransformOrigins[OriginIndex] = TransformOrigin;
+		}
+	}
+	
+	// Read all transform origins from subsequences.
+	auto GatherTransformOrigins = [Components, Hierarchy, &ChildToParentMapping, &TransformOrigins, &FocusedSubsequenceHierarchy, &InstanceData, this]
+		(const FEntityAllocation* Allocation, const FMovieSceneSequenceID* SequenceIDs, const FInterrogationKey* InputKeys)
+	{
+		const int32 NumAllocations = Allocation->Num();
+		const int32 StrideSize = Interrogations.Num();
+
+		TOptionalComponentReader<double> DoubleComponents[6];
+		for (int32 Index = 0; Index < 6; ++Index)
+		{
+			DoubleComponents[Index] = Allocation->TryReadComponents(Components->DoubleResult[Index]);
+		}
+
+		for (int32 ComponentIndex = 0; ComponentIndex < NumAllocations; ++ComponentIndex)
+		{
+			const FMovieSceneSequenceID CurrentSequence = SequenceIDs[ComponentIndex];
+			const FInterrogationKey InterrogationKey = InputKeys[ComponentIndex];
+
+			if (const FMovieSceneSequenceHierarchyNode* HierarchyNode = Hierarchy->FindNode(CurrentSequence))
+			{
+				// If one of the sequences isn't in the focused hierarchy, ignore it. This would be a sibling subsequence on the same track, and blending is not supported.
+				// @todo support blending, and full subsequence hierarchies
+				if (FocusedSubsequenceHierarchy.IndexOfByKey(CurrentSequence) == INDEX_NONE || FocusedSubsequenceHierarchy.IndexOfByKey(HierarchyNode->ParentID) == INDEX_NONE)
+				{
+					continue;
+				}
+				ChildToParentMapping.AddUnique(TPair<FMovieSceneSequenceID, FMovieSceneSequenceID>(CurrentSequence, HierarchyNode->ParentID));
+			}
+		
+			FVector Translation = FVector::ZeroVector;
+			FRotator Rotation = FRotator::ZeroRotator;
+
+			if (DoubleComponents[0]) { Translation.X = DoubleComponents[0][ComponentIndex]; }
+			if (DoubleComponents[1]) { Translation.Y = DoubleComponents[1][ComponentIndex]; }
+			if (DoubleComponents[2]) { Translation.Z = DoubleComponents[2][ComponentIndex]; }
+			if (DoubleComponents[3]) { Rotation.Roll = DoubleComponents[3][ComponentIndex]; }
+			if (DoubleComponents[4]) { Rotation.Pitch = DoubleComponents[4][ComponentIndex]; }
+			if (DoubleComponents[5]) { Rotation.Yaw = DoubleComponents[5][ComponentIndex]; }
+
+			const int32 SequenceIndex = FocusedSubsequenceHierarchy.IndexOfByKey(CurrentSequence);
+			if (SequenceIndex == INDEX_NONE)
+			{
+				continue;
+			}
+			const int32 TransformIndex = InterrogationKey.InterrogationIndex + (SequenceIndex * StrideSize);
+			const FTransform TransformOrigin =  FTransform(Rotation, Translation);
+			
+			TransformOrigins[TransformIndex] = TransformOrigin;
+		}
+		
+	};
+	
+	FEntityTaskBuilder()
+	.Read(Components->SequenceID)
+	.Read(Components->Interrogation.InputKey)
+	.FilterAll({Components->Tags.SubInstance})
+	.FilterAll({Components->Tags.ImportedEntity})
+	.Iterate_PerAllocation(&Linker->EntityManager, GatherTransformOrigins);
+	
+	// Sort parent first
+	Algo::Sort(ChildToParentMapping, [](TPair<FMovieSceneSequenceID, FMovieSceneSequenceID> A, TPair<FMovieSceneSequenceID, FMovieSceneSequenceID> B){return A.Key == B.Value;});
+
+	// compose transform origins
+	for (const TPair<FMovieSceneSequenceID, FMovieSceneSequenceID>& ChildParentPair : ChildToParentMapping)
+	{
+		const int32 ParentIndexStart = FocusedSubsequenceHierarchy.IndexOfByKey(ChildParentPair.Value) * Interrogations.Num();
+		const int32 ChildIndexStart = FocusedSubsequenceHierarchy.IndexOfByKey(ChildParentPair.Key) * Interrogations.Num();
+
+		for (int32 InterrogationIndex = 0; InterrogationIndex < Interrogations.Num(); ++InterrogationIndex)
+		{
+			const int32 ParentIndex = ParentIndexStart + InterrogationIndex;
+			const int32 ChildIndex = ChildIndexStart + InterrogationIndex;
+
+			// If the focused view in the sequence editor goes beyond the range of the subsequence, there will be interrogations where the subsequence isn't evaluated.
+			if (TransformOrigins[ParentIndex].IsSet())
+			{
+				const FTransform ParentTransform = TransformOrigins[ParentIndex].GetValue();
+
+				const bool bIsChildSet = TransformOrigins[ChildIndex].IsSet();
+
+				TransformOrigins[ChildIndex] = bIsChildSet ? TransformOrigins[ChildIndex].GetValue() * ParentTransform : ParentTransform;
+			}
+		}
+	}
+
+	const int32 FocusedSequenceIndex = (FocusedSubsequenceHierarchy.Num() - 1) * Interrogations.Num();
+	
+	for (int32 OutTransformIndex = 0; OutTransformIndex < Interrogations.Num(); ++OutTransformIndex)
+	{
+		FTransform CurrentTransform = TransformOrigins[FocusedSequenceIndex + OutTransformIndex].IsSet() ? TransformOrigins[FocusedSequenceIndex + OutTransformIndex].GetValue() : FTransform::Identity;
+		OutTransformOrigins.Add(CurrentTransform);
+	}
+}
+
 EEntitySystemCategory FSystemInterrogator::GetInterrogationCategory()
 {
 	static EEntitySystemCategory Interrogation = UMovieSceneEntitySystem::RegisterCustomSystemCategory();
@@ -666,12 +788,12 @@ void FSystemInterrogator::Reset()
 	Linker->Reset();
 }
 
-void FSystemInterrogator::ImportTrack(UMovieSceneTrack* Track, FInterrogationChannel InChannel)
+void FSystemInterrogator::ImportTrack(UMovieSceneTrack* Track, FInterrogationChannel InChannel, FMovieSceneSequenceID SequenceID)
 {
-	ImportTrack(Track, Track->FindObjectBindingGuid(), InChannel);
+	ImportTrack(Track, Track->FindObjectBindingGuid(), InChannel, SequenceID);
 }
 
-void FSystemInterrogator::ImportTrack(UMovieSceneTrack* Track, const FGuid& ObjectBindingID, FInterrogationChannel InChannel)
+void FSystemInterrogator::ImportTrack(UMovieSceneTrack* Track, const FGuid& ObjectBindingID, FInterrogationChannel InChannel, FMovieSceneSequenceID SequenceID)
 {
 	if (Track->IsEvalDisabled())
 	{
@@ -691,7 +813,7 @@ void FSystemInterrogator::ImportTrack(UMovieSceneTrack* Track, const FGuid& Obje
 	UMovieSceneSequence* OwningSequence = Track->GetTypedOuter<UMovieSceneSequence>();
 	if (!ExtraMetaData.IsValidIndex(SharedMetaDataIndex))
 	{
-		ExtraMetaData.Insert(SharedMetaDataIndex, FExtraMetaData{ FInterrogationInstance { OwningSequence }, InChannel });
+		ExtraMetaData.Insert(SharedMetaDataIndex, FExtraMetaData{ FInterrogationInstance { OwningSequence }, InChannel, SequenceID });
 	}
 	else
 	{
@@ -780,6 +902,7 @@ void FSystemInterrogator::InterrogateEntity(int32 InterrogationIndex, const FMov
 		const FExtraMetaData& InterrogationMetaData = ExtraMetaData[SharedMetaDataIndex];
 		Params.InterrogationKey.Channel = InterrogationMetaData.InterrogationChannel;
 		Params.InterrogationInstance = InterrogationMetaData.InterrogationInstance;
+		Params.Sequence.SequenceID = InterrogationMetaData.InterrogationSequenceID;
 	}
 
 	FImportedEntity ImportedEntity;
diff --git a/Engine/Source/Runtime/MovieSceneTracks/Public/EntitySystem/Interrogation/MovieSceneInterrogationLinker.h b/Engine/Source/Runtime/MovieSceneTracks/Public/EntitySystem/Interrogation/MovieSceneInterrogationLinker.h
--- a/Engine/Source/Runtime/MovieSceneTracks/Public/EntitySystem/Interrogation/MovieSceneInterrogationLinker.h
+++ b/Engine/Source/Runtime/MovieSceneTracks/Public/EntitySystem/Interrogation/MovieSceneInterrogationLinker.h
@@ -14,6 +14,7 @@
 #include "EntitySystem/MovieSceneEntitySystemLinkerExtension.h"
 
 #include "EntitySystem/Interrogation/MovieSceneInterrogationExtension.h"
+#include "Evaluation/MovieSceneSequenceHierarchy.h"
 
 struct FGuid;
 struct FMovieSceneBinding;
@@ -242,6 +243,8 @@ public:
 	 */
 	MOVIESCENETRACKS_API void QueryWorldSpaceTransforms(UMovieSceneEntitySystemLinker* Linker, const TBitArray<>& ChannelsToQuery, TSparseArray<TArray<FTransform>>& OutTransformsByChannel) const;
 
+	MOVIESCENETRACKS_API void QueryTransformOrigins(UMovieSceneEntitySystemLinker* Linker, const FMovieSceneSequenceHierarchy* Hierarchy, TArray<FTransform>& OutTransformOrigins, TArray<FMovieSceneSequenceID>& SubsequenceHierarchy, const UObject* InstanceData) const;
+	
 private:
 
 	/** Query implementation functions */
@@ -305,8 +308,9 @@ public:
 	 *
 	 * @param Track      The track to import
 	 * @param InChannel  The channel to import this track onto. FInterrogationChannel::Default() can be used if this interrogator is only being used for a single output.
-	 */
-	MOVIESCENETRACKS_API void ImportTrack(UMovieSceneTrack* Track, FInterrogationChannel InChannel);
+	 * @param SequenceID The SequenceID of the subsequence if importing one.
+	*/
+	MOVIESCENETRACKS_API void ImportTrack(UMovieSceneTrack* Track, FInterrogationChannel InChannel, FMovieSceneSequenceID SequenceID = MovieSceneSequenceID::Invalid);
 
 
 	/**
@@ -317,8 +321,9 @@ public:
 	 * @param Track            The track to import
 	 * @param ObjectBindingID  The binding ID for the object binding that this track resides within
 	 * @param InChannel        The channel to import this track onto. FInterrogationChannel::Default() can be used if this interrogator is only being used for a single output.
+	 * @param SequenceID	   The SequenceID of the subsequence if importing one.
 	 */
-	MOVIESCENETRACKS_API void ImportTrack(UMovieSceneTrack* Track, const FGuid& ObjectBindingID, FInterrogationChannel InChannel);
+	MOVIESCENETRACKS_API void ImportTrack(UMovieSceneTrack* Track, const FGuid& ObjectBindingID, FInterrogationChannel InChannel, FMovieSceneSequenceID SequenceID = MovieSceneSequenceID::Invalid);
 
 
 	/**
@@ -564,6 +569,20 @@ public:
 		PropertyDefinition.Handler->RebuildOperational(PropertyDefinition, PropertyComposites, ValueEntityIDs, Linker, OutValues);
 	}
 
+
+	/**
+	 * Query the transform origin of a subsection based on the imported hiearchy
+	 *
+	 * @param OutTransformOrigins	The interrogated transforms for the currently focused sequence.
+	 * @param SubsequenceHierarchy	The currently focused sequence hierarchy
+	 * @param InstanceData			The object that contains the origin overrides for the world context object of the root sequence.
+	 * 
+	 */
+	void QueryTransformOrigins(TArray<FTransform>& OutTransformOrigins, TArray<FMovieSceneSequenceID>& SubsequenceHierarchy, const UObject* InstanceData) const
+	{
+		Channels.QueryTransformOrigins(Linker, Hierarchy, OutTransformOrigins, SubsequenceHierarchy, InstanceData);
+	}
+
 public:
 
 	const FSparseInterrogationChannelInfo& GetSparseChannelInfo() const override
@@ -571,6 +590,16 @@ public:
 		return Channels.GetSparseChannelInfo();
 	}
 
+	virtual const FMovieSceneSequenceHierarchy* GetHierarchy() const override
+	{
+		return Hierarchy;
+	}
+
+	void SetHierarchy(FMovieSceneSequenceHierarchy* InHierarchy)
+	{
+		Hierarchy = InHierarchy;
+	}
+
 private:
 
 	/**
@@ -609,11 +638,15 @@ protected:
 	/** Initial value cache */
 	TSharedPtr<FInitialValueCache> InitialValueCache;
 
+	/** SubSection Hiearchy, for interrogating subsequence data e.g. subsequence transform origins */
+	FMovieSceneSequenceHierarchy* Hierarchy;
+
 	/** Array of shared extra metadata on entities */
 	struct FExtraMetaData
 	{
 		UE::MovieScene::FInterrogationInstance InterrogationInstance;
 		UE::MovieScene::FInterrogationChannel InterrogationChannel;
+		FMovieSceneSequenceID InterrogationSequenceID;
 	};
 	TSparseArray<FExtraMetaData> ExtraMetaData;
 };
