PR #12139: Fix Audio repeating in sequence when sequence time dilation is set below 1 and audio is not set to looping

diff --git a/Engine/Source/Runtime/MovieSceneTracks/Private/Systems/MovieSceneAudioSystem.cpp b/Engine/Source/Runtime/MovieSceneTracks/Private/Systems/MovieSceneAudioSystem.cpp
--- a/Engine/Source/Runtime/MovieSceneTracks/Private/Systems/MovieSceneAudioSystem.cpp
+++ b/Engine/Source/Runtime/MovieSceneTracks/Private/Systems/MovieSceneAudioSystem.cpp
@@ -51,6 +51,13 @@ FAutoConsoleVariableRef CVarPlayAudioWhenPlaybackJumps(
 	TEXT("Play audio when playback jumps.\n"),
 	ECVF_Default);
 
+static bool bUseTimeDilationToAdjustPlayDurationCVar = true;
+FAutoConsoleVariableRef CVarUseTimeDilationToAdjustPlayDuration(
+	TEXT("Sequencer.Audio.UseTimeDilationToAdjustPlayDuration"),
+	bUseTimeDilationToAdjustPlayDurationCVar,
+	TEXT("Use the effective time dilation to scale the current time of audio.\n"),
+	ECVF_Default);
+
 namespace UE::MovieScene
 {
 
@@ -507,20 +514,59 @@ private:
 		FFrameNumber SectionStartFrame = (AudioSection.HasStartFrame() ? AudioSection.GetInclusiveStartFrame() : 0);
 		float SectionStartTimeSeconds = SectionStartFrame / AudioSection.GetTypedOuter<UMovieScene>()->GetTickResolution();
 
+		float InverseTimeDilation = 1.0f;
+		if (bUseTimeDilationToAdjustPlayDurationCVar)
+		{
+			// Use time dilation to correct the duration so that the sound stops at the correct time.
+			// Without this adjustment, time dialations < 1.0 will incorrectly attempt to play beyond the end of the section.
+			InverseTimeDilation = (1.0f / PlaybackContext->GetWorld()->GetWorldSettings()->GetEffectiveTimeDilation());
+			SectionStartTimeSeconds *= InverseTimeDilation;
+		}
+
 		const FFrameNumber AudioStartOffset = AudioSection.GetStartOffset();
 		USoundBase* Sound = AudioSection.GetSound();
 
-		float AudioTime = (Context.GetTime() / Context.GetFrameRate()) 
+		float AudioTime = (Context.GetTime() / Context.GetFrameRate()) * InverseTimeDilation
 			- SectionStartTimeSeconds 
 			+ (float)Context.GetFrameRate().AsSeconds(AudioStartOffset);
+
 		if (AudioTime >= 0.f && Sound)
 		{
+			if (bUseTimeDilationToAdjustPlayDurationCVar && PlaybackContext->GetWorld())
+			{
+				// Keep track of initial Audio and Context times when sound started playing.
+				if (!(EvaluationData.LastAudioTime.IsSet() && EvaluationData.LastContextTime.IsSet()))
+				{
+					// Store current audio time and current context time
+					EvaluationData.LastAudioTime = AudioTime;
+					EvaluationData.LastContextTime = Context.GetTime() / Context.GetFrameRate();
+				}
+				else
+				{
+					// Get previous AudioTime
+					AudioTime = (EvaluationData.LastAudioTime.GetValue() < AudioTime) ? EvaluationData.LastAudioTime.GetValue() : AudioTime;
+					float CurrContextTime = (Context.GetTime() / Context.GetFrameRate());
+					float PrevContextTime = EvaluationData.LastContextTime.GetValue();
+					PrevContextTime = (PrevContextTime < CurrContextTime) ? PrevContextTime : CurrContextTime;
+
+					// Get Time Delta between previous time in sequencer context and current time, not taking into account dilation.
+					// Add to previous frame's audio time
+					AudioTime += (CurrContextTime - PrevContextTime) * (InverseTimeDilation);
+					EvaluationData.LastAudioTime = AudioTime;
+					EvaluationData.LastContextTime = CurrContextTime;
+				}
+			}
+
 			const float Duration = MovieSceneHelpers::GetSoundDuration(Sound);
 
 			if (!AudioSection.GetLooping() && AudioTime > Duration && Duration != 0.f)
 			{
-				UE_LOG(LogMovieScene, Verbose, TEXT("Audio Component reached end of playback. Component: %s Sound: %s"), *AudioComponent.GetName(), *GetNameSafe(AudioComponent.Sound));
-				AudioComponent.Stop();
+				if (AudioComponent.IsPlaying())
+				{
+					UE_LOG(LogMovieScene, Verbose, TEXT("Audio Component reached end of playback. Component: %s Sound: %s"), *AudioComponent.GetName(), *GetNameSafe(AudioComponent.Sound));
+					AudioComponent.Stop();
+				}
+				
 				return;
 			}
 
@@ -930,6 +976,10 @@ UMovieSceneAudioSystem::FAudioComponentEvaluationData* UMovieSceneAudioSystem::G
 				FAudioComponentEvaluationData MovedData(Pair.Value);
 				Map->Remove(Pair.Key);
 				MovedData.PartialDesyncComputation.Reset();
+
+				MovedData.LastAudioTime.Reset();
+				MovedData.LastContextTime.Reset();
+
 				return &Map->Add(DataKey, MovedData);
 			}
 		}
diff --git a/Engine/Source/Runtime/MovieSceneTracks/Private/Systems/MovieSceneAudioSystem.h b/Engine/Source/Runtime/MovieSceneTracks/Private/Systems/MovieSceneAudioSystem.h
--- a/Engine/Source/Runtime/MovieSceneTracks/Private/Systems/MovieSceneAudioSystem.h
+++ b/Engine/Source/Runtime/MovieSceneTracks/Private/Systems/MovieSceneAudioSystem.h
@@ -51,6 +51,11 @@ namespace UE::MovieScene
 		 */
 		TOptional<float> PartialDesyncComputation;
 
+		/** Previous audio time taking into account any time dilation */
+		TOptional<float> LastAudioTime;
+		/** The context time from the previous evaluation pass */
+		TOptional<float> LastContextTime;
+
 		/** Flag to keep track of audio components evaluated on a given frame */
 		bool bEvaluatedThisFrame;
 	};
