Fix static analysis null check warning.

#rnx

diff --git a/Engine/Source/Runtime/MovieSceneTracks/Private/Systems/MovieSceneAudioSystem.cpp b/Engine/Source/Runtime/MovieSceneTracks/Private/Systems/MovieSceneAudioSystem.cpp
--- a/Engine/Source/Runtime/MovieSceneTracks/Private/Systems/MovieSceneAudioSystem.cpp
+++ b/Engine/Source/Runtime/MovieSceneTracks/Private/Systems/MovieSceneAudioSystem.cpp
@@ -479,6 +479,7 @@ private:
 
 		ensureMsgf(EvaluationData.AudioComponent.IsValid(), TEXT("Trying to evaluate audio track on an invalid audio component"));
 		UAudioComponent& AudioComponent = *EvaluationData.AudioComponent.Get();
+		UWorld* World = PlaybackContext ? PlaybackContext->GetWorld() : nullptr;
 
 		AudioSystem->PreAnimatedStorage->CachePreAnimatedValue(
 				FCachePreAnimatedValueParams(), &AudioComponent,
@@ -515,11 +516,13 @@ private:
 		float SectionStartTimeSeconds = SectionStartFrame / AudioSection.GetTypedOuter<UMovieScene>()->GetTickResolution();
 
 		float InverseTimeDilation = 1.0f;
-		if (bUseTimeDilationToAdjustPlayDurationCVar)
+		const bool bUseTimeDilationToAdjustment = (bUseTimeDilationToAdjustPlayDurationCVar && World);
+
+		if (bUseTimeDilationToAdjustment)
 		{
 			// Use time dilation to correct the duration so that the sound stops at the correct time.
 			// Without this adjustment, time dialations < 1.0 will incorrectly attempt to play beyond the end of the section.
-			InverseTimeDilation = (1.0f / PlaybackContext->GetWorld()->GetWorldSettings()->GetEffectiveTimeDilation());
+			InverseTimeDilation = (1.0f / World->GetWorldSettings()->GetEffectiveTimeDilation());
 			SectionStartTimeSeconds *= InverseTimeDilation;
 		}
 
@@ -532,7 +535,7 @@ private:
 
 		if (AudioTime >= 0.f && Sound)
 		{
-			if (bUseTimeDilationToAdjustPlayDurationCVar && PlaybackContext->GetWorld())
+			if (bUseTimeDilationToAdjustment)
 			{
 				// Keep track of initial Audio and Context times when sound started playing.
 				if (!(EvaluationData.LastAudioTime.IsSet() && EvaluationData.LastContextTime.IsSet()))
@@ -580,8 +583,6 @@ private:
 		const bool bSoundNeedsStateChange =  AudioComponent.Sound != Sound;
 		bool bSoundNeedsTimeSync = false;
 
-		UWorld* World = PlaybackContext ? PlaybackContext->GetWorld() : nullptr;
-
 		// Sync only if there is no time dilation because otherwise the system will constantly resync because audio 
 		// playback is not dilated and will never match the expected playback time.
 		const bool bDoTimeSync = 
