Fix static analysis null check warning.

#rnx

diff --git a/Engine/Source/Runtime/MovieSceneTracks/Private/Systems/MovieSceneAudioSystem.cpp b/Engine/Source/Runtime/MovieSceneTracks/Private/Systems/MovieSceneAudioSystem.cpp
--- a/Engine/Source/Runtime/MovieSceneTracks/Private/Systems/MovieSceneAudioSystem.cpp
+++ b/Engine/Source/Runtime/MovieSceneTracks/Private/Systems/MovieSceneAudioSystem.cpp
@@ -479,7 +479,9 @@ private:
 
 		ensureMsgf(EvaluationData.AudioComponent.IsValid(), TEXT("Trying to evaluate audio track on an invalid audio component"));
 		UAudioComponent& AudioComponent = *EvaluationData.AudioComponent.Get();
+		
 		UWorld* World = PlaybackContext ? PlaybackContext->GetWorld() : nullptr;
+		AWorldSettings* WorldSettings = World ? World->GetWorldSettings() : nullptr;
 
 		AudioSystem->PreAnimatedStorage->CachePreAnimatedValue(
 				FCachePreAnimatedValueParams(), &AudioComponent,
@@ -516,14 +518,18 @@ private:
 		float SectionStartTimeSeconds = SectionStartFrame / AudioSection.GetTypedOuter<UMovieScene>()->GetTickResolution();
 
 		float InverseTimeDilation = 1.0f;
-		const bool bUseTimeDilationToAdjustment = (bUseTimeDilationToAdjustPlayDurationCVar && World);
+		const bool bUseTimeDilationToAdjustment = (bUseTimeDilationToAdjustPlayDurationCVar && WorldSettings);
 
 		if (bUseTimeDilationToAdjustment)
 		{
 			// Use time dilation to correct the duration so that the sound stops at the correct time.
 			// Without this adjustment, time dialations < 1.0 will incorrectly attempt to play beyond the end of the section.
-			InverseTimeDilation = (1.0f / World->GetWorldSettings()->GetEffectiveTimeDilation());
-			SectionStartTimeSeconds *= InverseTimeDilation;
+			const float EffectiveTimeDilation = WorldSettings->GetEffectiveTimeDilation();
+			if (!FMath::IsNearlyEqual(EffectiveTimeDilation, 0.f))
+			{
+				InverseTimeDilation = (1.0f / EffectiveTimeDilation);
+				SectionStartTimeSeconds *= InverseTimeDilation;
+			}
 		}
 
 		const FFrameNumber AudioStartOffset = AudioSection.GetStartOffset();
@@ -586,8 +592,8 @@ private:
 		// Sync only if there is no time dilation because otherwise the system will constantly resync because audio 
 		// playback is not dilated and will never match the expected playback time.
 		const bool bDoTimeSync = 
-			World && World->GetWorldSettings() &&
-			(FMath::IsNearlyEqual(World->GetWorldSettings()->GetEffectiveTimeDilation(), 1.f) ||
+			World && WorldSettings &&
+			(FMath::IsNearlyEqual(WorldSettings->GetEffectiveTimeDilation(), 1.f) ||
 			 !bIgnoreAudioSyncDuringWorldTimeDilationCVar);
 
 		if (bDoTimeSync)
