Sequencer: Animation Track 'root motion' now works with 'Evaluate Nearest Section'.

Previously, using evaluate nearest section with anim sequence tracks would only grab the previous pose from the section, but not apply the root motion properly. This fixes that.

diff --git a/Engine/Source/Runtime/MovieSceneTracks/Private/Tracks/MovieSceneSkeletalAnimationTrack.cpp b/Engine/Source/Runtime/MovieSceneTracks/Private/Tracks/MovieSceneSkeletalAnimationTrack.cpp
--- a/Engine/Source/Runtime/MovieSceneTracks/Private/Tracks/MovieSceneSkeletalAnimationTrack.cpp
+++ b/Engine/Source/Runtime/MovieSceneTracks/Private/Tracks/MovieSceneSkeletalAnimationTrack.cpp
@@ -495,51 +495,55 @@ TOptional<FTransform>  UMovieSceneSkeletalAnimationTrack::GetRootMotion(FFrameTi
 		return Transform;
 	}
 	SortSections();
-	TArray< UMovieSceneSkeletalAnimationSection*> SectionsAtCurrentTime;
 	TArray<FTransform> CurrentTransforms;
 	TArray<float> CurrentWeights;
 	TArray<FTransform> CurrentAdditiveTransforms;
 	TArray<float> CurrentAdditiveWeights;
 	FTransform CurrentTransform = FTransform::Identity;
-	for (UMovieSceneSection* Section : AnimationSections)
+
+	// Use evaluation field to iterate, as it will take into account 'Evaluate Nearest Section'.
+	for (const FMovieSceneTrackEvaluationFieldEntry& FieldEntry : GetEvaluationField().Entries)
 	{
-		UMovieSceneSkeletalAnimationSection* AnimSection = Cast<UMovieSceneSkeletalAnimationSection>(Section);
-		if (AnimSection)
+		if (FieldEntry.Range.Contains(CurrentTime.FrameNumber))
 		{
-			if (Section && AnimSection->Params.Animation && Section->GetRange().Contains(CurrentTime.FrameNumber))
+			if (UMovieSceneSkeletalAnimationSection* AnimSection = Cast<UMovieSceneSkeletalAnimationSection>(FieldEntry.Section))
 			{
-				UAnimSequenceBase* ValidAnimSequence = AnimSection->Params.Animation;
-				FMemMark Mark(FMemStack::Get());
-				FCompactPose OutPose;
-				TArray<FBoneIndexType> RequiredBoneIndexArray;
-				RequiredBoneIndexArray.AddUninitialized(ValidAnimSequence->GetSkeleton()->GetReferenceSkeleton().GetNum());
-				for (int32 BoneIndex = 0; BoneIndex < RequiredBoneIndexArray.Num(); ++BoneIndex)
-				{
-					RequiredBoneIndexArray[BoneIndex] = BoneIndex;
-				}
-
-				FBoneContainer BoneContainer(RequiredBoneIndexArray, UE::Anim::FCurveFilterSettings(UE::Anim::ECurveFilterMode::None), *ValidAnimSequence->GetSkeleton());
-				OutPose.ResetToRefPose(BoneContainer);
-				FBlendedCurve OutCurve;
-				OutCurve.InitFrom(BoneContainer);
-				UE::Anim::FStackAttributeContainer TempAttributes;
-				FAnimationPoseData AnimationPoseData(OutPose, OutCurve, TempAttributes);
-				UMovieSceneSkeletalAnimationSection::FRootMotionTransformParam Param;
-				Param.FrameRate = TickResolution;
-				Param.CurrentTime = CurrentTime.FrameNumber;
-
-				if (AnimSection->GetRootMotionTransform(AnimationPoseData, Param))
+				FFrameNumber EvaluationFrame = FieldEntry.ForcedTime == TNumericLimits<int32>::Lowest() ? CurrentTime.FrameNumber : FieldEntry.ForcedTime;
+				if (AnimSection->Params.Animation)
 				{
-					if (!Param.bOutIsAdditive)
+					UAnimSequenceBase* ValidAnimSequence = AnimSection->Params.Animation;
+					FMemMark Mark(FMemStack::Get());
+					FCompactPose OutPose;
+					TArray<FBoneIndexType> RequiredBoneIndexArray;
+					RequiredBoneIndexArray.AddUninitialized(ValidAnimSequence->GetSkeleton()->GetReferenceSkeleton().GetNum());
+					for (int32 BoneIndex = 0; BoneIndex < RequiredBoneIndexArray.Num(); ++BoneIndex)
 					{
-						CurrentTransform = Param.OutTransform * AnimSection->PreviousTransform;
-						CurrentTransforms.Add(CurrentTransform);
-						CurrentWeights.Add(Param.OutWeight);
+						RequiredBoneIndexArray[BoneIndex] = BoneIndex;
 					}
-					else
+
+					FBoneContainer BoneContainer(RequiredBoneIndexArray, UE::Anim::FCurveFilterSettings(UE::Anim::ECurveFilterMode::None), *ValidAnimSequence->GetSkeleton());
+					OutPose.ResetToRefPose(BoneContainer);
+					FBlendedCurve OutCurve;
+					OutCurve.InitFrom(BoneContainer);
+					UE::Anim::FStackAttributeContainer TempAttributes;
+					FAnimationPoseData AnimationPoseData(OutPose, OutCurve, TempAttributes);
+					UMovieSceneSkeletalAnimationSection::FRootMotionTransformParam Param;
+					Param.FrameRate = TickResolution;
+					Param.CurrentTime = EvaluationFrame;
+
+					if (AnimSection->GetRootMotionTransform(AnimationPoseData, Param))
 					{
-						CurrentAdditiveTransforms.Add(Param.OutTransform);
-						CurrentAdditiveWeights.Add(Param.OutWeight);
+						if (!Param.bOutIsAdditive)
+						{
+							CurrentTransform = Param.OutTransform * AnimSection->PreviousTransform;
+							CurrentTransforms.Add(CurrentTransform);
+							CurrentWeights.Add(Param.OutWeight);
+						}
+						else
+						{
+							CurrentAdditiveTransforms.Add(Param.OutTransform);
+							CurrentAdditiveWeights.Add(Param.OutWeight);
+						}
 					}
 				}
 			}
