Sequencer: Fix root bone root motion matching between sections. This wasn't working because the calculations to compare root bone positions weren't working due when 'force root lock' was selected on the anim sequence options. Added the 'Ignore Root Lock' option to the extraction settings to fix this problem.

diff --git a/Engine/Source/Runtime/MovieSceneTracks/Private/Tracks/MovieSceneSkeletalAnimationTrack.cpp b/Engine/Source/Runtime/MovieSceneTracks/Private/Tracks/MovieSceneSkeletalAnimationTrack.cpp
--- a/Engine/Source/Runtime/MovieSceneTracks/Private/Tracks/MovieSceneSkeletalAnimationTrack.cpp
+++ b/Engine/Source/Runtime/MovieSceneTracks/Private/Tracks/MovieSceneSkeletalAnimationTrack.cpp
@@ -802,6 +802,9 @@ static FTransform GetTransformForBoneRelativeToIndex(UAnimSequence* AnimSequence
 	UE::Anim::FStackAttributeContainer TempAttributes;
 	FAnimationPoseData AnimationPoseData(OutPose, OutCurve, TempAttributes);
 	FAnimExtractContext ExtractionContext(Seconds, false);
+#if WITH_EDITOR
+	ExtractionContext.bIgnoreRootLock = true;
+#endif
 	AnimSequence->GetAnimationPose(AnimationPoseData, ExtractionContext);
 	int32 MeshIndex = AnimationPoseData.GetPose().GetBoneContainer().GetPoseBoneIndexForBoneName(InBoneName);
 	if (MeshIndex != INDEX_NONE)
