Fix for an issue in sequencer where shots that are set to blend will have the incorrect blend time if the sequence play rate has changed.

diff --git a/Engine/Source/Runtime/MovieSceneTracks/Private/TrackInstances/MovieSceneCameraCutGameHandler.cpp b/Engine/Source/Runtime/MovieSceneTracks/Private/TrackInstances/MovieSceneCameraCutGameHandler.cpp
--- a/Engine/Source/Runtime/MovieSceneTracks/Private/TrackInstances/MovieSceneCameraCutGameHandler.cpp
+++ b/Engine/Source/Runtime/MovieSceneTracks/Private/TrackInstances/MovieSceneCameraCutGameHandler.cpp
@@ -321,12 +321,12 @@ void FCameraCutGameHandler::SetCameraCut(
 			{
 				CameraManager->SetGameCameraCutThisFrame();
 			}
-
+		
 			if (CameraComponent)
 			{
 				CameraComponent->NotifyCameraCut();
 			}
-
+		
 			if (UMovieSceneMotionVectorSimulationSystem* MotionVectorSim = Linker->FindSystem<UMovieSceneMotionVectorSimulationSystem>())
 			{
 				MotionVectorSim->SimulateAllTransforms();
@@ -350,7 +350,9 @@ void FCameraCutGameHandler::SetCameraCut(
 
 		// Convert known easing functions to their corresponding view target blend parameters.
 		TTuple<EViewTargetBlendFunction, float> BlendFunctionAndExp = BuiltInEasingTypeToBlendFunction(CameraCutParams.BlendType.GetValue());
-		TransitionParams.BlendTime = CameraCutParams.BlendTime * PlayRateFactor;
+		
+		// The playrate of sequences is defined as delta time * rate, so we need to match approach.
+		TransitionParams.BlendTime = CameraCutParams.BlendTime * (1 / FMath::Max(UE_SMALL_NUMBER, abs(PlayRateFactor)));
 		TransitionParams.bLockOutgoing = CameraCutParams.bLockPreviousCamera;
 		TransitionParams.BlendFunction = BlendFunctionAndExp.Get<0>();
 		TransitionParams.BlendExp = BlendFunctionAndExp.Get<1>();
@@ -416,4 +418,3 @@ void FCameraCutGameHandler::SetCameraCut(
 }
 
 }  // namespace UE::MovieScene
-
