Sequencer: Fixed post-evaluation events unnecessarily forcing an instantiation phase

diff --git a/Engine/Source/Runtime/MovieSceneTracks/Private/Sections/MovieSceneEventRepeaterSection.cpp b/Engine/Source/Runtime/MovieSceneTracks/Private/Sections/MovieSceneEventRepeaterSection.cpp
--- a/Engine/Source/Runtime/MovieSceneTracks/Private/Sections/MovieSceneEventRepeaterSection.cpp
+++ b/Engine/Source/Runtime/MovieSceneTracks/Private/Sections/MovieSceneEventRepeaterSection.cpp
@@ -47,13 +47,16 @@ void UMovieSceneEventRepeaterSection::ImportEntityImpl(UMovieSceneEntitySystemLi
 	}
 
 	UMovieSceneEventSystem* EventSystem = nullptr;
+	bool bMimicChanged = false;
 
 	if (EventTrack->EventPosition == EFireEventsAtPosition::AtStartOfEvaluation)
 	{
+		bMimicChanged = true;
 		EventSystem = EntityLinker->LinkSystem<UMovieScenePreSpawnEventSystem>();
 	}
 	else if (EventTrack->EventPosition == EFireEventsAtPosition::AfterSpawn)
 	{
+		bMimicChanged = true;
 		EventSystem = EntityLinker->LinkSystem<UMovieScenePostSpawnEventSystem>();
 	}
 	else
@@ -76,8 +79,11 @@ void UMovieSceneEventRepeaterSection::ImportEntityImpl(UMovieSceneEntitySystemLi
 
 	EventSystem->AddEvent(ThisInstance.GetRootInstanceHandle(), TriggerData);
 
-	// Mimic the structure changing in order to ensure that the instantiation phase runs
-	EntityLinker->EntityManager.MimicStructureChanged();
+	if (bMimicChanged)
+	{
+		// Mimic the structure changing in order to ensure that the instantiation phase runs
+		EntityLinker->EntityManager.MimicStructureChanged();
+	}
 }
 
 bool UMovieSceneEventRepeaterSection::PopulateEvaluationFieldImpl(const TRange<FFrameNumber>& EffectiveRange, const FMovieSceneEvaluationFieldEntityMetaData& InMetaData, FMovieSceneEntityComponentFieldBuilder* OutFieldBuilder)
diff --git a/Engine/Source/Runtime/MovieSceneTracks/Private/Sections/MovieSceneEventTriggerSection.cpp b/Engine/Source/Runtime/MovieSceneTracks/Private/Sections/MovieSceneEventTriggerSection.cpp
--- a/Engine/Source/Runtime/MovieSceneTracks/Private/Sections/MovieSceneEventTriggerSection.cpp
+++ b/Engine/Source/Runtime/MovieSceneTracks/Private/Sections/MovieSceneEventTriggerSection.cpp
@@ -66,13 +66,16 @@ void UMovieSceneEventTriggerSection::ImportEntityImpl(UMovieSceneEntitySystemLin
 	}
 
 	UMovieSceneEventSystem* EventSystem = nullptr;
+	bool bMimicChanged = false;
 
 	if (EventTrack->EventPosition == EFireEventsAtPosition::AtStartOfEvaluation)
 	{
+		bMimicChanged = true;
 		EventSystem = EntityLinker->LinkSystem<UMovieScenePreSpawnEventSystem>();
 	}
 	else if (EventTrack->EventPosition == EFireEventsAtPosition::AfterSpawn)
 	{
+		bMimicChanged = true;
 		EventSystem = EntityLinker->LinkSystem<UMovieScenePostSpawnEventSystem>();
 	}
 	else
@@ -95,8 +98,11 @@ void UMovieSceneEventTriggerSection::ImportEntityImpl(UMovieSceneEntitySystemLin
 
 	EventSystem->AddEvent(ThisInstance.GetRootInstanceHandle(), TriggerData);
 
-	// Mimic the structure changing in order to ensure that the instantiation phase runs
-	EntityLinker->EntityManager.MimicStructureChanged();
+	if (bMimicChanged)
+	{
+		// Mimic the structure changing in order to ensure that the instantiation phase runs
+		EntityLinker->EntityManager.MimicStructureChanged();
+	}
 }
 
 bool UMovieSceneEventTriggerSection::PopulateEvaluationFieldImpl(const TRange<FFrameNumber>& EffectiveRange, const FMovieSceneEvaluationFieldEntityMetaData& InMetaData, FMovieSceneEntityComponentFieldBuilder* OutFieldBuilder)
