Sequencer: No longer restore cached anim instances if they do not relate to the same object that is being restored

- this can happen when the anim instance was cached from one component, but construction scripts or object reinstantiation occurs

diff --git a/Engine/Source/Runtime/MovieSceneTracks/Private/Systems/MovieSceneSkeletalAnimationSystem.cpp b/Engine/Source/Runtime/MovieSceneTracks/Private/Systems/MovieSceneSkeletalAnimationSystem.cpp
--- a/Engine/Source/Runtime/MovieSceneTracks/Private/Systems/MovieSceneSkeletalAnimationSystem.cpp
+++ b/Engine/Source/Runtime/MovieSceneTracks/Private/Systems/MovieSceneSkeletalAnimationSystem.cpp
@@ -132,9 +132,11 @@ struct FPreAnimatedSkeletalAnimationTraits : FBoundObjectPreAnimatedStateTraits
 			// if we're using same anim blueprint, we don't want to keep reinitializing it. 
 			Component->SetAnimationMode(InOutCachedValue.AnimationMode);
 		}
-		if (InOutCachedValue.CachedAnimInstance.Get())
+
+		UAnimInstance* PreviousAnimInstance = InOutCachedValue.CachedAnimInstance.Get();
+		if (PreviousAnimInstance && PreviousAnimInstance->GetSkelMeshComponent() == Component)
 		{
-			Component->AnimScriptInstance = InOutCachedValue.CachedAnimInstance.Get();
+			Component->AnimScriptInstance = PreviousAnimInstance;
 			InOutCachedValue.CachedAnimInstance.Reset();
 			if (Component->AnimScriptInstance && Component->GetSkeletalMeshAsset() && Component->AnimScriptInstance->CurrentSkeleton != Component->GetSkeletalMeshAsset()->GetSkeleton())
 			{
