Sequencer: don't cache/restore pre-animated viewport FOV in editor

When restoring pre-animated state for an editor viewport, we always go back to where the viewport was before the cinematic camera took control. But while we do want to go back to the location/rotation we were at, we don't want to go back to whatever FOV the viewport had back then. We always want viewport FOVs to use the set FOV angle in the perspective settings. So there's no need to remember the FOV value.

diff --git a/Engine/Source/Runtime/MovieSceneTracks/Private/TrackInstances/MovieSceneCameraCutEditorHandler.cpp b/Engine/Source/Runtime/MovieSceneTracks/Private/TrackInstances/MovieSceneCameraCutEditorHandler.cpp
--- a/Engine/Source/Runtime/MovieSceneTracks/Private/TrackInstances/MovieSceneCameraCutEditorHandler.cpp
+++ b/Engine/Source/Runtime/MovieSceneTracks/Private/TrackInstances/MovieSceneCameraCutEditorHandler.cpp
@@ -32,7 +32,6 @@ FPreAnimatedCameraCutEditorState FPreAnimatedCameraCutEditorTraits::CachePreAnim
 	FPreAnimatedCameraCutEditorState CachedValue;
 	CachedValue.ViewportLocation = InKey->GetViewLocation();
 	CachedValue.ViewportRotation = InKey->GetViewRotation();
-	CachedValue.ViewportFOV = InKey->ViewFOV;
 	return CachedValue;
 }
 
@@ -46,7 +45,8 @@ void FPreAnimatedCameraCutEditorTraits::RestorePreAnimatedValue(
 		return;
 	}
 
-	// Check that our pointer is still valid by searching it in active viewports.
+	// Check that our pointer is still valid by searching it in active viewports,
+	// in case the user removed that viewport since we cached the pre-animated state.
 	if (GEditor->GetLevelViewportClients().Find(InKey) == INDEX_NONE)
 	{
 		return;
@@ -75,7 +75,7 @@ void FPreAnimatedCameraCutEditorTraits::RestorePreAnimatedValue(
 	{
 		InKey->SetViewLocation(CachedValue.ViewportLocation);
 		InKey->SetViewRotation(CachedValue.ViewportRotation);
-		InKey->ViewFOV = CachedValue.ViewportFOV;
+		InKey->ViewFOV = InKey->FOVAngle;
 	}
 
 	InKey->SetCinematicActorLock(nullptr);
@@ -312,7 +312,7 @@ void FCameraCutEditorHandler::SetCameraCutForViewport(
 			FPreAnimatedCameraCutEditorState CachedValue = PreAnimatedStorage->GetCachedValue(StorageIndex);
 			ViewLocation = CachedValue.ViewportLocation;
 			ViewRotation = CachedValue.ViewportRotation;
-			ViewFOV = CachedValue.ViewportFOV;
+			ViewFOV = ViewportClient.FOVAngle;
 		}
 	}
 
@@ -442,7 +442,7 @@ void FCameraCutEditorHandler::ReleaseCameraCutForViewport(
 		FPreAnimatedCameraCutEditorState CachedValue = PreAnimatedStorage->GetCachedValue(StorageIndex);
 		ViewportClient.SetViewLocation(CachedValue.ViewportLocation);
 		ViewportClient.SetViewRotation(CachedValue.ViewportRotation);
-		ViewportClient.ViewFOV = CachedValue.ViewportFOV;
+		ViewportClient.ViewFOV = ViewportClient.FOVAngle;
 	}
 
 	// Actually release control.
diff --git a/Engine/Source/Runtime/MovieSceneTracks/Private/TrackInstances/MovieSceneCameraCutEditorHandler.h b/Engine/Source/Runtime/MovieSceneTracks/Private/TrackInstances/MovieSceneCameraCutEditorHandler.h
--- a/Engine/Source/Runtime/MovieSceneTracks/Private/TrackInstances/MovieSceneCameraCutEditorHandler.h
+++ b/Engine/Source/Runtime/MovieSceneTracks/Private/TrackInstances/MovieSceneCameraCutEditorHandler.h
@@ -28,7 +28,6 @@ struct FPreAnimatedCameraCutEditorState
 {
 	FVector ViewportLocation;
 	FRotator ViewportRotation;
-	float ViewportFOV = 0.f;
 };
 
 /** Pre-animated traits for in-editor camera cuts */
diff --git a/Engine/Source/Runtime/MovieSceneTracks/Private/TrackInstances/MovieSceneCameraCutViewportPreviewer.cpp b/Engine/Source/Runtime/MovieSceneTracks/Private/TrackInstances/MovieSceneCameraCutViewportPreviewer.cpp
--- a/Engine/Source/Runtime/MovieSceneTracks/Private/TrackInstances/MovieSceneCameraCutViewportPreviewer.cpp
+++ b/Engine/Source/Runtime/MovieSceneTracks/Private/TrackInstances/MovieSceneCameraCutViewportPreviewer.cpp
@@ -39,7 +39,7 @@ void FCameraCutViewportPreviewerTarget::Get(FLevelEditorViewportClient* InClient
 			FPreAnimatedCameraCutEditorState CachedValue = PreAnimatedStorage->GetCachedValue(StorageIndex);
 			OutLocation = CachedValue.ViewportLocation;
 			OutRotation = CachedValue.ViewportRotation;
-			OutFOV = CachedValue.ViewportFOV;
+			OutFOV = InClient->FOVAngle;
 			return;
 		}
 	}
