Sequencer: Fix crash when selecting the default animation target.

diff --git a/Engine/Source/Runtime/MovieSceneTracks/Public/MovieSceneMixedAnimationTarget.h b/Engine/Source/Runtime/MovieSceneTracks/Public/MovieSceneMixedAnimationTarget.h
--- a/Engine/Source/Runtime/MovieSceneTracks/Public/MovieSceneMixedAnimationTarget.h
+++ b/Engine/Source/Runtime/MovieSceneTracks/Public/MovieSceneMixedAnimationTarget.h
@@ -18,4 +18,9 @@ struct MOVIESCENETRACKS_API FMovieSceneMixedAnimationTarget
 	GENERATED_BODY() 
 
 	virtual ~FMovieSceneMixedAnimationTarget() {}
+
+	FORCEINLINE friend uint32 GetTypeHash(const FMovieSceneMixedAnimationTarget& Target)
+	{
+		return GetTypeHash(FMovieSceneMixedAnimationTarget::StaticStruct());
+	}
 };
\ No newline at end of file
