Sequencer: add support for post-process settings to camera cut blending preview in editor

diff --git a/Engine/Source/Runtime/MovieSceneTracks/Private/TrackInstances/MovieSceneCameraCutViewportPreviewer.cpp b/Engine/Source/Runtime/MovieSceneTracks/Private/TrackInstances/MovieSceneCameraCutViewportPreviewer.cpp
--- a/Engine/Source/Runtime/MovieSceneTracks/Private/TrackInstances/MovieSceneCameraCutViewportPreviewer.cpp
+++ b/Engine/Source/Runtime/MovieSceneTracks/Private/TrackInstances/MovieSceneCameraCutViewportPreviewer.cpp
@@ -4,8 +4,8 @@
 
 #if WITH_EDITOR
 
+#include "Engine/PostProcessUtils.h"
 #include "Evaluation/CameraCutPlaybackCapability.h"
-#include "IMovieScenePlayer.h"
 #include "LevelEditorViewport.h"
 #include "MovieSceneCommonHelpers.h"
 #include "TrackInstances/MovieSceneCameraCutEditorHandler.h"
@@ -13,13 +13,15 @@
 namespace UE::MovieScene
 {
 
-void FCameraCutViewportPreviewerTarget::Get(FLevelEditorViewportClient* InClient, FVector& OutLocation, FRotator& OutRotation, float& OutFOV) const
+void FCameraCutViewportPreviewerTarget::Get(FLevelEditorViewportClient* InClient, FVector& OutLocation, FRotator& OutRotation, float& OutFOV, const FPostProcessSettings*& OutPP, float& OutPPWeight) const
 {
 	if (CameraComponent)
 	{
 		OutLocation = CameraComponent->GetComponentLocation();
 		OutRotation = CameraComponent->GetComponentRotation();
 		OutFOV = CameraComponent->FieldOfView;
+		OutPP = &CameraComponent->PostProcessSettings;
+		OutPPWeight = CameraComponent->PostProcessBlendWeight;
 		return;
 	}
 
@@ -28,6 +30,8 @@ void FCameraCutViewportPreviewerTarget::Get(FLevelEditorViewportClient* InClient
 		OutLocation = CameraActor->GetActorLocation();
 		OutRotation = CameraActor->GetActorRotation();
 		OutFOV = -1.f;
+		OutPP = nullptr;
+		OutPPWeight = 0.f;
 		return;
 	}
 
@@ -40,6 +44,8 @@ void FCameraCutViewportPreviewerTarget::Get(FLevelEditorViewportClient* InClient
 			OutLocation = CachedValue.ViewportLocation;
 			OutRotation = CachedValue.ViewportRotation;
 			OutFOV = InClient->FOVAngle;
+			OutPP = nullptr;
+			OutPPWeight = 0.f;
 			return;
 		}
 	}
@@ -49,6 +55,8 @@ void FCameraCutViewportPreviewerTarget::Get(FLevelEditorViewportClient* InClient
 	OutLocation = InClient->GetViewLocation();
 	OutRotation = InClient->GetViewRotation();
 	OutFOV = InClient->ViewFOV;
+	OutPP = nullptr;
+	OutPPWeight = 0.f;
 }
 
 FCameraCutViewportPreviewer::FCameraCutViewportPreviewer()
@@ -126,10 +134,12 @@ void FCameraCutViewportPreviewer::ModifyViewportClientView(FEditorViewportViewMo
 	FVector FromViewLocation, ToViewLocation;
 	FRotator FromViewRotation, ToViewRotation;
 	float FromViewFOV, ToViewFOV;
+	const FPostProcessSettings* FromPP; const FPostProcessSettings* ToPP;
+	float FromPPWeight, ToPPWeight;
 	FLevelEditorViewportClient* ViewportClient = static_cast<FLevelEditorViewportClient*>(Params.ViewportClient);
 
-	FromTarget.Get(ViewportClient, FromViewLocation, FromViewRotation, FromViewFOV);
-	ToTarget.Get(ViewportClient, ToViewLocation, ToViewRotation, ToViewFOV);
+	FromTarget.Get(ViewportClient, FromViewLocation, FromViewRotation, FromViewFOV, FromPP, FromPPWeight);
+	ToTarget.Get(ViewportClient, ToViewLocation, ToViewRotation, ToViewFOV, ToPP, ToPPWeight);
 
 	const FVector BlendedLocation = FMath::Lerp(FromViewLocation, ToViewLocation, BlendFactor);
 	const FRotator BlendedRotation = FMath::Lerp(FromViewRotation, ToViewRotation, BlendFactor);
@@ -138,6 +148,28 @@ void FCameraCutViewportPreviewer::ModifyViewportClientView(FEditorViewportViewMo
 	Params.ViewInfo.Location = BlendedLocation;
 	Params.ViewInfo.Rotation = BlendedRotation;
 	Params.ViewInfo.FOV = BlendedFOV;
+
+	if (FromPP)
+	{
+		FPostProcessUtils::OverridePostProcessSettings(Params.ViewInfo.PostProcessSettings, *FromPP);
+	}
+	if (ToPP)
+	{
+		FPostProcessUtils::BlendPostProcessSettings(Params.ViewInfo.PostProcessSettings, *ToPP, BlendFactor);
+	}
+
+	if (FromPP && !ToPP)
+	{
+		Params.ViewInfo.PostProcessBlendWeight = (1.f - BlendFactor);
+	}
+	else if (!FromPP && ToPP)
+	{
+		Params.ViewInfo.PostProcessBlendWeight = BlendFactor;
+	}
+	else if (FromPP && ToPP)
+	{
+		Params.ViewInfo.PostProcessBlendWeight = 1.f;
+	}
 }
 
 void FCameraCutViewportPreviewer::OnLevelViewportClientListChanged()
diff --git a/Engine/Source/Runtime/MovieSceneTracks/Private/TrackInstances/MovieSceneCameraCutViewportPreviewer.h b/Engine/Source/Runtime/MovieSceneTracks/Private/TrackInstances/MovieSceneCameraCutViewportPreviewer.h
--- a/Engine/Source/Runtime/MovieSceneTracks/Private/TrackInstances/MovieSceneCameraCutViewportPreviewer.h
+++ b/Engine/Source/Runtime/MovieSceneTracks/Private/TrackInstances/MovieSceneCameraCutViewportPreviewer.h
@@ -11,6 +11,7 @@ class AActor;
 class FLevelEditorViewportClient;
 class UCameraComponent;
 struct FEditorViewportViewModifierParams;
+struct FPostProcessSettings;
 struct FMovieSceneCameraCutParams;
 
 namespace UE::MovieScene
@@ -31,7 +32,7 @@ namespace UE::MovieScene
 		TSharedPtr<FPreAnimatedCameraCutEditorStorage> PreAnimatedStorage;
 
 		/** Returns the target's location, rotation, and FOV */
-		void Get(FLevelEditorViewportClient* InClient, FVector& OutLocation, FRotator& OutRotation, float& OutFOV) const;
+		void Get(FLevelEditorViewportClient* InClient, FVector& OutLocation, FRotator& OutRotation, float& OutFOV, const FPostProcessSettings*& OutPP, float& OutPPWeight) const;
 	};
 
 	/**
@@ -69,6 +70,7 @@ namespace UE::MovieScene
 		void TeardownBlend();
 
 	private:
+
 		void UpdatePreviewLevelViewportClientFromCameraCut(FLevelEditorViewportClient& InViewportClient, UObject* InCameraObject, const FMovieSceneCameraCutParams& CameraCutParams);
 		void ModifyViewportClientView(FEditorViewportViewModifierParams& Params);
 		void OnLevelViewportClientListChanged();
