Sequencer: Automatically enable Sequencer.AnimationUIFlickerFix when using the new Control Rig ECS codepath

AnimationUIFlickerFix only exists as a temporary workaround for the legacy template codepath, and actually creates a bug with layered rigs whereby the animation instance from the skeletal animation system is not set up correctly when the control rig ECS code binds to the anim instance, meaning layered rigs do not get the skel anim instance bound correctly. This flicker fix cvar will eventually be removed (and always enabled), but while the template codepath is enabled by default we need to keep it around

diff --git a/Engine/Plugins/Animation/ControlRig/Source/ControlRig/Private/Sequencer/MovieSceneControlRigParameterTrack.cpp b/Engine/Plugins/Animation/ControlRig/Source/ControlRig/Private/Sequencer/MovieSceneControlRigParameterTrack.cpp
--- a/Engine/Plugins/Animation/ControlRig/Source/ControlRig/Private/Sequencer/MovieSceneControlRigParameterTrack.cpp
+++ b/Engine/Plugins/Animation/ControlRig/Source/ControlRig/Private/Sequencer/MovieSceneControlRigParameterTrack.cpp
@@ -25,6 +25,8 @@
 
 namespace UE::MovieScene
 {
+	MOVIESCENETRACKS_API extern bool (*ShouldUseLegacyControlRigTemplate)();
+
 	bool GUseLegacyControlRigTemplate = true;
 	FAutoConsoleVariableRef CVarUseLegacyControlRigTemplate(
 		TEXT("ControlRig.UseLegacySequencerTemplate"),
@@ -73,6 +75,8 @@ UMovieSceneControlRigParameterTrack::UMovieSceneControlRigParameterTrack(const F
 	SupportedBlendTypes.Add(EMovieSceneBlendType::Additive);
 	SupportedBlendTypes.Add(EMovieSceneBlendType::Absolute);
 	SupportedBlendTypes.Add(EMovieSceneBlendType::Override);
+
+	UE::MovieScene::ShouldUseLegacyControlRigTemplate = &UMovieSceneControlRigParameterTrack::ShouldUseLegacyTemplate;
 }
 
 void UMovieSceneControlRigParameterTrack::BeginDestroy()
diff --git a/Engine/Source/Runtime/MovieSceneTracks/Private/Systems/MovieSceneSkeletalAnimationSystem.cpp b/Engine/Source/Runtime/MovieSceneTracks/Private/Systems/MovieSceneSkeletalAnimationSystem.cpp
--- a/Engine/Source/Runtime/MovieSceneTracks/Private/Systems/MovieSceneSkeletalAnimationSystem.cpp
+++ b/Engine/Source/Runtime/MovieSceneTracks/Private/Systems/MovieSceneSkeletalAnimationSystem.cpp
@@ -58,7 +58,13 @@ DECLARE_CYCLE_STAT(TEXT("Evaluate skeletal animations"), MovieSceneEval_Evaluate
 namespace UE::MovieScene
 {
 
+MOVIESCENETRACKS_API bool (*ShouldUseLegacyControlRigTemplate)() = nullptr;
+
 bool GAnimationUIFlickerFix = false;
+bool AnimationUIFlickerFixActive()
+{
+	 return GAnimationUIFlickerFix == true || (ShouldUseLegacyControlRigTemplate && !ShouldUseLegacyControlRigTemplate());
+}
 FAutoConsoleVariableRef CVarAnimationUIFlickerFix(
 	TEXT("Sequencer.Animation.UIFlickerFix"),
 	GAnimationUIFlickerFix,
@@ -557,7 +563,7 @@ private:
 			return;
 		}
 
-		if (GAnimationUIFlickerFix == false)
+		if (!AnimationUIFlickerFixActive())
 		{
 			// Cache pre-animated state for this bound object before doing anything.
 			// We don't yet track what entities have already started animated vs. entities that just started this frame,
@@ -1241,7 +1247,7 @@ void UMovieSceneSkeletalAnimationSystem::OnRun(FSystemTaskPrerequisites& InPrere
 	TSharedRef<FMovieSceneEntitySystemRunner> Runner = Linker->GetRunner();
 	if (Runner->GetCurrentPhase() == ESystemPhase::Instantiation)
 	{
-		if (GAnimationUIFlickerFix == true)
+		if (AnimationUIFlickerFixActive())
 		{
 			// Begin tracking pre-animated state for all bound skel animation components
 			TSharedPtr<FPreAnimatedSkeletalAnimationStorage> PreAnimatedStorage = Linker->PreAnimatedState.GetOrCreateStorage<FPreAnimatedSkeletalAnimationStorage>();
