Sequencer: Fixed MPC components not being grouped correctly

- This was the cause of hierarchical bias and easing not working, since those systems didn't know how to group MPCs

diff --git a/Engine/Source/Runtime/MovieSceneTracks/Private/Systems/MovieSceneMaterialParameterCollectionSystem.cpp b/Engine/Source/Runtime/MovieSceneTracks/Private/Systems/MovieSceneMaterialParameterCollectionSystem.cpp
--- a/Engine/Source/Runtime/MovieSceneTracks/Private/Systems/MovieSceneMaterialParameterCollectionSystem.cpp
+++ b/Engine/Source/Runtime/MovieSceneTracks/Private/Systems/MovieSceneMaterialParameterCollectionSystem.cpp
@@ -13,6 +13,7 @@
 #include "Systems/FloatChannelEvaluatorSystem.h"
 #include "Systems/MovieScenePiecewiseDoubleBlenderSystem.h"
 #include "Systems/MovieScenePreAnimatedMaterialParameters.h"
+#include "Systems/WeightAndEasingEvaluatorSystem.h"
 
 #include "Logging/MessageLog.h"
 #include "Misc/UObjectToken.h"
@@ -25,6 +26,56 @@
 
 #include UE_INLINE_GENERATED_CPP_BY_NAME(MovieSceneMaterialParameterCollectionSystem)
 
+namespace UE::MovieScene
+{
+
+struct FMPCGroupingPolicy
+{
+	using GroupKeyType = FObjectKey;
+
+	void InitializeGroupKeys(
+		TEntityGroupingHandlerBase<FMPCGroupingPolicy>& Handler,
+		FEntityGroupBuilder* Builder,
+		FEntityAllocationIteratorItem Item,
+		FReadEntityIDs EntityIDs,
+		TWrite<FEntityGroupID> GroupIDs,
+		TRead<TWeakObjectPtr<UMaterialParameterCollection>> MPCs)
+	{
+		const FEntityAllocation* Allocation = Item.GetAllocation();
+		const int32              Num        = Allocation->Num();
+
+		for (int32 Index = 0; Index < Num; ++Index)
+		{
+			if (!GroupIDs[Index].IsValid())
+			{
+				FObjectKey Key = MPCs[Index].Get();
+
+				const int32    NewGroupIndex = Handler.GetOrAllocateGroupIndex(Key, Builder);
+				FEntityGroupID NewGroupID = Builder->MakeGroupID(NewGroupIndex);
+
+				Builder->AddEntityToGroup(EntityIDs[Index], NewGroupID);
+
+				// Write out the group ID component
+				GroupIDs[Index] = NewGroupID;
+			}
+		}
+	}
+
+#if WITH_EDITOR
+	bool OnObjectsReplaced(FObjectKey& InOutKey, const TMap<UObject*, UObject*>& ReplacementMap)
+	{
+		if (UObject* const * NewObject = ReplacementMap.Find(InOutKey.ResolveObjectPtr()))
+		{
+			InOutKey = *NewObject;
+			return true;
+		}
+		return false;
+	}
+#endif
+};
+
+} // namespace UE::MovieScene
+
 
 UMovieSceneMaterialParameterCollectionSystem::UMovieSceneMaterialParameterCollectionSystem(const FObjectInitializer& ObjInit)
 	: Super(ObjInit)
@@ -42,6 +93,7 @@ UMovieSceneMaterialParameterCollectionSystem::UMovieSceneMaterialParameterCollec
 		DefineComponentConsumer(GetClass(), BuiltInComponents->Tags.Root);
 		DefineComponentProducer(GetClass(), TracksComponents->BoundMaterial);
 		DefineImplicitPrerequisite(UMovieSceneCachePreAnimatedStateSystem::StaticClass(), GetClass());
+		DefineImplicitPrerequisite(UMovieSceneHierarchicalEasingInstantiatorSystem::StaticClass(), GetClass());
 	}
 }
 
@@ -51,6 +103,25 @@ void UMovieSceneMaterialParameterCollectionSystem::OnLink()
 
 	ScalarParameterStorage = Linker->PreAnimatedState.GetOrCreateStorage<FPreAnimatedScalarMaterialParameterStorage>();
 	VectorParameterStorage = Linker->PreAnimatedState.GetOrCreateStorage<FPreAnimatedVectorMaterialParameterStorage>();
+
+	UMovieSceneEntityGroupingSystem* GroupingSystem = Linker->LinkSystem<UMovieSceneEntityGroupingSystem>();
+
+	GroupingKey = GroupingSystem->AddGrouping(
+		FMPCGroupingPolicy(),
+		FMovieSceneTracksComponentTypes::Get()->MPC
+	);
+}
+
+void UMovieSceneMaterialParameterCollectionSystem::OnUnlink()
+{
+	using namespace UE::MovieScene;
+
+	UMovieSceneEntityGroupingSystem* GroupingSystem = Linker->FindSystem<UMovieSceneEntityGroupingSystem>();
+	if (ensure(GroupingSystem))
+	{
+		GroupingSystem->RemoveGrouping(GroupingKey);
+	}
+	GroupingKey = FEntityGroupingPolicyKey();
 }
 
 void UMovieSceneMaterialParameterCollectionSystem::OnRun(FSystemTaskPrerequisites& InPrerequisites, FSystemSubsequentTasks& Subsequents)
diff --git a/Engine/Source/Runtime/MovieSceneTracks/Public/Systems/MovieSceneMaterialParameterCollectionSystem.h b/Engine/Source/Runtime/MovieSceneTracks/Public/Systems/MovieSceneMaterialParameterCollectionSystem.h
--- a/Engine/Source/Runtime/MovieSceneTracks/Public/Systems/MovieSceneMaterialParameterCollectionSystem.h
+++ b/Engine/Source/Runtime/MovieSceneTracks/Public/Systems/MovieSceneMaterialParameterCollectionSystem.h
@@ -3,6 +3,7 @@
 #pragma once
 
 #include "EntitySystem/MovieSceneEntitySystem.h"
+#include "EntitySystem/BuiltInComponentTypes.h"
 #include "Evaluation/PreAnimatedState/MovieScenePreAnimatedStateStorage.h"
 
 #include "Systems/MovieSceneMaterialSystem.h"
@@ -35,6 +36,7 @@ public:
 private:
 
 	virtual void OnLink() override;
+	virtual void OnUnlink() override;
 	virtual void OnRun(FSystemTaskPrerequisites& InPrerequisites, FSystemSubsequentTasks& Subsequents) override;
 
 private:
@@ -43,4 +45,6 @@ private:
 	TSharedPtr<UE::MovieScene::FPreAnimatedScalarMaterialParameterStorage> ScalarParameterStorage;
 	/** Holds pre-animated values for vector or color MPC values */
 	TSharedPtr<UE::MovieScene::FPreAnimatedVectorMaterialParameterStorage> VectorParameterStorage;
+	/** Registered key for the grouping policy that groups MPC tracks together */
+	UE::MovieScene::FEntityGroupingPolicyKey GroupingKey;
 };
