Sequencer: Deprecate UniqueTrackName in UMovieScenePropertyTrack in favor of PropertyBinding.PropertyName.

diff --git a/Engine/Plugins/MovieScene/SequencerScripting/Source/SequencerScripting/Private/ExtensionLibraries/MovieScenePropertyTrackExtensions.cpp b/Engine/Plugins/MovieScene/SequencerScripting/Source/SequencerScripting/Private/ExtensionLibraries/MovieScenePropertyTrackExtensions.cpp
--- a/Engine/Plugins/MovieScene/SequencerScripting/Source/SequencerScripting/Private/ExtensionLibraries/MovieScenePropertyTrackExtensions.cpp
+++ b/Engine/Plugins/MovieScene/SequencerScripting/Source/SequencerScripting/Private/ExtensionLibraries/MovieScenePropertyTrackExtensions.cpp
@@ -49,7 +49,7 @@ FName UMovieScenePropertyTrackExtensions::GetUniqueTrackName(UMovieSceneProperty
 	}
 
 #if WITH_EDITORONLY_DATA
-	return Track->UniqueTrackName;
+	return Track->GetTrackName();
 #else
 	return NAME_None;
 #endif
diff --git a/Engine/Source/Editor/MovieSceneTools/Public/PropertyTrackEditor.h b/Engine/Source/Editor/MovieSceneTools/Public/PropertyTrackEditor.h
--- a/Engine/Source/Editor/MovieSceneTools/Public/PropertyTrackEditor.h
+++ b/Engine/Source/Editor/MovieSceneTools/Public/PropertyTrackEditor.h
@@ -150,11 +150,7 @@ protected:
 
 	virtual UMovieSceneTrack* AddTrack(UMovieScene* FocusedMovieScene, const FGuid& ObjectHandle, TSubclassOf<class UMovieSceneTrack> TrackClass, FName UniqueTypeName) override
 	{
-		UMovieSceneTrack* Track = FocusedMovieScene->AddTrack(TrackClass, ObjectHandle);
-		if (UMovieScenePropertyTrack* PropertyTrack = Cast<UMovieScenePropertyTrack>(Track))
-		{
-			PropertyTrack->UniqueTrackName = UniqueTypeName;
-		}
+		UMovieSceneTrack* Track = FocusedMovieScene->AddTrack(TrackClass, ObjectHandle); 
 		return Track;
 	}
 
diff --git a/Engine/Source/Runtime/MovieSceneTracks/Private/Tracks/MovieScenePropertyTrack.cpp b/Engine/Source/Runtime/MovieSceneTracks/Private/Tracks/MovieScenePropertyTrack.cpp
--- a/Engine/Source/Runtime/MovieSceneTracks/Private/Tracks/MovieScenePropertyTrack.cpp
+++ b/Engine/Source/Runtime/MovieSceneTracks/Private/Tracks/MovieScenePropertyTrack.cpp
@@ -26,13 +26,6 @@ void UMovieScenePropertyTrack::SetPropertyNameAndPath(FName InPropertyName, cons
 	check((InPropertyName != NAME_None) && !InPropertyPath.IsEmpty());
 
 	PropertyBinding = FMovieScenePropertyBinding(InPropertyName, InPropertyPath);
-
-#if WITH_EDITORONLY_DATA
-	if (UniqueTrackName == NAME_None)
-	{
-		UniqueTrackName = *InPropertyPath;
-	}
-#endif
 }
 
 
@@ -44,15 +37,6 @@ const TArray<UMovieSceneSection*>& UMovieScenePropertyTrack::GetAllSections() co
 
 void UMovieScenePropertyTrack::PostLoad()
 {
-#if WITH_EDITORONLY_DATA
-
-	if (UniqueTrackName.IsNone())
-	{
-		UniqueTrackName = PropertyBinding.PropertyPath;
-	}
-
-#endif
-
 	Super::PostLoad();
 }
 
@@ -133,12 +117,12 @@ FSlateColor UMovieScenePropertyTrack::GetLabelColor(const FMovieSceneLabelParams
 	}
 	return LabelParams.bIsDimmed ? FSlateColor(FLinearColor::Red.Desaturate(0.6f)) : FLinearColor::Red;
 }
+#endif
 
 FName UMovieScenePropertyTrack::GetTrackName() const
 {
-	return UniqueTrackName;
+	return PropertyBinding.PropertyPath;
 }
-#endif
 
 void UMovieScenePropertyTrack::RemoveAllAnimationData()
 {
diff --git a/Engine/Source/Runtime/MovieSceneTracks/Public/Tracks/MovieScenePropertyTrack.h b/Engine/Source/Runtime/MovieSceneTracks/Public/Tracks/MovieScenePropertyTrack.h
--- a/Engine/Source/Runtime/MovieSceneTracks/Public/Tracks/MovieScenePropertyTrack.h
+++ b/Engine/Source/Runtime/MovieSceneTracks/Public/Tracks/MovieScenePropertyTrack.h
@@ -42,8 +42,8 @@ public:
 	MOVIESCENETRACKS_API virtual FText GetDisplayNameToolTipText(const FMovieSceneLabelParams& LabelParams) const override;
 	MOVIESCENETRACKS_API virtual FSlateColor GetLabelColor(const FMovieSceneLabelParams& LabelParams) const override;
 	virtual bool CanRename() const override { return false; }
-	MOVIESCENETRACKS_API virtual FName GetTrackName() const override;
 #endif
+	MOVIESCENETRACKS_API virtual FName GetTrackName() const override;
 
 	MOVIESCENETRACKS_API virtual void PostLoad() override;
 	MOVIESCENETRACKS_API virtual void Serialize(FArchive& Ar) override;
@@ -122,6 +122,7 @@ public:
 #if WITH_EDITORONLY_DATA
 public:
 	/** Unique name for this track to afford multiple tracks on a given object (i.e. for array properties) */
+	UE_DEPRECATED(5.6, "UniqueTrackName is deprecated, use GetTrackName instead.")
 	UPROPERTY()
 	FName UniqueTrackName;
 
