Sequencer: MPCs now also group by parameter name, imported entities are excluded

diff --git a/Engine/Source/Runtime/MovieSceneTracks/Private/Systems/MovieSceneMaterialParameterCollectionSystem.cpp b/Engine/Source/Runtime/MovieSceneTracks/Private/Systems/MovieSceneMaterialParameterCollectionSystem.cpp
--- a/Engine/Source/Runtime/MovieSceneTracks/Private/Systems/MovieSceneMaterialParameterCollectionSystem.cpp
+++ b/Engine/Source/Runtime/MovieSceneTracks/Private/Systems/MovieSceneMaterialParameterCollectionSystem.cpp
@@ -13,6 +13,7 @@
 #include "Systems/FloatChannelEvaluatorSystem.h"
 #include "Systems/MovieScenePiecewiseDoubleBlenderSystem.h"
 #include "Systems/MovieScenePreAnimatedMaterialParameters.h"
+#include "Systems/WeightAndEasingEvaluatorSystem.h"
 
 #include "Logging/MessageLog.h"
 #include "Misc/UObjectToken.h"
@@ -25,6 +26,96 @@
 
 #include UE_INLINE_GENERATED_CPP_BY_NAME(MovieSceneMaterialParameterCollectionSystem)
 
+namespace UE::MovieScene
+{
+
+struct FMPCGroupingPolicy
+{
+	using GroupKeyType = TTuple<FObjectKey, FMaterialParameterInfo>;
+
+	void InitializeGroupKeys(
+		TEntityGroupingHandlerBase<FMPCGroupingPolicy>& Handler,
+		FEntityGroupBuilder* Builder,
+		FEntityAllocationIteratorItem Item,
+		FReadEntityIDs EntityIDs,
+		TWrite<FEntityGroupID> GroupIDs,
+		TRead<TWeakObjectPtr<UMaterialParameterCollection>> MPCs)
+	{
+		const FEntityAllocation* Allocation     = Item.GetAllocation();
+		const FComponentMask&    AllocationType = Item.GetAllocationType();
+		const int32              Num            = Allocation->Num();
+		
+		FMovieSceneTracksComponentTypes* TracksComponents = FMovieSceneTracksComponentTypes::Get();
+
+		TComponentTypeID<FName> ParameterName;
+		if (AllocationType.Contains(TracksComponents->ScalarParameterName))
+		{
+			ParameterName = TracksComponents->ScalarParameterName;
+		}
+		else if (AllocationType.Contains(TracksComponents->VectorParameterName))
+		{
+			ParameterName = TracksComponents->VectorParameterName;
+		}
+		else if (AllocationType.Contains(TracksComponents->ColorParameterName))
+		{
+			ParameterName = TracksComponents->ColorParameterName;
+		}
+
+		TComponentTypeID<FMaterialParameterInfo> ParameterInfo;
+		if (AllocationType.Contains(TracksComponents->ScalarMaterialParameterInfo))
+		{
+			ParameterInfo = TracksComponents->ScalarMaterialParameterInfo;
+		}
+		else if (AllocationType.Contains(TracksComponents->VectorMaterialParameterInfo))
+		{
+			ParameterInfo = TracksComponents->VectorMaterialParameterInfo;
+		}
+		else if (AllocationType.Contains(TracksComponents->ColorMaterialParameterInfo))
+		{
+			ParameterInfo = TracksComponents->ColorMaterialParameterInfo;
+		}
+
+		TReadOptional<FName>                  ParameterNames = ParameterName ? Allocation->TryReadComponents(ParameterName) : TReadOptional<FName>();
+		TReadOptional<FMaterialParameterInfo> ParameterInfos = ParameterInfo ? Allocation->TryReadComponents(ParameterInfo) : TReadOptional<FMaterialParameterInfo>();
+
+		if (!ensure(ParameterInfos || ParameterNames))
+		{
+			return;
+		}
+
+		for (int32 Index = 0; Index < Num; ++Index)
+		{
+			if (!GroupIDs[Index].IsValid())
+			{
+				FMaterialParameterInfo ThisParameterInfo = ParameterInfos ? ParameterInfos[Index] : FMaterialParameterInfo(ParameterNames[Index]);
+				GroupKeyType Key = MakeTuple(FObjectKey(MPCs[Index].Get()), ThisParameterInfo);
+
+				const int32    NewGroupIndex = Handler.GetOrAllocateGroupIndex(Key, Builder);
+				FEntityGroupID NewGroupID = Builder->MakeGroupID(NewGroupIndex);
+
+				Builder->AddEntityToGroup(EntityIDs[Index], NewGroupID);
+
+				// Write out the group ID component
+				GroupIDs[Index] = NewGroupID;
+			}
+		}
+	}
+
+#if WITH_EDITOR
+	bool OnObjectsReplaced(GroupKeyType& InOutKey, const TMap<UObject*, UObject*>& ReplacementMap)
+	{
+		if (UObject* const * NewObject = ReplacementMap.Find(InOutKey.Get<0>().ResolveObjectPtr()))
+		{
+			InOutKey.Get<0>() = *NewObject;
+			return true;
+		}
+		return false;
+	}
+#endif
+};
+
+} // namespace UE::MovieScene
+
 
 UMovieSceneMaterialParameterCollectionSystem::UMovieSceneMaterialParameterCollectionSystem(const FObjectInitializer& ObjInit)
 	: Super(ObjInit)
@@ -42,6 +133,7 @@ UMovieSceneMaterialParameterCollectionSystem::UMovieSceneMaterialParameterCollec
 		DefineComponentConsumer(GetClass(), BuiltInComponents->Tags.Root);
 		DefineComponentProducer(GetClass(), TracksComponents->BoundMaterial);
 		DefineImplicitPrerequisite(UMovieSceneCachePreAnimatedStateSystem::StaticClass(), GetClass());
+		DefineImplicitPrerequisite(UMovieSceneHierarchicalEasingInstantiatorSystem::StaticClass(), GetClass());
 	}
 }
 
@@ -51,6 +143,29 @@ void UMovieSceneMaterialParameterCollectionSystem::OnLink()
 
 	ScalarParameterStorage = Linker->PreAnimatedState.GetOrCreateStorage<FPreAnimatedScalarMaterialParameterStorage>();
 	VectorParameterStorage = Linker->PreAnimatedState.GetOrCreateStorage<FPreAnimatedVectorMaterialParameterStorage>();
+
+	UMovieSceneEntityGroupingSystem* GroupingSystem = Linker->LinkSystem<UMovieSceneEntityGroupingSystem>();
+
+	FEntityComponentFilter Filter;
+	Filter.None({ FBuiltInComponentTypes::Get()->Tags.ImportedEntity });
+
+	GroupingKey = GroupingSystem->AddGrouping(
+		FMPCGroupingPolicy(),
+		MoveTemp(Filter),
+		FMovieSceneTracksComponentTypes::Get()->MPC
+	);
+}
+
+void UMovieSceneMaterialParameterCollectionSystem::OnUnlink()
+{
+	using namespace UE::MovieScene;
+
+	UMovieSceneEntityGroupingSystem* GroupingSystem = Linker->FindSystem<UMovieSceneEntityGroupingSystem>();
+	if (ensure(GroupingSystem))
+	{
+		GroupingSystem->RemoveGrouping(GroupingKey);
+	}
+	GroupingKey = FEntityGroupingPolicyKey();
 }
 
 void UMovieSceneMaterialParameterCollectionSystem::OnRun(FSystemTaskPrerequisites& InPrerequisites, FSystemSubsequentTasks& Subsequents)
diff --git a/Engine/Source/Runtime/MovieSceneTracks/Public/Systems/MovieSceneMaterialParameterCollectionSystem.h b/Engine/Source/Runtime/MovieSceneTracks/Public/Systems/MovieSceneMaterialParameterCollectionSystem.h
--- a/Engine/Source/Runtime/MovieSceneTracks/Public/Systems/MovieSceneMaterialParameterCollectionSystem.h
+++ b/Engine/Source/Runtime/MovieSceneTracks/Public/Systems/MovieSceneMaterialParameterCollectionSystem.h
@@ -3,6 +3,7 @@
 #pragma once
 
 #include "EntitySystem/MovieSceneEntitySystem.h"
+#include "EntitySystem/BuiltInComponentTypes.h"
 #include "Evaluation/PreAnimatedState/MovieScenePreAnimatedStateStorage.h"
 
 #include "Systems/MovieSceneMaterialSystem.h"
@@ -35,6 +36,7 @@ public:
 private:
 
 	virtual void OnLink() override;
+	virtual void OnUnlink() override;
 	virtual void OnRun(FSystemTaskPrerequisites& InPrerequisites, FSystemSubsequentTasks& Subsequents) override;
 
 private:
@@ -43,4 +45,6 @@ private:
 	TSharedPtr<UE::MovieScene::FPreAnimatedScalarMaterialParameterStorage> ScalarParameterStorage;
 	/** Holds pre-animated values for vector or color MPC values */
 	TSharedPtr<UE::MovieScene::FPreAnimatedVectorMaterialParameterStorage> VectorParameterStorage;
+	/** Registered key for the grouping policy that groups MPC tracks together */
+	UE::MovieScene::FEntityGroupingPolicyKey GroupingKey;
 };
