Sequencer: Set the localized display name for the Transform track

diff --git a/Engine/Source/Runtime/MovieSceneTracks/Private/Tracks/MovieScene3DTransformTrack.cpp b/Engine/Source/Runtime/MovieSceneTracks/Private/Tracks/MovieScene3DTransformTrack.cpp
--- a/Engine/Source/Runtime/MovieSceneTracks/Private/Tracks/MovieScene3DTransformTrack.cpp
+++ b/Engine/Source/Runtime/MovieSceneTracks/Private/Tracks/MovieScene3DTransformTrack.cpp
@@ -38,6 +38,13 @@ UMovieSceneSection* UMovieScene3DTransformTrack::CreateNewSection()
 	return NewObject<UMovieScene3DTransformSection>(this, NAME_None, RF_Transactional);
 }
 
+#if WITH_EDITORONLY_DATA
+FText UMovieScene3DTransformTrack::GetDisplayName() const
+{
+	return NSLOCTEXT("MovieScene3DTransformTrack", "DefaultLabel", "Transform");
+}
+#endif
+
 TSubclassOf<UMovieSceneBlenderSystem> UMovieScene3DTransformTrack::GetBlenderSystem() const
 {
 	if (BlenderSystemClass == nullptr)
diff --git a/Engine/Source/Runtime/MovieSceneTracks/Public/Tracks/MovieScene3DTransformTrack.h b/Engine/Source/Runtime/MovieSceneTracks/Public/Tracks/MovieScene3DTransformTrack.h
--- a/Engine/Source/Runtime/MovieSceneTracks/Public/Tracks/MovieScene3DTransformTrack.h
+++ b/Engine/Source/Runtime/MovieSceneTracks/Public/Tracks/MovieScene3DTransformTrack.h
@@ -69,6 +69,7 @@ public:
 
 #if WITH_EDITORONLY_DATA
 	virtual bool CanRename() const override { return true; }
+	virtual FText GetDisplayName() const override;
 #endif
 
 	// IMovieSceneBlenderSystemSupport
