Sequencer: Force 'SequencerActor' tag to be added before calling AActor::FinishSpawning to fix regression.

diff --git a/Engine/Source/Runtime/MovieSceneTracks/Private/Bindings/MovieSceneSpawnableActorBinding.cpp b/Engine/Source/Runtime/MovieSceneTracks/Private/Bindings/MovieSceneSpawnableActorBinding.cpp
--- a/Engine/Source/Runtime/MovieSceneTracks/Private/Bindings/MovieSceneSpawnableActorBinding.cpp
+++ b/Engine/Source/Runtime/MovieSceneTracks/Private/Bindings/MovieSceneSpawnableActorBinding.cpp
@@ -133,6 +133,18 @@ UObject* UMovieSceneSpawnableActorBindingBase::SpawnObjectInternal(UWorld* World
 			}
 		}
 	}
+	
+	// The below block duplicates code in UMovieSceneSpawnableBinding, but this ensures that at least for this
+	// type of binding, these tags are added before calling FinishSpawning, as some client code may check for this tag
+	// in component initialization.
+	{
+	#if WITH_EDITOR
+		SpawnedActor->bIsEditorPreviewActor = false;
+	#endif
+
+		// tag this actor so we know it was spawned by sequencer
+		SpawnedActor->Tags.AddUnique(SequencerActorTag);
+	}
 
 	const bool bIsDefaultTransform = true;
 	SpawnedActor->FinishSpawning(SpawnTransformToUse, bIsDefaultTransform);
