Sequencer: Ensure that hierarchical bias is initialized before the hierarchical easing system

diff --git a/Engine/Source/Runtime/MovieSceneTracks/Private/Systems/MovieSceneHierarchicalBiasSystem.cpp b/Engine/Source/Runtime/MovieSceneTracks/Private/Systems/MovieSceneHierarchicalBiasSystem.cpp
--- a/Engine/Source/Runtime/MovieSceneTracks/Private/Systems/MovieSceneHierarchicalBiasSystem.cpp
+++ b/Engine/Source/Runtime/MovieSceneTracks/Private/Systems/MovieSceneHierarchicalBiasSystem.cpp
@@ -8,6 +8,7 @@
 #include "EntitySystem/MovieSceneEntitySystemLinker.h"
 #include "EntitySystem/MovieSceneEntityGroupingSystem.h"
 #include "EntitySystem/MovieScenePreAnimatedStateSystem.h"
+#include "Systems/WeightAndEasingEvaluatorSystem.h"
 
 #include UE_INLINE_GENERATED_CPP_BY_NAME(MovieSceneHierarchicalBiasSystem)
 
@@ -181,6 +182,7 @@ UMovieSceneHierarchicalBiasSystem::UMovieSceneHierarchicalBiasSystem(const FObje
 
 		DefineComponentProducer(GetClass(), BuiltInComponents->Tags.Ignored);
 		DefineComponentProducer(GetClass(), BuiltInComponents->HierarchicalBlendTarget);
+		DefineImplicitPrerequisite(GetClass(), UMovieSceneHierarchicalEasingInstantiatorSystem::StaticClass());
 	}
 }
 
