Sequencer: Fix custom transform track names

diff --git a/Engine/Source/Runtime/MovieSceneTracks/Private/Tracks/MovieScene3DTransformTrack.cpp b/Engine/Source/Runtime/MovieSceneTracks/Private/Tracks/MovieScene3DTransformTrack.cpp
--- a/Engine/Source/Runtime/MovieSceneTracks/Private/Tracks/MovieScene3DTransformTrack.cpp
+++ b/Engine/Source/Runtime/MovieSceneTracks/Private/Tracks/MovieScene3DTransformTrack.cpp
@@ -41,7 +41,14 @@ UMovieSceneSection* UMovieScene3DTransformTrack::CreateNewSection()
 #if WITH_EDITORONLY_DATA
 FText UMovieScene3DTransformTrack::GetDisplayName() const
 {
-	return NSLOCTEXT("MovieScene3DTransformTrack", "DefaultLabel", "Transform");
+	FText CustomDisplayName = Super::GetDisplayName();
+
+	static FName Transform("Transform");
+	if (CustomDisplayName.ToString() == Transform)
+	{
+		return NSLOCTEXT("MovieScene3DTransformTrack", "DefaultLabel", "Transform");
+	}
+	return CustomDisplayName;
 }
 #endif
 
