Sequencer: Do not create entities for parameters that have no data associated with them

diff --git a/Engine/Source/Runtime/MovieSceneTracks/Private/Sections/MovieSceneParameterSection.cpp b/Engine/Source/Runtime/MovieSceneTracks/Private/Sections/MovieSceneParameterSection.cpp
--- a/Engine/Source/Runtime/MovieSceneTracks/Private/Sections/MovieSceneParameterSection.cpp
+++ b/Engine/Source/Runtime/MovieSceneTracks/Private/Sections/MovieSceneParameterSection.cpp
@@ -750,32 +750,56 @@ void UMovieSceneParameterSection::ExternalPopulateEvaluationField(const TRange<F
 
 	for (int32 Index = 0; Index < NumScalarID; ++Index)
 	{
-		const int32 EntityIndex = OutFieldBuilder->FindOrAddEntity(this, EncodeEntityID(Index, 0));
-		OutFieldBuilder->AddPersistentEntity(EffectiveRange, EntityIndex, MetaDataIndex);
+		if (ScalarParameterNamesAndCurves[Index].ParameterCurve.HasAnyData())
+		{
+			const int32 EntityIndex = OutFieldBuilder->FindOrAddEntity(this, EncodeEntityID(Index, 0));
+			OutFieldBuilder->AddPersistentEntity(EffectiveRange, EntityIndex, MetaDataIndex);
+		}
 	}
 	for (int32 Index = 0; Index < NumBoolID; ++Index)
 	{
-		const int32 EntityIndex = OutFieldBuilder->FindOrAddEntity(this, EncodeEntityID(Index, 1));
-		OutFieldBuilder->AddPersistentEntity(EffectiveRange, EntityIndex, MetaDataIndex);
+		if (BoolParameterNamesAndCurves[Index].ParameterCurve.HasAnyData())
+		{
+			const int32 EntityIndex = OutFieldBuilder->FindOrAddEntity(this, EncodeEntityID(Index, 1));
+			OutFieldBuilder->AddPersistentEntity(EffectiveRange, EntityIndex, MetaDataIndex);
+		}
 	}
 	for (int32 Index = 0; Index < NumVector2DID; ++Index)
 	{
-		const int32 EntityIndex = OutFieldBuilder->FindOrAddEntity(this, EncodeEntityID(Index, 2));
-		OutFieldBuilder->AddPersistentEntity(EffectiveRange, EntityIndex, MetaDataIndex);
+		const FVector2DParameterNameAndCurves& Vector2D = Vector2DParameterNamesAndCurves[Index];
+		if (Vector2D.XCurve.HasAnyData() || Vector2D.YCurve.HasAnyData())
+		{
+			const int32 EntityIndex = OutFieldBuilder->FindOrAddEntity(this, EncodeEntityID(Index, 2));
+			OutFieldBuilder->AddPersistentEntity(EffectiveRange, EntityIndex, MetaDataIndex);
+		}
 	}
 	for (int32 Index = 0; Index < NumVectorID; ++Index)
 	{
-		const int32 EntityIndex = OutFieldBuilder->FindOrAddEntity(this, EncodeEntityID(Index, 3));
-		OutFieldBuilder->AddPersistentEntity(EffectiveRange, EntityIndex, MetaDataIndex);
+		const FVectorParameterNameAndCurves& Vector = VectorParameterNamesAndCurves[Index];
+		if (Vector.XCurve.HasAnyData() || Vector.YCurve.HasAnyData() || Vector.ZCurve.HasAnyData())
+		{
+			const int32 EntityIndex = OutFieldBuilder->FindOrAddEntity(this, EncodeEntityID(Index, 3));
+			OutFieldBuilder->AddPersistentEntity(EffectiveRange, EntityIndex, MetaDataIndex);
+		}
 	}
 	for (int32 Index = 0; Index < NumColorID; ++Index)
 	{
-		const int32 EntityIndex = OutFieldBuilder->FindOrAddEntity(this, EncodeEntityID(Index, 4));
-		OutFieldBuilder->AddPersistentEntity(EffectiveRange, EntityIndex, MetaDataIndex);
+		const FColorParameterNameAndCurves& Color = ColorParameterNamesAndCurves[Index];
+		if (Color.RedCurve.HasAnyData() || Color.GreenCurve.HasAnyData() || Color.BlueCurve.HasAnyData() || Color.AlphaCurve.HasAnyData())
+		{
+			const int32 EntityIndex = OutFieldBuilder->FindOrAddEntity(this, EncodeEntityID(Index, 4));
+			OutFieldBuilder->AddPersistentEntity(EffectiveRange, EntityIndex, MetaDataIndex);
+		}
 	}
 	for (int32 Index = 0; Index < NumTransformID; ++Index)
 	{
-		const int32 EntityIndex = OutFieldBuilder->FindOrAddEntity(this, EncodeEntityID(Index, 5));
-		OutFieldBuilder->AddPersistentEntity(EffectiveRange, EntityIndex, MetaDataIndex);
+		const FTransformParameterNameAndCurves& Transform = TransformParameterNamesAndCurves[Index];
+		if (Transform.Translation[0].HasAnyData() || Transform.Translation[1].HasAnyData() || Transform.Translation[2].HasAnyData() ||
+			Transform.Rotation[0].HasAnyData() || Transform.Rotation[1].HasAnyData() || Transform.Rotation[2].HasAnyData() ||
+			Transform.Scale[0].HasAnyData() || Transform.Scale[1].HasAnyData() || Transform.Scale[2].HasAnyData())
+		{
+			const int32 EntityIndex = OutFieldBuilder->FindOrAddEntity(this, EncodeEntityID(Index, 5));
+			OutFieldBuilder->AddPersistentEntity(EffectiveRange, EntityIndex, MetaDataIndex);
+		}
 	}
 }
