Sequencer: Added GC tests for the compiled data manager

Compiled data must not maintain any strong references back to the source data in order to allow the assets to be freely available for GC. This is an unenforceable constraint that is easy to contravene (both conditions and timewarp have both inadvertently done this).

Despite this, it is good to have a test that we can use to verify these assumptions for the known offenders.

diff --git a/Engine/Source/Runtime/MovieScene/Private/Compilation/MovieSceneCompiledDataManager.cpp b/Engine/Source/Runtime/MovieScene/Private/Compilation/MovieSceneCompiledDataManager.cpp
--- a/Engine/Source/Runtime/MovieScene/Private/Compilation/MovieSceneCompiledDataManager.cpp
+++ b/Engine/Source/Runtime/MovieScene/Private/Compilation/MovieSceneCompiledDataManager.cpp
@@ -47,6 +47,8 @@ TAutoConsoleVariable<bool> CVarAddKeepStateDeterminismFences(
 		 "at the cost of an extra evaluation on frames that cross over KeepState sections' end time.\n"),
 	ECVF_Default);
 
+TSet<UMovieSceneCompiledDataManager*> UMovieSceneCompiledDataManager::ActiveManagers;
+
 
 IMovieSceneModule& GetMovieSceneModule()
 {
@@ -491,6 +493,24 @@ UMovieSceneCompiledDataManager::UMovieSceneCompiledDataManager()
 	if (!HasAnyFlags(RF_ClassDefaultObject))
 	{
 		FCoreUObjectDelegates::OnPackageReloaded.AddWeakLambda(this, OnPackageReloaded);
+		ActiveManagers.Add(this);
+	}
+}
+
+void UMovieSceneCompiledDataManager::BeginDestroy()
+{
+	ActiveManagers.Remove(this);
+	Super::BeginDestroy();
+}
+
+void UMovieSceneCompiledDataManager::ReportSequenceDestroyed(UMovieSceneSequence* InSequence)
+{
+	if (!GExitPurge)
+	{
+		for (UMovieSceneCompiledDataManager* Manager : ActiveManagers)
+		{
+			Manager->Reset(InSequence);
+		}
 	}
 }
 
diff --git a/Engine/Source/Runtime/MovieScene/Private/MovieSceneSequence.cpp b/Engine/Source/Runtime/MovieScene/Private/MovieSceneSequence.cpp
--- a/Engine/Source/Runtime/MovieScene/Private/MovieSceneSequence.cpp
+++ b/Engine/Source/Runtime/MovieScene/Private/MovieSceneSequence.cpp
@@ -140,12 +140,7 @@ void UMovieSceneSequence::BeginDestroy()
 
 	if (!GExitPurge && !HasAnyFlags(RF_ClassDefaultObject))
 	{
-		UMovieSceneCompiledDataManager::GetPrecompiledData()->Reset(this);
-
-#if WITH_EDITOR
-		UMovieSceneCompiledDataManager::GetPrecompiledData(EMovieSceneServerClientMask::Client)->Reset(this);
-		UMovieSceneCompiledDataManager::GetPrecompiledData(EMovieSceneServerClientMask::Server)->Reset(this);
-#endif
+		UMovieSceneCompiledDataManager::ReportSequenceDestroyed(this);
 	}
 }
 
diff --git a/Engine/Source/Runtime/MovieScene/Public/Compilation/MovieSceneCompiledDataManager.h b/Engine/Source/Runtime/MovieScene/Public/Compilation/MovieSceneCompiledDataManager.h
--- a/Engine/Source/Runtime/MovieScene/Public/Compilation/MovieSceneCompiledDataManager.h
+++ b/Engine/Source/Runtime/MovieScene/Public/Compilation/MovieSceneCompiledDataManager.h
@@ -190,6 +190,8 @@ public:
 	static MOVIESCENE_API UMovieSceneCompiledDataManager* GetPrecompiledData();
 #endif
 
+	static MOVIESCENE_API void ReportSequenceDestroyed(UMovieSceneSequence* InSequence);
+
 	MOVIESCENE_API UMovieSceneCompiledData* MakeCompiledData(UMovieSceneSequence* Sequence) const;
 
 	MOVIESCENE_API EMovieSceneServerClientMask GetNetworkMask() const { return NetworkMask; }
@@ -296,6 +298,8 @@ private:
 
 	MOVIESCENE_API void DestroyData(FMovieSceneCompiledDataID DataID);
 
+	virtual void BeginDestroy() override;
+
 private:
 
 	FMovieSceneCompiledDataEntry* GetEntryPtr(FMovieSceneCompiledDataID DataID)
@@ -339,6 +343,8 @@ private:
 	uint32 ReallocationVersion;
 
 	EMovieSceneServerClientMask NetworkMask;
+
+	static TSet<UMovieSceneCompiledDataManager*> ActiveManagers;
 };
 
 
diff --git a/Engine/Source/Runtime/MovieSceneTracks/Private/Tests/MovieSceneCompiledDataManagerTests.cpp b/Engine/Source/Runtime/MovieSceneTracks/Private/Tests/MovieSceneCompiledDataManagerTests.cpp
new file mode 100644
--- /dev/null
+++ b/Engine/Source/Runtime/MovieSceneTracks/Private/Tests/MovieSceneCompiledDataManagerTests.cpp
@@ -0,0 +1,168 @@
+// Copyright Epic Games, Inc. All Rights Reserved.
+
+#include "CoreTypes.h"
+#include "Compilation/MovieSceneCompiledDataManager.h"
+#include "Conditions/MovieScenePlatformCondition.h"
+#include "Engine/Engine.h"
+#include "Misc/AutomationTest.h"
+#include "Sections/MovieSceneFadeSection.h"
+#include "Sections/MovieSceneSubSection.h"
+#include "Sections/MovieSceneTimeWarpSection.h"
+#include "Tests/MovieSceneTestDataBuilders.h"
+#include "Tracks/MovieSceneFadeTrack.h"
+#include "Tracks/MovieSceneSubTrack.h"
+#include "Tracks/MovieSceneTimeWarpTrack.h"
+#include "UObject/Package.h"
+#include "Variants/MovieScenePlayRateCurve.h"
+
+#if WITH_DEV_AUTOMATION_TESTS
+
+#define LOCTEXT_NAMESPACE "MovieSceneCompiledDataManagerTests"
+
+IMPLEMENT_SIMPLE_AUTOMATION_TEST(FMovieSceneCompiledDataManagerGCReferenceTest,
+		"System.Engine.Sequencer.Compilation.GCReferenceTests", 
+		EAutomationTestFlags::EditorContext | EAutomationTestFlags::EngineFilter)
+bool FMovieSceneCompiledDataManagerGCReferenceTest::RunTest(const FString& Parameters)
+{
+	using namespace UE::MovieScene;
+	using namespace UE::MovieScene::Test;
+
+	TStrongObjectPtr<UMovieSceneCompiledDataManager> CompiledData(
+		NewObject<UMovieSceneCompiledDataManager>(GetTransientPackage(), "GCReferenceTest")
+	);
+
+	UMovieSceneTimeWarpSection* TimeWarpSection = nullptr;
+	UMovieSceneSubSection*      SubSection      = nullptr;
+	UMovieSceneFadeSection*     FadeSection     = nullptr;
+
+	FSequenceBuilder RootSequenceBuilder;
+	{
+		RootSequenceBuilder
+		.AddRootTrack<UMovieSceneTimeWarpTrack>()
+			.AddSection(0, 96000)
+			.Assign(TimeWarpSection)
+			.Pop()
+		.Pop()
+		.AddRootTrack<UMovieSceneFadeTrack>()
+			.AddSection(0, 96000)
+			.Assign(FadeSection)
+			.Pop()
+		.Pop()
+		.AddRootTrack<UMovieSceneSubTrack>()
+			.AddSection(0, 2400)
+			.Assign(SubSection)
+			.Pop()
+		.Pop();
+	}
+
+	// Add a condition to ensure that they do not hold a strong reference in the compiled data
+	FadeSection->ConditionContainer.Condition = NewObject<UMovieScenePlatformCondition>(FadeSection);
+
+	// Build an empty sub sequence - we only need this so we can create a hierarchy and shove
+	//     its timewarp into the compiled data manager for testing
+	FSequenceBuilder SubSequenceBuilder;
+
+	TWeakObjectPtr<UMovieSceneSequence> WeakRootSequence = RootSequenceBuilder.Sequence;
+	TWeakObjectPtr<UMovieSceneSequence> WeakSubSequence  = SubSequenceBuilder.Sequence;
+
+	TWeakObjectPtr<UMovieScenePlayRateCurve> WeakRootTimeWarp;
+	TWeakObjectPtr<UMovieScenePlayRateCurve> WeakSubTimeWarp;
+
+	// Add root time warp
+	{
+		UMovieScenePlayRateCurve* PlayRate = NewObject<UMovieScenePlayRateCurve>(TimeWarpSection);
+
+		WeakRootTimeWarp = PlayRate;
+
+		TMovieSceneChannelData<FMovieSceneDoubleValue> ChannelData = PlayRate->PlayRate.GetData();
+		ChannelData.AddKey(0, FMovieSceneDoubleValue(1.0));
+		ChannelData.AddKey(24000, FMovieSceneDoubleValue(0.5));
+		ChannelData.AddKey(48000, FMovieSceneDoubleValue(1.0));
+
+		TimeWarpSection->TimeWarp.Set(PlayRate);
+	}
+
+	// Setup the subsequence and time warp
+	{
+		SubSection->SetSequence(SubSequenceBuilder.Sequence);
+
+		UMovieScenePlayRateCurve* SubPlayRate = NewObject<UMovieScenePlayRateCurve>(SubSection);
+
+		WeakSubTimeWarp = SubPlayRate;
+
+		TMovieSceneChannelData<FMovieSceneDoubleValue> ChannelData = SubPlayRate->PlayRate.GetData();
+		ChannelData.AddKey(0, FMovieSceneDoubleValue(0.25));
+
+		SubSection->Parameters.TimeScale.Set(SubPlayRate);
+	}
+
+
+	// Compile the sequence, run a GC (which should clean up all the sequences), and verify that the compiled data manager has removed everything
+
+	if (WeakRootSequence.Get()    == nullptr
+		|| WeakSubSequence.Get()  == nullptr
+		|| WeakRootTimeWarp.Get() == nullptr
+		|| WeakSubTimeWarp.Get()  == nullptr)
+	{
+		AddError(TEXT("Test failed to initialize weak pointers correctly."));
+		return false;
+	}
+
+	CompiledData->Compile(RootSequenceBuilder.Sequence);
+
+	FMovieSceneCompiledDataID RootDataID = CompiledData->FindDataID(RootSequenceBuilder.Sequence);
+	FMovieSceneCompiledDataID SubDataID  = CompiledData->FindDataID(SubSequenceBuilder.Sequence);
+
+	if (!RootDataID.IsValid() || !SubDataID.IsValid())
+	{
+		AddError(TEXT("Sequence was not compiled successfully."));
+		return false;
+	}
+
+	if (CompiledData->FindHierarchy(RootDataID) == nullptr)
+	{
+		AddError(TEXT("Sequence did not compile a hierarchy when it should have."));
+		return false;
+	}
+
+	CollectGarbage(GARBAGE_COLLECTION_KEEPFLAGS, true);
+
+	if (WeakRootSequence.Get()    != nullptr
+		|| WeakSubSequence.Get()  != nullptr
+		|| WeakRootTimeWarp.Get() != nullptr
+		|| WeakSubTimeWarp.Get()  != nullptr)
+	{
+		if (UObject* Object = WeakRootSequence.Get())
+		{
+			GEngine->Exec(nullptr, *FString::Printf(TEXT("obj refs name=\"%s\""), *Object->GetPathName()));
+		}
+		if (UObject* Object = WeakSubSequence.Get())
+		{
+			GEngine->Exec(nullptr, *FString::Printf(TEXT("obj refs name=\"%s\""), *Object->GetPathName()));
+		}
+		if (UObject* Object = WeakRootTimeWarp.Get())
+		{
+			GEngine->Exec(nullptr, *FString::Printf(TEXT("obj refs name=\"%s\""), *Object->GetPathName()));
+		}
+		if (UObject* Object = WeakSubTimeWarp.Get())
+		{
+			GEngine->Exec(nullptr, *FString::Printf(TEXT("obj refs name=\"%s\""), *Object->GetPathName()));
+		}
+
+		AddError(TEXT("Objects were unexpectedly still alive after GC please see logs for references."));
+		return false;
+	}
+
+	if (CompiledData->FindHierarchy(RootDataID) != nullptr)
+	{
+		AddError(TEXT("Compiled data was not correctly cleaned up."));
+		return false;
+	}
+
+	return true;
+}
+
+#undef LOCTEXT_NAMESPACE
+
+#endif // WITH_DEV_AUTOMATION_TESTS
+
