Sequencer: enable missing support for completion mode on double vector sections, and set the default completion mode to the same value as other sections

diff --git a/Engine/Source/Runtime/MovieSceneTracks/Private/Sections/MovieSceneVectorSection.cpp b/Engine/Source/Runtime/MovieSceneTracks/Private/Sections/MovieSceneVectorSection.cpp
--- a/Engine/Source/Runtime/MovieSceneTracks/Private/Sections/MovieSceneVectorSection.cpp
+++ b/Engine/Source/Runtime/MovieSceneTracks/Private/Sections/MovieSceneVectorSection.cpp
@@ -281,6 +281,8 @@ UMovieSceneDoubleVectorSection::UMovieSceneDoubleVectorSection(const FObjectInit
 {
 	ChannelsUsed = 0;
 	bSupportsInfiniteRange = true;
+
+	EvalOptions.EnableAndSetCompletionMode(EMovieSceneCompletionMode::ProjectDefault);
 	BlendType = EMovieSceneBlendType::Absolute;
 }
 
