Sequencer: add missing call to new section's OnAddedToTrack when a property track auto-creates it

diff --git a/Engine/Source/Runtime/MovieSceneTracks/Private/Tracks/MovieScenePropertyTrack.cpp b/Engine/Source/Runtime/MovieSceneTracks/Private/Tracks/MovieScenePropertyTrack.cpp
--- a/Engine/Source/Runtime/MovieSceneTracks/Private/Tracks/MovieScenePropertyTrack.cpp
+++ b/Engine/Source/Runtime/MovieSceneTracks/Private/Tracks/MovieScenePropertyTrack.cpp
@@ -348,7 +348,7 @@ UMovieSceneSection* UMovieScenePropertyTrack::FindOrAddSection(FFrameNumber Time
 	NewSection->SetFlags(RF_Transactional);
 	NewSection->SetRange(TRange<FFrameNumber>::Inclusive(Time, Time));
 
-	Sections.Add(NewSection);
+	UMovieSceneTrack::AddSection(FSectionParameter(*NewSection));
 	
 	bSectionAdded = true;
 
