Control Rig: Android build break fixes for 40302964

#rnx

diff --git a/Engine/Plugins/Animation/ControlRig/Source/ControlRig/Private/ModularRigController.cpp b/Engine/Plugins/Animation/ControlRig/Source/ControlRig/Private/ModularRigController.cpp
--- a/Engine/Plugins/Animation/ControlRig/Source/ControlRig/Private/ModularRigController.cpp
+++ b/Engine/Plugins/Animation/ControlRig/Source/ControlRig/Private/ModularRigController.cpp
@@ -1952,7 +1952,7 @@ bool UModularRigController::IsNameAvailable(const FRigName& InDesiredName, FStri
 		{
 			continue;
 		}
-		if(Module.Name == DesiredName)
+		if(Module.Name == DesiredName.GetFName())
 		{
 			if(OutErrorMessage)
 			{
diff --git a/Engine/Plugins/Animation/ControlRig/Source/ControlRig/Private/Rigs/RigHierarchy.cpp b/Engine/Plugins/Animation/ControlRig/Source/ControlRig/Private/Rigs/RigHierarchy.cpp
--- a/Engine/Plugins/Animation/ControlRig/Source/ControlRig/Private/Rigs/RigHierarchy.cpp
+++ b/Engine/Plugins/Animation/ControlRig/Source/ControlRig/Private/Rigs/RigHierarchy.cpp
@@ -9562,7 +9562,7 @@ void URigHierarchy::PatchElementMetadata(const TMap<FRigHierarchyModulePath, FNa
 			TMap<FName, FRigBaseMetadata*> OldMetadataMap;
 			Swap(OldMetadataMap, MetadataStorage.MetadataMap);
 
-			for(const TPair<FName, FRigBaseMetadata*>& Pair : MetadataStorage.MetadataMap)
+			for(const TPair<FName, FRigBaseMetadata*>& Pair : OldMetadataMap)
 			{
 				FRigHierarchyModulePath ModulePath(Pair.Key);
 				if(ModulePath.ConvertToModuleNameFormatInline(&InModulePathToName))
diff --git a/Engine/Plugins/Animation/ControlRig/Source/ControlRig/Public/Rigs/RigHierarchyDefines.h b/Engine/Plugins/Animation/ControlRig/Source/ControlRig/Public/Rigs/RigHierarchyDefines.h
--- a/Engine/Plugins/Animation/ControlRig/Source/ControlRig/Public/Rigs/RigHierarchyDefines.h
+++ b/Engine/Plugins/Animation/ControlRig/Source/ControlRig/Public/Rigs/RigHierarchyDefines.h
@@ -1635,6 +1635,11 @@ public:
 		return ModulePath.ToLower() == InOther.ModulePath.ToLower();
 	}
 
+	bool operator==(const FString& InOther) const
+	{
+		return ModulePath.ToLower() == InOther.ToLower();
+	}
+
 	bool Split(FStringView* OutModuleName, FStringView* OutElementName) const;
 	bool Split(FString* OutModuleName, FString* OutElementName) const;
 
diff --git a/Engine/Plugins/Animation/ControlRig/Source/ControlRig/Public/Sequencer/MovieSceneControlRigParameterSection.h b/Engine/Plugins/Animation/ControlRig/Source/ControlRig/Public/Sequencer/MovieSceneControlRigParameterSection.h
--- a/Engine/Plugins/Animation/ControlRig/Source/ControlRig/Public/Sequencer/MovieSceneControlRigParameterSection.h
+++ b/Engine/Plugins/Animation/ControlRig/Source/ControlRig/Public/Sequencer/MovieSceneControlRigParameterSection.h
@@ -106,16 +106,29 @@ struct CONTROLRIG_API FEnumParameterNameAndCurve : public FBaseParameterNameAndV
 {
 	GENERATED_USTRUCT_BODY()
 
+	FEnumParameterNameAndCurve()
+	{}
+
+	FEnumParameterNameAndCurve(FName InParameterName)
+	: FBaseParameterNameAndValue(InParameterName)
+	{}
+
 	UPROPERTY()
 	FMovieSceneByteChannel ParameterCurve;
 };
 
-
 USTRUCT()
 struct CONTROLRIG_API FIntegerParameterNameAndCurve : public FBaseParameterNameAndValue
 {
 	GENERATED_USTRUCT_BODY()
 
+	FIntegerParameterNameAndCurve()
+	{}
+
+	FIntegerParameterNameAndCurve(FName InParameterName)
+	: FBaseParameterNameAndValue(InParameterName)
+	{}
+	
 	UPROPERTY()
 	FMovieSceneIntegerChannel ParameterCurve;
 };
diff --git a/Engine/Source/Runtime/MovieSceneTracks/Public/Sections/MovieSceneParameterSection.h b/Engine/Source/Runtime/MovieSceneTracks/Public/Sections/MovieSceneParameterSection.h
--- a/Engine/Source/Runtime/MovieSceneTracks/Public/Sections/MovieSceneParameterSection.h
+++ b/Engine/Source/Runtime/MovieSceneTracks/Public/Sections/MovieSceneParameterSection.h
@@ -184,6 +184,13 @@ struct FBoolParameterNameAndCurve : public FBaseParameterNameAndValue
 {
 	GENERATED_USTRUCT_BODY()
 
+	FBoolParameterNameAndCurve()
+	{}
+
+	MOVIESCENETRACKS_API FBoolParameterNameAndCurve(FName InParameterName)
+	: FBaseParameterNameAndValue(InParameterName)
+	{}
+
 	/** The curve which contains the animation data for the scalar parameter. */
 	UPROPERTY()
 	FMovieSceneBoolChannel ParameterCurve;
@@ -197,6 +204,13 @@ struct FScalarParameterNameAndCurve : public FBaseParameterNameAndValue
 {
 	GENERATED_USTRUCT_BODY()
 
+	FScalarParameterNameAndCurve()
+	{}
+
+	MOVIESCENETRACKS_API FScalarParameterNameAndCurve(FName InParameterName)
+	: FBaseParameterNameAndValue(InParameterName)
+	{}
+
 	/** The curve which contains the animation data for the scalar parameter. */
 	UPROPERTY()
 	FMovieSceneFloatChannel ParameterCurve;
@@ -210,6 +224,13 @@ struct FVector2DParameterNameAndCurves : public FBaseParameterNameAndValue
 {
 	GENERATED_USTRUCT_BODY()
 
+	FVector2DParameterNameAndCurves()
+	{}
+
+	MOVIESCENETRACKS_API FVector2DParameterNameAndCurves(FName InParameterName)
+	: FBaseParameterNameAndValue(InParameterName)
+	{}
+
 	/** The curve which contains the animation data for the x component of the vector parameter. */
 	UPROPERTY()
 	FMovieSceneFloatChannel XCurve;
@@ -229,6 +250,13 @@ struct FVectorParameterNameAndCurves : public FBaseParameterNameAndValue
 {
 	GENERATED_USTRUCT_BODY()
 
+	FVectorParameterNameAndCurves()
+	{}
+
+	MOVIESCENETRACKS_API FVectorParameterNameAndCurves(FName InParameterName)
+	: FBaseParameterNameAndValue(InParameterName)
+	{}
+
 	/** The curve which contains the animation data for the x component of the vector parameter. */
 	UPROPERTY()
 	FMovieSceneFloatChannel XCurve;
@@ -251,6 +279,13 @@ struct FColorParameterNameAndCurves : public FBaseParameterNameAndValue
 {
 	GENERATED_USTRUCT_BODY()
 
+	FColorParameterNameAndCurves()
+	{}
+
+	MOVIESCENETRACKS_API FColorParameterNameAndCurves(FName InParameterName)
+	: FBaseParameterNameAndValue(InParameterName)
+	{}
+
 	/** The curve which contains the animation data for the red component of the color parameter. */
 	UPROPERTY()
 	FMovieSceneFloatChannel RedCurve;
@@ -276,6 +311,13 @@ struct FTransformParameterNameAndCurves : public FBaseParameterNameAndValue
 {
 	GENERATED_USTRUCT_BODY()
 
+	FTransformParameterNameAndCurves()
+	{}
+
+	MOVIESCENETRACKS_API FTransformParameterNameAndCurves(FName InParameterName)
+	: FBaseParameterNameAndValue(InParameterName)
+	{}
+
 	/** Translation curves */
 	UPROPERTY()
 	FMovieSceneFloatChannel Translation[3];
