Sequencer: Fix bound object no longer valid

diff --git a/Engine/Source/Runtime/MovieSceneTracks/Private/Tracks/MovieScenePropertyTrack.cpp b/Engine/Source/Runtime/MovieSceneTracks/Private/Tracks/MovieScenePropertyTrack.cpp
--- a/Engine/Source/Runtime/MovieSceneTracks/Private/Tracks/MovieScenePropertyTrack.cpp
+++ b/Engine/Source/Runtime/MovieSceneTracks/Private/Tracks/MovieScenePropertyTrack.cpp
@@ -74,6 +74,11 @@ FText UMovieScenePropertyTrack::GetDisplayNameToolTipText(const FMovieSceneLabel
 	const TArrayView<TWeakObjectPtr<>> FoundBoundObjects = LabelParams.Player->FindBoundObjects(LabelParams.BindingID, LabelParams.SequenceID);
 	for (const TWeakObjectPtr<> BoundObject : FoundBoundObjects)
 	{
+		if (!BoundObject.IsValid())
+		{
+			continue;
+		}
+
 		FTrackInstancePropertyBindings InstancePropertyBinding(GetPropertyName(), GetPropertyPath().ToString());
 		if (FProperty* BoundProperty = InstancePropertyBinding.GetProperty(*BoundObject))
 		{
@@ -111,6 +116,11 @@ FSlateColor UMovieScenePropertyTrack::GetLabelColor(const FMovieSceneLabelParams
 	const TArrayView<TWeakObjectPtr<>> FoundBoundObjects = LabelParams.Player->FindBoundObjects(LabelParams.BindingID, LabelParams.SequenceID);
 	for (const TWeakObjectPtr<> BoundObject : FoundBoundObjects)
 	{
+		if (!BoundObject.IsValid())
+		{
+			continue;
+		}
+
 		if (InstancePropertyBinding.HasValidBinding(*BoundObject))
 		{
 			return LabelParams.bIsDimmed ? FSlateColor::UseSubduedForeground() : FSlateColor::UseForeground();
