Sequencer: Added an option to skeletal animation sections to allow them to play back linearly when scaled by a scaling driver

diff --git a/Engine/Plugins/Experimental/MovieSceneAnimMixer/Source/MovieSceneAnimMixer/Private/Systems/MovieSceneMixedSkeletalAnimationSystem.cpp b/Engine/Plugins/Experimental/MovieSceneAnimMixer/Source/MovieSceneAnimMixer/Private/Systems/MovieSceneMixedSkeletalAnimationSystem.cpp
--- a/Engine/Plugins/Experimental/MovieSceneAnimMixer/Source/MovieSceneAnimMixer/Private/Systems/MovieSceneMixedSkeletalAnimationSystem.cpp
+++ b/Engine/Plugins/Experimental/MovieSceneAnimMixer/Source/MovieSceneAnimMixer/Private/Systems/MovieSceneMixedSkeletalAnimationSystem.cpp
@@ -3,12 +3,14 @@
 #include "Systems/MovieSceneMixedSkeletalAnimationSystem.h"
 
 #include "Async/TaskGraphInterfaces.h"
+#include "Decorations/MovieSceneScalingAnchors.h"
 #include "EntitySystem/BuiltInComponentTypes.h"
 #include "AnimMixerComponentTypes.h"
 #include "EntitySystem/MovieSceneEntityIDs.h"
 #include "EntitySystem/MovieSceneEntitySystemLinker.h"
 #include "EntitySystem/MovieSceneEntitySystemRunner.h"
 #include "EntitySystem/MovieSceneEntitySystemTask.h"
+#include "MovieScene.h"
 #include "MovieSceneTracksComponentTypes.h"
 #include "Sections/MovieSceneSkeletalAnimationSection.h"
 #include "Tracks/MovieSceneSkeletalAnimationTrack.h"
@@ -197,21 +199,57 @@ struct FGatherMixableSkeletalAnimations
 			const FSequenceInstance& SequenceInstance = InstanceRegistry->GetInstance(InstanceHandle);
 			const FMovieSceneContext& Context = SequenceInstance.GetContext();
 
-			const FFrameTime EvalFrameTime = OptionalEvalTimes ? OptionalEvalTimes[Index] : Context.GetTime();
-
 			// Calculate the time at which to evaluate the animation
 			const UMovieSceneSkeletalAnimationSection* AnimSection = SkeletalAnimation.Section;
 			const FMovieSceneSkeletalAnimationParams& AnimParams = AnimSection->Params;
 
 			// Get the bound skeletal mesh component.
 			USkeletalMeshComponent* SkeletalMeshComponent = CastChecked<USkeletalMeshComponent>(BoundObject);
-			UAnimSequence* AnimSequence = Cast<UAnimSequence>(AnimParams.Animation);
+			UAnimSequence* AnimSequence = Cast<UAnimSequence>(AnimSection->GetPlaybackAnimation());
 			if (!SkeletalMeshComponent || AnimSequence == nullptr)
 			{
 				continue;
 			}
-			const float EvalTime = AnimParams.MapTimeToAnimation(AnimSection, EvalFrameTime, Context.GetFrameRate());
-			const float PreviousEvalTime = AnimParams.MapTimeToAnimation(AnimSection, Context.GetPreviousTime(), Context.GetFrameRate());
+
+
+			FFrameTime EvalFrameTime = OptionalEvalTimes ? OptionalEvalTimes[Index] : Context.GetTime();
+			FFrameTime PreviousEvalFrameTime = Context.GetPreviousTime();
+
+			FFrameNumber SectionStartTime = AnimSection->GetInclusiveStartFrame();
+			FFrameNumber SectionEndTime   = AnimSection->GetExclusiveEndFrame();
+
+			if (AnimParams.bLinearPlaybackWhenScaled)
+			{
+				UMovieSceneScalingAnchors* ScalingAnchors = AnimSection->GetTypedOuter<UMovieScene>()->FindDecoration<UMovieSceneScalingAnchors>();
+
+				if (ScalingAnchors)
+				{
+					TOptional<FFrameTime> UnwarpedTime = ScalingAnchors->InverseRemapTimeCycled(EvalFrameTime, EvalFrameTime, FInverseTransformTimeParams());
+					if (UnwarpedTime.IsSet())
+					{
+						EvalFrameTime = UnwarpedTime.GetValue();
+					}
+					TOptional<FFrameTime> PreviousUnwarpedTime = ScalingAnchors->InverseRemapTimeCycled(PreviousEvalFrameTime, PreviousEvalFrameTime, FInverseTransformTimeParams());
+					if (PreviousUnwarpedTime.IsSet())
+					{
+						PreviousEvalFrameTime = PreviousUnwarpedTime.GetValue();
+					}
+					TOptional<FFrameTime> UnwarpedStartTime = ScalingAnchors->InverseRemapTimeCycled(SectionStartTime, SectionStartTime, FInverseTransformTimeParams());
+					if (UnwarpedStartTime.IsSet())
+					{
+						SectionStartTime = UnwarpedStartTime.GetValue().RoundToFrame();
+					}
+					TOptional<FFrameTime> UnwarpedEndTime = ScalingAnchors->InverseRemapTimeCycled(SectionEndTime, SectionEndTime, FInverseTransformTimeParams());
+					if (UnwarpedEndTime.IsSet())
+					{
+						SectionEndTime = UnwarpedEndTime.GetValue().RoundToFrame();
+					}
+				}
+			}
+
+
+			const float EvalTime = AnimParams.MapTimeToAnimation(SectionStartTime, SectionEndTime, EvalFrameTime, Context.GetFrameRate(), AnimSequence);
+			const float PreviousEvalTime = AnimParams.MapTimeToAnimation(SectionStartTime, SectionEndTime, PreviousEvalFrameTime, Context.GetFrameRate(), AnimSequence);
 
 			const EMovieScenePlayerStatus::Type PlayerStatus = Context.GetStatus();
 
diff --git a/Engine/Source/Runtime/MovieScene/Private/Decorations/MovieSceneLanguagePreviewDecoration.cpp b/Engine/Source/Runtime/MovieScene/Private/Decorations/MovieSceneLanguagePreviewDecoration.cpp
new file mode 100644
--- /dev/null
+++ b/Engine/Source/Runtime/MovieScene/Private/Decorations/MovieSceneLanguagePreviewDecoration.cpp
@@ -0,0 +1,54 @@
+// Copyright Epic Games, Inc. All Rights Reserved.
+
+#include "Decorations/MovieSceneLanguagePreviewDecoration.h"
+#include "MovieSceneSection.h"
+#include "MovieScene.h"
+#include "Misc/PackageName.h"
+#include "UObject/Package.h"
+
+UObject* UMovieSceneLanguagePreviewDecoration::FindLocalizedAsset(UObject* InUnlocalizedAsset, const UMovieSceneSection* Section)
+{
+	check(Section);
+
+#if WITH_EDITORONLY_DATA
+
+	if (InUnlocalizedAsset != nullptr)
+	{
+		// Handle custom language preview if necessary
+		UMovieSceneLanguagePreviewDecoration* This = Section->GetTypedOuter<UMovieScene>()->FindDecoration<UMovieSceneLanguagePreviewDecoration>();
+		if (This)
+		{
+			return This->FindLocalizedAsset(InUnlocalizedAsset);
+		}
+	}
+
+#endif
+
+	// Asset will have already been localized if we're in game by package-level localization
+	return InUnlocalizedAsset;
+}
+
+UObject* UMovieSceneLanguagePreviewDecoration::FindLocalizedAsset(UObject* InUnlocalizedAsset) const
+{
+
+#if WITH_EDITORONLY_DATA
+
+	// Handle custom language preview if necessary
+	if (InUnlocalizedAsset != nullptr && PreviewLanguage.Len() != 0)
+	{
+		FString LocalizedPackageName = FPackageName::GetLocalizedPackagePath(InUnlocalizedAsset->GetOutermost()->GetPathName(), PreviewLanguage);
+		if (LocalizedPackageName.Len() != 0)
+		{
+			// Attempt to load the localized asset
+			UObject* LocalizedAsset = LoadObject<UObject>(nullptr, *FString(LocalizedPackageName + TEXT(".") + InUnlocalizedAsset->GetName()));
+			if (LocalizedAsset)
+			{
+				return LocalizedAsset;
+			}
+		}
+	}
+
+#endif
+
+	return InUnlocalizedAsset;
+}
\ No newline at end of file
diff --git a/Engine/Source/Runtime/MovieScene/Public/Decorations/MovieSceneLanguagePreviewDecoration.h b/Engine/Source/Runtime/MovieScene/Public/Decorations/MovieSceneLanguagePreviewDecoration.h
--- a/Engine/Source/Runtime/MovieScene/Public/Decorations/MovieSceneLanguagePreviewDecoration.h
+++ b/Engine/Source/Runtime/MovieScene/Public/Decorations/MovieSceneLanguagePreviewDecoration.h
@@ -6,6 +6,7 @@
 #include "UObject/Interface.h"
 #include "MovieSceneLanguagePreviewDecoration.generated.h"
 
+class UMovieSceneSection;
 
 UCLASS(MinimalAPI)
 class UMovieSceneLanguagePreviewDecoration
@@ -15,6 +16,50 @@ public:
 
 	GENERATED_BODY()
 
+	/** Currently assigned preview language */
 	UPROPERTY()
 	FString PreviewLanguage;
+
+	/**
+	 * Find the localized asset for this preview language based on the specified reference asset
+	 */
+	MOVIESCENE_API UObject* FindLocalizedAsset(UObject* InUnlocalizedAsset) const;
+
+	/**
+	 * Find the localized asset for a preview language based on the specified reference asset and section.
+	 * @note This function uses the UMovieSceneLanguagePreviewDecoration on the specified section's MovieScene, if present
+	 */
+	MOVIESCENE_API static UObject* FindLocalizedAsset(UObject* InUnlocalizedAsset, const UMovieSceneSection* Section);
+
+	/**
+	 * Find the localized asset for a preview language based on the specified reference asset and section.
+	 * @note This function uses the UMovieSceneLanguagePreviewDecoration on the specified section's MovieScene, if present
+	 */
+	template<typename AssetClass>
+	static AssetClass* FindLocalizedAsset(TObjectPtr<AssetClass> InUnlocalizedAsset, const UMovieSceneSection* Section)
+	{
+		UObject* UntypedResult = FindLocalizedAsset(static_cast<UObject*>(InUnlocalizedAsset), Section);
+		return CastChecked<AssetClass>(UntypedResult, ECastCheckedType::NullAllowed);
+	}
+
+	/**
+	 * Find the localized asset for a preview language based on the specified reference asset and section.
+	 * @note This function uses the UMovieSceneLanguagePreviewDecoration on the specified section's MovieScene, if present
+	 */
+	template<typename AssetClass>
+	static AssetClass* FindLocalizedAsset(AssetClass* InUnlocalizedAsset, const UMovieSceneSection* Section)
+	{
+		UObject* UntypedResult = FindLocalizedAsset(static_cast<UObject*>(InUnlocalizedAsset), Section);
+		return CastChecked<AssetClass>(UntypedResult, ECastCheckedType::NullAllowed);
+	}
+
+	/**
+	 * Find the localized asset for this preview language based on the specified reference asset
+	 */
+	template<typename AssetClass>
+	AssetClass* FindLocalizedAsset(TObjectPtr<AssetClass> InUnlocalizedAsset) const
+	{
+		UObject* UntypedResult = FindLocalizedAsset(static_cast<UObject*>(InUnlocalizedAsset));
+		return CastChecked<AssetClass>(UntypedResult, ECastCheckedType::NullAllowed);
+	}
 };
\ No newline at end of file
diff --git a/Engine/Source/Runtime/MovieSceneTracks/Private/Sections/MovieSceneAudioSection.cpp b/Engine/Source/Runtime/MovieSceneTracks/Private/Sections/MovieSceneAudioSection.cpp
--- a/Engine/Source/Runtime/MovieSceneTracks/Private/Sections/MovieSceneAudioSection.cpp
+++ b/Engine/Source/Runtime/MovieSceneTracks/Private/Sections/MovieSceneAudioSection.cpp
@@ -154,33 +154,7 @@ EMovieSceneChannelProxyType  UMovieSceneAudioSection::CacheChannelProxy()
 
 USoundBase* UMovieSceneAudioSection::GetPlaybackSound() const
 {
-	if (!Sound)
-	{
-		return nullptr;
-	}
-
-#if WITH_EDITORONLY_DATA
-
-	// Handle custom language preview if necessary
-	UMovieSceneLanguagePreviewDecoration* PreviewLanguage = GetTypedOuter<UMovieScene>()->FindDecoration<UMovieSceneLanguagePreviewDecoration>();
-	if (PreviewLanguage && PreviewLanguage->PreviewLanguage.Len() != 0)
-	{
-		FString LocalizedPackageName = FPackageName::GetLocalizedPackagePath(Sound->GetOutermost()->GetPathName(), PreviewLanguage->PreviewLanguage);
-		if (LocalizedPackageName.Len() != 0)
-		{
-			// Attempt to load the localized asset
-			USoundBase* LocalizedSound = LoadObject<USoundBase>(nullptr, *FString(LocalizedPackageName + TEXT(".") + Sound->GetName()));
-			if (LocalizedSound)
-			{
-				return LocalizedSound;
-			}
-		}
-	}
-
-#endif
-
-	// Sound will have already been localized if we're in game by package-level localization
-	return Sound;
+	return UMovieSceneLanguagePreviewDecoration::FindLocalizedAsset(Sound, this);
 }
 
 void UMovieSceneAudioSection::PopulateInitialAnchors(TMap<FGuid, FMovieSceneScalingAnchor>& OutAnchors)
diff --git a/Engine/Source/Runtime/MovieSceneTracks/Private/Sections/MovieSceneSkeletalAnimationSection.cpp b/Engine/Source/Runtime/MovieSceneTracks/Private/Sections/MovieSceneSkeletalAnimationSection.cpp
--- a/Engine/Source/Runtime/MovieSceneTracks/Private/Sections/MovieSceneSkeletalAnimationSection.cpp
+++ b/Engine/Source/Runtime/MovieSceneTracks/Private/Sections/MovieSceneSkeletalAnimationSection.cpp
@@ -20,6 +20,7 @@
 #include "Sections/MovieSceneSectionTimingParameters.h"
 #include "Variants/MovieSceneTimeWarpGetter.h"
 #include "Systems/MovieSceneSkeletalAnimationSystem.h"
+#include "Decorations/MovieSceneLanguagePreviewDecoration.h"
 
 #include UE_INLINE_GENERATED_CPP_BY_NAME(MovieSceneSkeletalAnimationSection)
 
@@ -48,16 +49,22 @@ FMovieSceneSkeletalAnimationParams::FMovieSceneSkeletalAnimationParams()
 	bSkipAnimNotifiers = false;
 	bForceCustomMode = false;
 	SwapRootBone = ESwapRootBone::SwapRootBone_None;
+	bLinearPlaybackWhenScaled = false;
 }
 
-FMovieSceneSequenceTransform FMovieSceneSkeletalAnimationParams::MakeTransform(const FFrameRate& OuterFrameRate, const TRange<FFrameNumber>& OuterRange) const
+FMovieSceneSequenceTransform FMovieSceneSkeletalAnimationParams::MakeTransform(const FFrameRate& OuterFrameRate, const TRange<FFrameNumber>& OuterRange, UAnimSequenceBase* OverrideSequence) const
 {
-	if (!Animation)
+	if (!OverrideSequence)
+	{
+		OverrideSequence = Animation;
+	}
+
+	if (!OverrideSequence)
 	{
 		return FMovieSceneSequenceTransform();
 	}
 
-	const double SequenceLength = GetSequenceLength();
+	const double SequenceLength = OverrideSequence->GetPlayLength();
 	const FFrameTime AnimationLength = SequenceLength * OuterFrameRate;
 	const int32 LengthInFrames = AnimationLength.FrameNumber.Value + (int)(AnimationLength.GetSubFrame() + 0.5f) + 1;
 
@@ -73,19 +80,19 @@ FMovieSceneSequenceTransform FMovieSceneSkeletalAnimationParams::MakeTransform(c
 		bReverse
 	};
 
-	return TimingParams.MakeTransform(OuterFrameRate, OuterRange, SequenceLength, Animation->RateScale);
+	return TimingParams.MakeTransform(OuterFrameRate, OuterRange, SequenceLength, OverrideSequence->RateScale);
 }
 
-double FMovieSceneSkeletalAnimationParams::MapTimeToAnimation(const UMovieSceneSection* InSection, FFrameTime InPosition, FFrameRate InFrameRate) const
+double FMovieSceneSkeletalAnimationParams::MapTimeToAnimation(const UMovieSceneSection* InSection, FFrameTime InPosition, FFrameRate InFrameRate, UAnimSequenceBase* OverrideSequence) const
 {
 	const FFrameNumber SectionStartTime = InSection->GetInclusiveStartFrame();
 	const FFrameNumber SectionEndTime = InSection->GetExclusiveEndFrame();
-	return MapTimeToAnimation(SectionStartTime, SectionEndTime, InPosition, InFrameRate);
+	return MapTimeToAnimation(SectionStartTime, SectionEndTime, InPosition, InFrameRate, OverrideSequence);
 }
 
-double FMovieSceneSkeletalAnimationParams::MapTimeToAnimation(FFrameNumber InSectionStartTime, FFrameNumber InSectionEndTime, FFrameTime InPosition, FFrameRate InFrameRate) const
+double FMovieSceneSkeletalAnimationParams::MapTimeToAnimation(FFrameNumber InSectionStartTime, FFrameNumber InSectionEndTime, FFrameTime InPosition, FFrameRate InFrameRate, UAnimSequenceBase* OverrideSequence) const
 {
-	return MakeTransform(InFrameRate, TRange<FFrameNumber>(InSectionStartTime, InSectionEndTime)).TransformTime(InPosition).AsDecimal();
+	return MakeTransform(InFrameRate, TRange<FFrameNumber>(InSectionStartTime, InSectionEndTime), OverrideSequence).TransformTime(InPosition).AsDecimal();
 }
 
 UMovieSceneSkeletalAnimationSection::UMovieSceneSkeletalAnimationSection( const FObjectInitializer& ObjectInitializer )
@@ -125,6 +132,18 @@ UMovieSceneSkeletalAnimationSection::UMovieSceneSkeletalAnimationSection( const
 	MixedAnimationTarget = TInstancedStruct<FMovieSceneMixedAnimationTarget>::Make();
 }
 
+UAnimSequenceBase* UMovieSceneSkeletalAnimationSection::GetAnimation() const
+{
+	return Params.Animation;
+}
+
+
+UAnimSequenceBase* UMovieSceneSkeletalAnimationSection::GetPlaybackAnimation() const
+{
+	return UMovieSceneLanguagePreviewDecoration::FindLocalizedAsset(Params.Animation, this);
+}
+
+
 EMovieSceneChannelProxyType UMovieSceneSkeletalAnimationSection::CacheChannelProxy()
 {
 	FMovieSceneChannelProxyData Channels;
diff --git a/Engine/Source/Runtime/MovieSceneTracks/Private/Systems/MovieSceneSkeletalAnimationSystem.cpp b/Engine/Source/Runtime/MovieSceneTracks/Private/Systems/MovieSceneSkeletalAnimationSystem.cpp
--- a/Engine/Source/Runtime/MovieSceneTracks/Private/Systems/MovieSceneSkeletalAnimationSystem.cpp
+++ b/Engine/Source/Runtime/MovieSceneTracks/Private/Systems/MovieSceneSkeletalAnimationSystem.cpp
@@ -2,6 +2,7 @@
 
 #include "Systems/MovieSceneSkeletalAnimationSystem.h"
 
+#include "Decorations/MovieSceneScalingAnchors.h"
 #include "Async/TaskGraphInterfaces.h"
 #include "EntitySystem/BuiltInComponentTypes.h"
 #include "EntitySystem/Interrogation/MovieSceneInterrogationLinker.h"
@@ -10,6 +11,7 @@
 #include "EntitySystem/MovieSceneEntitySystemRunner.h"
 #include "EntitySystem/MovieSceneEntitySystemTask.h"
 #include "IMovieScenePlayer.h"
+#include "MovieScene.h"
 #include "MovieSceneExecutionToken.h"
 #include "MovieSceneTracksComponentTypes.h"
 #include "Sections/MovieSceneSkeletalAnimationSection.h"
@@ -431,20 +433,55 @@ struct FGatherSkeletalAnimations
 			const FSequenceInstance& SequenceInstance = InstanceRegistry->GetInstance(InstanceHandle);
 			const FMovieSceneContext& Context = SequenceInstance.GetContext();
 
-			const FFrameTime EvalFrameTime = OptionalEvalTimes ? OptionalEvalTimes[Index] : Context.GetTime();
-
 			// Calculate the time at which to evaluate the animation
 			const UMovieSceneSkeletalAnimationSection* AnimSection = SkeletalAnimation.Section;
 			const FMovieSceneSkeletalAnimationParams& AnimParams = AnimSection->Params;
 
 			// Get the bound skeletal mesh component.
 			USkeletalMeshComponent* SkeletalMeshComponent = CastChecked<USkeletalMeshComponent>(BoundObject);
-			if (!SkeletalMeshComponent || AnimParams.Animation == nullptr)
+			UAnimSequenceBase* AnimSequence = AnimSection->GetPlaybackAnimation();
+			if (!SkeletalMeshComponent || AnimSequence == nullptr)
 			{
 				continue;
 			}
-			const float EvalTime = AnimParams.MapTimeToAnimation(AnimSection, EvalFrameTime, Context.GetFrameRate());
-			const float PreviousEvalTime = AnimParams.MapTimeToAnimation(AnimSection, Context.GetPreviousTime(), Context.GetFrameRate());
+
+			FFrameTime EvalFrameTime = OptionalEvalTimes ? OptionalEvalTimes[Index] : Context.GetTime();
+			FFrameTime PreviousEvalFrameTime = Context.GetPreviousTime();
+
+			FFrameNumber SectionStartTime = AnimSection->GetInclusiveStartFrame();
+			FFrameNumber SectionEndTime   = AnimSection->GetExclusiveEndFrame();
+
+			if (AnimParams.bLinearPlaybackWhenScaled)
+			{
+				UMovieSceneScalingAnchors* ScalingAnchors = AnimSection->GetTypedOuter<UMovieScene>()->FindDecoration<UMovieSceneScalingAnchors>();
+
+				if (ScalingAnchors)
+				{
+					TOptional<FFrameTime> UnwarpedTime = ScalingAnchors->InverseRemapTimeCycled(EvalFrameTime, EvalFrameTime, FInverseTransformTimeParams());
+					if (UnwarpedTime.IsSet())
+					{
+						EvalFrameTime = UnwarpedTime.GetValue();
+					}
+					TOptional<FFrameTime> PreviousUnwarpedTime = ScalingAnchors->InverseRemapTimeCycled(PreviousEvalFrameTime, PreviousEvalFrameTime, FInverseTransformTimeParams());
+					if (PreviousUnwarpedTime.IsSet())
+					{
+						PreviousEvalFrameTime = PreviousUnwarpedTime.GetValue();
+					}
+					TOptional<FFrameTime> UnwarpedStartTime = ScalingAnchors->InverseRemapTimeCycled(SectionStartTime, SectionStartTime, FInverseTransformTimeParams());
+					if (UnwarpedStartTime.IsSet())
+					{
+						SectionStartTime = UnwarpedStartTime.GetValue().RoundToFrame();
+					}
+					TOptional<FFrameTime> UnwarpedEndTime = ScalingAnchors->InverseRemapTimeCycled(SectionEndTime, SectionEndTime, FInverseTransformTimeParams());
+					if (UnwarpedEndTime.IsSet())
+					{
+						SectionEndTime = UnwarpedEndTime.GetValue().RoundToFrame();
+					}
+				}
+			}
+
+			const float EvalTime = AnimParams.MapTimeToAnimation(SectionStartTime, SectionEndTime, EvalFrameTime, Context.GetFrameRate(), AnimSequence);
+			const float PreviousEvalTime = AnimParams.MapTimeToAnimation(SectionStartTime, SectionEndTime, PreviousEvalFrameTime, Context.GetFrameRate(), AnimSequence);
 
 			const FSequenceInstance& RootInstance = InstanceRegistry->GetInstance(RootInstanceHandle);
 			const FMovieSceneContext& RootContext = RootInstance.GetContext();
@@ -829,7 +866,8 @@ private:
 		static const bool bLooping = false;
 
 		const FMovieSceneSkeletalAnimationParams& AnimParams = Params.Section->Params;
-		if (!CanPlayAnimation(Params.SkeletalMeshComponent, AnimParams.Animation))
+		UAnimSequenceBase* Animation = Params.Section->GetPlaybackAnimation();
+		if (!CanPlayAnimation(Params.SkeletalMeshComponent, Animation))
 		{
 			return;
 		}
@@ -865,7 +903,7 @@ private:
 			TOptional<FTransform> CurrentTransform = GetCurrentTransform(Params.Section->Params.SwapRootBone, Params.SkeletalMeshComponent);
 
 			FAnimSequencerData AnimSequencerData(
-					AnimParams.Animation,
+					Animation,
 					AnimSequenceID,
 					RootMotion, 
 					Params.FromPosition,
@@ -884,12 +922,12 @@ private:
 			int32 InstanceId = (SectionData) ? SectionData->MontageInstanceId : INDEX_NONE;
 
 
-			const float AssetPlayRate = FMath::IsNearlyZero(AnimParams.Animation->RateScale) ? 1.0f : AnimParams.Animation->RateScale;
+			const float AssetPlayRate = FMath::IsNearlyZero(Animation->RateScale) ? 1.0f : Animation->RateScale;
 			TWeakObjectPtr<UAnimMontage> WeakMontage = FAnimMontageInstance::SetSequencerMontagePosition(
 					AnimParams.SlotName,
 					AnimInst, 
 					InstanceId, 
-					AnimParams.Animation, 
+					Animation, 
 					Params.FromPosition / AssetPlayRate, 
 					Params.ToPosition / AssetPlayRate, 
 					Params.Weight, 
@@ -925,7 +963,8 @@ private:
 		static const bool bLooping = false;
 
 		const FMovieSceneSkeletalAnimationParams& AnimParams = Params.Section->Params;
-		if (!CanPlayAnimation(Params.SkeletalMeshComponent, AnimParams.Animation))
+		UAnimSequenceBase* Animation = Params.Section->GetPlaybackAnimation();
+		if (!CanPlayAnimation(Params.SkeletalMeshComponent, Animation))
 		{
 			return;
 		}
@@ -959,7 +998,7 @@ private:
 			TOptional<FTransform> CurrentTransform = GetCurrentTransform(Params.Section->Params.SwapRootBone, Params.SkeletalMeshComponent);
 
 			FAnimSequencerData AnimSequencerData(
-					AnimParams.Animation,
+					Animation,
 					AnimSequenceID,
 					RootMotion,
 					Params.FromPosition,
@@ -977,13 +1016,13 @@ private:
 
 			int32 InstanceId = SectionData ? SectionData->MontageInstanceId : INDEX_NONE;
 		
-			const float AssetPlayRate = FMath::IsNearlyZero(AnimParams.Animation->RateScale) ? 1.0f : AnimParams.Animation->RateScale;
+			const float AssetPlayRate = FMath::IsNearlyZero(Animation->RateScale) ? 1.0f : Animation->RateScale;
 			TWeakObjectPtr<UAnimMontage> WeakMontage = FAnimMontageInstance::PreviewSequencerMontagePosition(
 					AnimParams.SlotName,
 					Params.SkeletalMeshComponent,
 					AnimInst,
 					InstanceId,
-					AnimParams.Animation,
+					Animation,
 					Params.FromPosition / AssetPlayRate,
 					Params.ToPosition / AssetPlayRate,
 					Params.Weight,
diff --git a/Engine/Source/Runtime/MovieSceneTracks/Public/Sections/MovieSceneSkeletalAnimationSection.h b/Engine/Source/Runtime/MovieSceneTracks/Public/Sections/MovieSceneSkeletalAnimationSection.h
--- a/Engine/Source/Runtime/MovieSceneTracks/Public/Sections/MovieSceneSkeletalAnimationSection.h
+++ b/Engine/Source/Runtime/MovieSceneTracks/Public/Sections/MovieSceneSkeletalAnimationSection.h
@@ -35,17 +35,17 @@ struct FMovieSceneSkeletalAnimationParams
 	 * Convert a sequence frame to a time in seconds inside the animation clip, taking into account start/end offsets,
 	 * looping, etc.
 	 */
-	double MOVIESCENETRACKS_API MapTimeToAnimation(const UMovieSceneSection* InSection, FFrameTime InPosition, FFrameRate InFrameRate) const;
+	double MOVIESCENETRACKS_API MapTimeToAnimation(const UMovieSceneSection* InSection, FFrameTime InPosition, FFrameRate InFrameRate, UAnimSequenceBase* OverrideSequence = nullptr) const;
 
 	/**
 	 * As above, but with already computed section bounds.
 	 */
-	double MOVIESCENETRACKS_API MapTimeToAnimation(FFrameNumber InSectionStartTime, FFrameNumber InSectionEndTime, FFrameTime InPosition, FFrameRate InFrameRate) const;
+	double MOVIESCENETRACKS_API MapTimeToAnimation(FFrameNumber InSectionStartTime, FFrameNumber InSectionEndTime, FFrameTime InPosition, FFrameRate InFrameRate, UAnimSequenceBase* OverrideSequence = nullptr) const;
 
 	/**
 	 * Make a transform structure from these animation parameters
 	 */
-	MOVIESCENETRACKS_API FMovieSceneSequenceTransform MakeTransform(const FFrameRate& OuterFrameRate, const TRange<FFrameNumber>& OuterRange) const;
+	MOVIESCENETRACKS_API FMovieSceneSequenceTransform MakeTransform(const FFrameRate& OuterFrameRate, const TRange<FFrameNumber>& OuterRange, UAnimSequenceBase* OverrideSequence = nullptr) const;
 
 	/** The animation this section plays */
 	UPROPERTY(BlueprintReadWrite, EditAnywhere, Category="Animation", meta=(AllowedClasses = "/Script/Engine.AnimSequence,/Script/Engine.AnimComposite,/Script/Engine.AnimStreamable"))
@@ -94,6 +94,9 @@ struct FMovieSceneSkeletalAnimationParams
 	UPROPERTY(BlueprintReadWrite, EditAnywhere, Category = "Animation|Playback")
 	ESwapRootBone SwapRootBone;
 
+	UPROPERTY(BlueprintReadWrite, EditAnywhere, Category = "Animation|Playback")
+	bool bLinearPlaybackWhenScaled;
+
 	UPROPERTY()
 	float StartOffset_DEPRECATED;
 
@@ -119,7 +122,10 @@ public:
 
 	/** Get Frame Time as Animation Time*/
 	MOVIESCENETRACKS_API double MapTimeToAnimation(FFrameTime InPosition, FFrameRate InFrameRate) const;
-	
+
+	MOVIESCENETRACKS_API UAnimSequenceBase* GetAnimation() const;
+	MOVIESCENETRACKS_API UAnimSequenceBase* GetPlaybackAnimation() const;
+
 	//~ UMovieSceneSection interface
 	virtual EMovieSceneChannelProxyType CacheChannelProxy() override;
 	virtual void SetRange(const TRange<FFrameNumber>& NewRange) override;
