Sequencer: Fix AutoSize computation

The frames are already in tick resolution and the end time should be calculated as start + length.

diff --git a/Engine/Source/Runtime/MovieSceneTracks/Private/Sections/MovieSceneSkeletalAnimationSection.cpp b/Engine/Source/Runtime/MovieSceneTracks/Private/Sections/MovieSceneSkeletalAnimationSection.cpp
--- a/Engine/Source/Runtime/MovieSceneTracks/Private/Sections/MovieSceneSkeletalAnimationSection.cpp
+++ b/Engine/Source/Runtime/MovieSceneTracks/Private/Sections/MovieSceneSkeletalAnimationSection.cpp
@@ -326,24 +326,23 @@ TOptional<TRange<FFrameNumber> > UMovieSceneSkeletalAnimationSection::GetAutoSiz
 {
 	if (UMovieScene* MovieScene = GetTypedOuter<UMovieScene>())
 	{
-		FFrameRate FrameRate = MovieScene->GetTickResolution();
-		FMovieSceneInverseSequenceTransform InnerToOuterTransform = Params.MakeTransform(FrameRate, GetRange()).Inverse();
+		FFrameRate TickResolution = MovieScene->GetTickResolution();
+		FMovieSceneInverseSequenceTransform InnerToOuterTransform = Params.MakeTransform(TickResolution, GetRange()).Inverse();
+		
+		FFrameTime AnimationLength = Params.GetSequenceLength() * TickResolution;
+		int32 IFrameNumber = AnimationLength.FrameNumber.Value + static_cast<int32>(AnimationLength.GetSubFrame() + 0.5f) + 1;
 
-		FFrameTime InnerStartTime = FFrameTime::FromDecimal((Params.StartFrameOffset + Params.FirstLoopStartFrameOffset) / FrameRate);
-		FFrameTime InnerEndTime   = FFrameTime::FromDecimal(Params.GetSequenceLength() - (Params.EndFrameOffset / FrameRate));
+		FFrameTime InnerStartTime = Params.StartFrameOffset + Params.FirstLoopStartFrameOffset;
+		FFrameTime InnerEndTime = IFrameNumber - Params.EndFrameOffset;
 
-		TOptional<FFrameTime> OuterStartTime = InnerToOuterTransform.TryTransformTime(InnerStartTime);
-		TOptional<FFrameTime> OuterEndTime   = InnerToOuterTransform.TryTransformTime(InnerEndTime);
+		const FFrameTime OuterStartTime = InnerToOuterTransform.TryTransformTime(InnerStartTime).Get(InnerStartTime);
+		const FFrameTime OuterEndTime = InnerToOuterTransform.TryTransformTime(InnerEndTime).Get(InnerEndTime);
 
-		if (OuterStartTime && OuterEndTime)
-		{
-			return TRange<FFrameNumber>(OuterStartTime.GetValue().FrameNumber, OuterEndTime.GetValue().FrameNumber + 1);
-		}
+		return TRange<FFrameNumber>(GetInclusiveStartFrame(), GetInclusiveStartFrame() + (OuterEndTime - OuterStartTime).FrameNumber);
 	}
 	return TOptional<TRange<FFrameNumber>>();
 }
 
-
 void UMovieSceneSkeletalAnimationSection::TrimSection(FQualifiedFrameTime TrimTime, bool bTrimLeft, bool bDeleteKeys)
 {
 	SetFlags(RF_Transactional);
