Fixing a bug where decal material tracks would display red and disconnected when they were actually connected.

diff --git a/Engine/Source/Runtime/MovieSceneTracks/Private/Tracks/MovieSceneMaterialTrack.cpp b/Engine/Source/Runtime/MovieSceneTracks/Private/Tracks/MovieSceneMaterialTrack.cpp
--- a/Engine/Source/Runtime/MovieSceneTracks/Private/Tracks/MovieSceneMaterialTrack.cpp
+++ b/Engine/Source/Runtime/MovieSceneTracks/Private/Tracks/MovieSceneMaterialTrack.cpp
@@ -487,6 +487,7 @@ FSlateColor UMovieSceneComponentMaterialTrack::GetLabelColor(const FMovieSceneLa
 				{
 					return GetDimmedColor(FLinearColor::Red);
 				}
+				break;
 			case EComponentMaterialType::VolumetricCloudMaterial:
 				if (UVolumetricCloudComponent* CloudComponent = Cast<UVolumetricCloudComponent>(WeakObject.Get()))
 				{
