Fix for crash with scrubbing procedural sound wave aasets

diff --git a/Engine/Source/Runtime/MovieSceneTracks/Private/Systems/MovieSceneAudioSystem.cpp b/Engine/Source/Runtime/MovieSceneTracks/Private/Systems/MovieSceneAudioSystem.cpp
--- a/Engine/Source/Runtime/MovieSceneTracks/Private/Systems/MovieSceneAudioSystem.cpp
+++ b/Engine/Source/Runtime/MovieSceneTracks/Private/Systems/MovieSceneAudioSystem.cpp
@@ -722,7 +722,8 @@ private:
 				{
 					if (USoundWave* SoundWave = Cast<USoundWave>(Sound))
 					{
-						if (ScrubbedSound)
+						// Can't scrub a procedural asset right now
+						if (ScrubbedSound && !SoundWave->bProcedural)
 						{
 							// If we're not playing the audio component, then set up the audio component to use the scrubbed sound wave and play it
 							if (!AudioComponent.IsPlaying())
