Sequencer: check that the binding is not null before using it

#rnx

diff --git a/Engine/Source/Runtime/MovieSceneTracks/Private/EntitySystem/Interrogation/MovieSceneInterrogationLinker.cpp b/Engine/Source/Runtime/MovieSceneTracks/Private/EntitySystem/Interrogation/MovieSceneInterrogationLinker.cpp
--- a/Engine/Source/Runtime/MovieSceneTracks/Private/EntitySystem/Interrogation/MovieSceneInterrogationLinker.cpp
+++ b/Engine/Source/Runtime/MovieSceneTracks/Private/EntitySystem/Interrogation/MovieSceneInterrogationLinker.cpp
@@ -1023,7 +1023,11 @@ FInterrogationChannel FSystemInterrogator::ImportLocalTransforms(USceneComponent
 	FGuid ObjectBindingID = State->FindCachedObjectId(*SceneComponent, SequenceID, *InPlayer);
 	if (ObjectBindingID.IsValid())
 	{
-		ImportTransformTracks(*Sequence->GetMovieScene()->FindBinding(ObjectBindingID), Channel);
+		const FMovieSceneBinding* Binding = Sequence->GetMovieScene()->FindBinding(ObjectBindingID);
+		if (ensure(Binding))
+		{
+			ImportTransformTracks(*Binding, Channel);
+		}
 	}
 
 	// Also blend in any transforms that exist for this scene component's actor as well (if it is the root)
@@ -1033,7 +1037,11 @@ FInterrogationChannel FSystemInterrogator::ImportLocalTransforms(USceneComponent
 		FGuid OwnerObjectBindingID = State->FindCachedObjectId(*Owner, SequenceID, *InPlayer);
 		if (OwnerObjectBindingID.IsValid())
 		{
-			ImportTransformTracks(*Sequence->GetMovieScene()->FindBinding(OwnerObjectBindingID), Channel);
+			const FMovieSceneBinding* Binding = Sequence->GetMovieScene()->FindBinding(OwnerObjectBindingID);
+			if (ensure(Binding))
+			{
+				ImportTransformTracks(*Binding, Channel);
+			}
 		}
 	}
 
