Fix for crash in scrubbed sound around GC

diff --git a/Engine/Source/Runtime/MovieSceneTracks/Private/Systems/MovieSceneAudioSystem.cpp b/Engine/Source/Runtime/MovieSceneTracks/Private/Systems/MovieSceneAudioSystem.cpp
--- a/Engine/Source/Runtime/MovieSceneTracks/Private/Systems/MovieSceneAudioSystem.cpp
+++ b/Engine/Source/Runtime/MovieSceneTracks/Private/Systems/MovieSceneAudioSystem.cpp
@@ -507,7 +507,10 @@ private:
 		
 		UWorld* World = PlaybackContext ? PlaybackContext->GetWorld() : nullptr;
 		AWorldSettings* WorldSettings = World ? World->GetWorldSettings() : nullptr;
+
+#if WITH_EDITOR
 		UScrubbedSound* ScrubbedSound = EvaluationData.ScrubbedSound.Get();
+#endif // WITH_EDITOR
 
 		AudioSystem->PreAnimatedStorage->CachePreAnimatedValue(
 				FCachePreAnimatedValueParams(), &AudioComponent,
@@ -733,6 +736,7 @@ private:
 			{
 				UE_LOG(LogMovieScene, Verbose, TEXT("Audio Component Play at Local Time: %6.2f CurrentTime: %6.2f(s) SectionStart: %6.2f(s), SoundDur: %6.2f OffsetIntoClip: %6.2f Component: %s Sound: %s"), AudioTime, (EvalTime / Context.GetFrameRate()), SectionStartTimeSeconds, AudioComponent.Sound ? AudioComponent.Sound->GetDuration() : 0.0f, (float)Context.GetFrameRate().AsSeconds(AudioStartOffset), *AudioComponent.GetName(), *GetNameSafe(AudioComponent.Sound));
 				
+#if WITH_EDITOR
 				if (Context.GetStatus() == EMovieScenePlayerStatus::Scrubbing && bEnableGranularScrubbingCVar)
 				{
 					if (USoundWave* SoundWave = Cast<USoundWave>(Sound))
@@ -764,8 +768,15 @@ private:
 					}
 				}
 				else
+#endif // WITH_EDITOR
 				{
-					ScrubbedSound->SetIsScrubbing(false);
+#if WITH_EDITOR
+					if (ScrubbedSound)
+					{
+						ScrubbedSound->SetIsScrubbing(false);
+					}
+#endif // WITH_EDITOR
+
 					AudioComponent.Play(AudioTime);
 
 					// Keep track of when we asked this audio clip to play (in game time) so that we can figure 
@@ -1127,10 +1138,8 @@ UMovieSceneAudioSystem::FAudioComponentEvaluationData* UMovieSceneAudioSystem::A
 
 		FAudioDevice::FCreateComponentParams Params(Actor->GetWorld(), Actor);
 		UAudioComponent* NewComponent = FAudioDevice::CreateComponent(TempPlaybackAudioCue, Params);
-		static int32 ScrubSoundCounter = 0;
-		UScrubbedSound* ScrubbedSound = NewObject<UScrubbedSound>(Actor, FName(*FString::Printf(TEXT("ScrubbedSound_Bound_%i"), ScrubSoundCounter++)));
 
-		if (!NewComponent || !ScrubbedSound)
+		if (!NewComponent)
 		{
 			const int32 RowIndex = Section->GetRowIndex();
 			UE_LOG(LogMovieScene, Warning, TEXT("Failed to create audio component for spatialized audio track (row %d on %s)."), RowIndex, *ObjectName);
@@ -1140,12 +1149,25 @@ UMovieSceneAudioSystem::FAudioComponentEvaluationData* UMovieSceneAudioSystem::A
 		NewComponent->SetFlags(RF_Transient);
 		NewComponent->AttachToComponent(SceneComponent, FAttachmentTransformRules::KeepRelativeTransform);
 
-		ScrubbedSound->SetFlags(RF_Transient);
-
 		FInstanceObjectKey DataKey{ InstanceHandle, SectionKey };
 		ExistingData = &ActorAudioComponentMap.Add(DataKey);
 		ExistingData->AudioComponent = NewComponent;
-		ExistingData->ScrubbedSound = ScrubbedSound;
+
+#if WITH_EDITOR
+		static int32 ScrubSoundCounter = 0;
+		UScrubbedSound* ScrubbedSound = NewObject<UScrubbedSound>(Actor, FName(*FString::Printf(TEXT("ScrubbedSound_Bound_%i"), ScrubSoundCounter++)));
+		if (!ScrubbedSound)
+		{
+			const int32 RowIndex = Section->GetRowIndex();
+			UE_LOG(LogMovieScene, Warning, TEXT("Failed to create scrubbed sound audio track (row %d on %s)."), RowIndex, *ObjectName);
+		}
+		else
+		{
+			ScrubbedSound->SetFlags(RF_Transient);
+			ExistingData->ScrubbedSound = TStrongObjectPtr<UScrubbedSound>(ScrubbedSound);
+		}
+#endif // WITH_EDITOR
+
 	}
 
 	return ExistingData;
@@ -1169,10 +1191,7 @@ UMovieSceneAudioSystem::FAudioComponentEvaluationData* UMovieSceneAudioSystem::A
 
 		UAudioComponent* NewComponent = FAudioDevice::CreateComponent(TempPlaybackAudioCue, Params);
 
-		static int32 ScrubSoundCounter = 0;
-		UScrubbedSound* ScrubbedSound = NewObject<UScrubbedSound>(World->GetCurrentLevel(), FName(*FString::Printf(TEXT("ScrubbedSound_Root_%i"), ScrubSoundCounter++)));
-
-		if (!NewComponent || !ScrubbedSound)
+		if (!NewComponent)
 		{
 			const int32 RowIndex = Section->GetRowIndex();
 			UE_LOG(LogMovieScene, Warning, TEXT("Failed to create audio component for root audio track (row %d)."), RowIndex);
@@ -1181,12 +1200,24 @@ UMovieSceneAudioSystem::FAudioComponentEvaluationData* UMovieSceneAudioSystem::A
 
 		NewComponent->SetFlags(RF_Transient);
 
-		ScrubbedSound->SetFlags(RF_Transient);
-
 		FInstanceObjectKey DataKey{ InstanceHandle, SectionKey };
 		ExistingData = &RootAudioComponentMap.Add(DataKey);
 		ExistingData->AudioComponent = NewComponent;
-		ExistingData->ScrubbedSound = ScrubbedSound;
+
+#if WITH_EDITOR
+		static int32 ScrubSoundCounter = 0;
+		UScrubbedSound* ScrubbedSound = NewObject<UScrubbedSound>(World->GetCurrentLevel(), FName(*FString::Printf(TEXT("ScrubbedSound_Root_%i"), ScrubSoundCounter++)));
+		if (!ScrubbedSound)
+		{
+			const int32 RowIndex = Section->GetRowIndex();
+			UE_LOG(LogMovieScene, Warning, TEXT("Failed to create scrubbed sound for root audio track (row %d)."), RowIndex);
+		}
+		else
+		{
+			ScrubbedSound->SetFlags(RF_Transient);
+			ExistingData->ScrubbedSound = TStrongObjectPtr<UScrubbedSound>(ScrubbedSound);
+		}
+#endif // WITH_EDITOR
 	}
 
 	return ExistingData;
diff --git a/Engine/Source/Runtime/MovieSceneTracks/Private/Systems/MovieSceneAudioSystem.h b/Engine/Source/Runtime/MovieSceneTracks/Private/Systems/MovieSceneAudioSystem.h
--- a/Engine/Source/Runtime/MovieSceneTracks/Private/Systems/MovieSceneAudioSystem.h
+++ b/Engine/Source/Runtime/MovieSceneTracks/Private/Systems/MovieSceneAudioSystem.h
@@ -39,7 +39,11 @@ namespace UE::MovieScene
 	{
 		/** The audio component that was created to play audio */
 		TWeakObjectPtr<UAudioComponent> AudioComponent;
-		TWeakObjectPtr<UScrubbedSound> ScrubbedSound;
+
+#if WITH_EDITOR
+		/** While in editor, we can scrub the audio in the audio component. */
+		TStrongObjectPtr<UScrubbedSound> ScrubbedSound;
+#endif
 
 		/** Volume multiplier to use this frame */
 		double VolumeMultiplier;
