Sequencer: fix some camera cut blending situations

We may go from one camera cut section to another without necessarily starting the new cut at the beginning. That is, we might cut to a new camera in the middle of its camera cut section.

This change make the blend factor available to the game (before it was only available in editor builds), and checks for it to decide whether blending should be used.

diff --git a/Engine/Source/Runtime/MovieScene/Public/IMovieScenePlayer.h b/Engine/Source/Runtime/MovieScene/Public/IMovieScenePlayer.h
--- a/Engine/Source/Runtime/MovieScene/Public/IMovieScenePlayer.h
+++ b/Engine/Source/Runtime/MovieScene/Public/IMovieScenePlayer.h
@@ -65,15 +65,17 @@ struct FMovieSceneCameraCutParams
 	/** Blending type to use to get to the new shot (only used when BlendTime is greater than 0) */
 	TOptional<EMovieSceneBuiltInEasing> BlendType;
 
-	/** When blending, whether to lock the previous camera */
-	bool bLockPreviousCamera = false;
-
-#if WITH_EDITOR
-	// Info for previewing shot blends in editor.
+	/** The previous camera object, if any */
 	UObject* PreviousCameraObject = nullptr;
+
+	/** The computed blend factor, if blending is enabled */
 	float PreviewBlendFactor = -1.f;
+
+	/** When blending, whether to lock the previous camera */
+	bool bLockPreviousCamera = false;
+	
+	/** Whether the camera cut track had blending enabled */
 	bool bCanBlend = false;
-#endif
 };
 
 /** Backwards compatibility to old struct name with typo */
diff --git a/Engine/Source/Runtime/MovieSceneTracks/Private/TrackInstances/MovieSceneCameraCutGameHandler.cpp b/Engine/Source/Runtime/MovieSceneTracks/Private/TrackInstances/MovieSceneCameraCutGameHandler.cpp
--- a/Engine/Source/Runtime/MovieSceneTracks/Private/TrackInstances/MovieSceneCameraCutGameHandler.cpp
+++ b/Engine/Source/Runtime/MovieSceneTracks/Private/TrackInstances/MovieSceneCameraCutGameHandler.cpp
@@ -339,7 +339,15 @@ void FCameraCutGameHandler::SetCameraCut(
 	// Time to set the camera cut! How we do it depends on whether we need to do some blending, or a straight cut.
 	bool bDoSetViewTarget = true;
 	FViewTargetTransitionParams TransitionParams;
-	if (CameraCutParams.BlendType.IsSet())
+	const bool bDoBlend = (
+			CameraCutParams.BlendType.IsSet() &&
+			CameraCutParams.BlendTime > 0.f &&
+			CameraCutParams.PreviewBlendFactor < 1.f   // If the blend factor is already 100%, we should cut to that camera
+													   // instead of blending to it. This can happen when cutting back to
+													   // the middle of a camera cut section. Even if that section has an
+													   // ease-in, we should simply cut to it.
+			);  
+	if (bDoBlend)
 	{
 		UE_LOG(LogMovieScene, Log, TEXT("Blending into new camera cut: '%s' -> '%s' (blend time: %f)"),
 			(ViewTarget ? *ViewTarget->GetName() : TEXT("None")),
diff --git a/Engine/Source/Runtime/MovieSceneTracks/Private/TrackInstances/MovieSceneCameraCutTrackInstance.cpp b/Engine/Source/Runtime/MovieSceneTracks/Private/TrackInstances/MovieSceneCameraCutTrackInstance.cpp
--- a/Engine/Source/Runtime/MovieSceneTracks/Private/TrackInstances/MovieSceneCameraCutTrackInstance.cpp
+++ b/Engine/Source/Runtime/MovieSceneTracks/Private/TrackInstances/MovieSceneCameraCutTrackInstance.cpp
@@ -244,12 +244,10 @@ public:
 		CameraCutParams.BlendType = Params.EaseIn.BlendType;
 		CameraCutParams.bLockPreviousCamera = Params.bLockPreviousCamera;
 
-#if WITH_EDITOR
 		UObject* PreviousCameraActor = FindBoundObject(Params.PreviousCameraBindingID, Params.PreviousOperandSequenceID, SequenceInstance.GetSharedPlaybackState());
 		CameraCutParams.PreviousCameraObject = PreviousCameraActor;
 		CameraCutParams.PreviewBlendFactor = Params.PreviewBlendFactor;
 		CameraCutParams.bCanBlend = Params.bCanBlend;
-#endif
 
 		const bool bMatchesCache = MatchesCameraCutCache(CameraActor, Params, CameraCutCache);
 		if (!bMatchesCache)
