Sequencer: Ported Bone matching to common animation track so that Animation Mixer's can use the functionality, tweaks and fixes to root motnion handling with different spaces and transform origins

diff --git a/Engine/Plugins/Experimental/MovieSceneAnimMixer/Source/MovieSceneAnimMixer/Private/Systems/MovieSceneAnimMixerSystem.cpp b/Engine/Plugins/Experimental/MovieSceneAnimMixer/Source/MovieSceneAnimMixer/Private/Systems/MovieSceneAnimMixerSystem.cpp
--- a/Engine/Plugins/Experimental/MovieSceneAnimMixer/Source/MovieSceneAnimMixer/Private/Systems/MovieSceneAnimMixerSystem.cpp
+++ b/Engine/Plugins/Experimental/MovieSceneAnimMixer/Source/MovieSceneAnimMixer/Private/Systems/MovieSceneAnimMixerSystem.cpp
@@ -19,6 +19,7 @@
 
 #include "Systems/WeightAndEasingEvaluatorSystem.h"
 #include "Systems/MovieScenePiecewiseDoubleBlenderSystem.h"
+#include "Systems/MoviesceneTransformOriginSystem.h"
 #include "EvaluationVM/Tasks/ExecuteProgram.h"
 #include "EvaluationVM/Tasks/BlendKeyframes.h"
 #include "EvaluationVM/Tasks/ApplyAdditiveKeyframe.h"
@@ -67,19 +68,24 @@ namespace UE::MovieScene
 	struct FEvaluateAnimMixers
 	{
 		TMap<FMovieSceneAnimMixerKey, TSharedPtr<FMovieSceneAnimMixer>>* Mixers;
-		UMovieSceneRootMotionSystem* RootMotionSystem;
+		const UMovieSceneRootMotionSystem* RootMotionSystem;
+		const UMovieSceneTransformOriginSystem* TransformOriginSystem;
 
-		FEvaluateAnimMixers(TMap<FMovieSceneAnimMixerKey, TSharedPtr<FMovieSceneAnimMixer>>* InMixers, UMovieSceneRootMotionSystem* InRootMotionSystem)
+		FEvaluateAnimMixers(TMap<FMovieSceneAnimMixerKey, TSharedPtr<FMovieSceneAnimMixer>>* InMixers, UMovieSceneEntitySystemLinker* InLinker)
 			: Mixers(InMixers)
-			, RootMotionSystem(InRootMotionSystem)
 		{
+			RootMotionSystem = InLinker->FindSystem<UMovieSceneRootMotionSystem>();
+			TransformOriginSystem = InLinker->FindSystem<UMovieSceneTransformOriginSystem>();
 		}
 
 		void ForEachEntity(
 			FObjectComponent MeshComponent,
 			const TInstancedStruct<FMovieSceneMixedAnimationTarget>& Target,
-			TSharedPtr<FAnimNextEvaluationTask>& MixerTask) const
+			TSharedPtr<FAnimNextEvaluationTask>& MixerTask
+			) const
 		{
+			using ESpaceConversions = FAnimNextConvertRootMotionToWorldSpaceTask::ESpaceConversions;
+
 			TSharedPtr<FMovieSceneAnimMixer> Mixer = Mixers->FindRef(FMovieSceneAnimMixerKey({ MeshComponent.GetObject(), Target }));
 			if (!Mixer)
 			{
@@ -93,6 +99,12 @@ namespace UE::MovieScene
 				return;
 			}
 
+			USceneComponent* Component = Cast<USceneComponent>(MeshComponent.GetObject());
+			USceneComponent* Root = Component ? Component->GetOwner()->GetRootComponent() : nullptr;
+
+			const bool bComponentHasKeyedTransform = (Component && RootMotionSystem->IsTransformKeyed(Component));
+			const bool bRootComponentHasKeyedTransform = (Root && RootMotionSystem->IsTransformKeyed(Root));
+
 			// Create a new eval program
 			Mixer->EvaluationProgram = MakeShared<UE::AnimNext::FEvaluationProgram>();
 
@@ -208,10 +220,39 @@ namespace UE::MovieScene
 					// Evaluate the pose itself
 					Mixer->EvaluationProgram->AppendTaskPtr(MixerEntry->EvalTask);
 
-					if (RootMotion && MixerEntry->RootMotionSettings && MixerEntry->RootMotionSettings->RootMotionSpace == EMovieSceneRootMotionSpace::AnimationSpace)
+					FTransform TransformOrigin;
+					ESpaceConversions Conversion = ESpaceConversions::None;
+
+					EMovieSceneRootMotionSpace ThisRootSpace = MixerEntry->RootMotionSettings
+						? MixerEntry->RootMotionSettings->RootMotionSpace
+						: EMovieSceneRootMotionSpace::AnimationSpace;
+
+					FTransform RootBaseTransform = MixerEntry->RootMotionSettings
+						? FTransform(MixerEntry->RootMotionSettings->RootBaseRotation, MixerEntry->RootMotionSettings->RootBaseLocation)
+						: FTransform::Identity;
+
+					if (ThisRootSpace == EMovieSceneRootMotionSpace::AnimationSpace)
+					{
+						Conversion = ESpaceConversions::ComponentToActorRotation | ESpaceConversions::PreBaseTransform;
+						if (bRootComponentHasKeyedTransform)
+						{
+							Conversion |= ESpaceConversions::AnimationToWorld;
+						}
+						else if (TransformOriginSystem && TransformOriginSystem->GetTransformOrigin(MixerEntry->InstanceHandle, TransformOrigin))
+						{
+							Conversion |= ESpaceConversions::TransformOriginToWorld;
+						}
+					}
+					else if (RootBaseTransform.GetLocation().IsZero() || !RootBaseTransform.GetRotation().IsIdentity())
+					{
+						Conversion = ESpaceConversions::PostBaseTransform;
+					}
+
+					if (Conversion != ESpaceConversions::None)
 					{
 						// Need to convert root motion attribute to world space for blending, insert task to do that
-						Mixer->EvaluationProgram->AppendTask(FAnimNextConvertRootMotionToWorldSpaceTask::Make(RootMotion));
+						Mixer->EvaluationProgram->AppendTask(FAnimNextConvertRootMotionToWorldSpaceTask::Make(
+							RootMotion, TransformOrigin, RootBaseTransform, Conversion));
 					}
 
 					if (MixerEntry->bRequiresBlend)
@@ -256,12 +297,6 @@ namespace UE::MovieScene
 
 			if (RootMotion)
 			{
-				USceneComponent* Component = Cast<USceneComponent>(MeshComponent.GetObject());
-				USceneComponent* Root      = Component ? Component->GetOwner()->GetRootComponent() : nullptr;
-
-				const bool bComponentHasKeyedTransform     = (Component && RootMotionSystem->IsTransformKeyed(Component));
-				const bool bRootComponentHasKeyedTransform = (Root      && RootMotionSystem->IsTransformKeyed(Root));
-
 				// Add task to store the final root motion result
 				Mixer->EvaluationProgram->AppendTask(FAnimNextStoreRootTransformTask::Make(RootMotion, bComponentHasKeyedTransform, bRootComponentHasKeyedTransform));
 			}
@@ -336,6 +371,7 @@ UMovieSceneAnimMixerSystem::UMovieSceneAnimMixerSystem(const FObjectInitializer&
 		DefineImplicitPrerequisite(UMovieSceneBoundSceneComponentInstantiator::StaticClass(), GetClass());
 		DefineImplicitPrerequisite(UWeightAndEasingEvaluatorSystem::StaticClass(), GetClass());
 		DefineImplicitPrerequisite(UMovieSceneComponentTransformSystem::StaticClass(), GetClass());
+		DefineImplicitPrerequisite(UMovieSceneTransformOriginSystem::StaticClass(), GetClass());
 		DefineImplicitPrerequisite(UByteChannelEvaluatorSystem::StaticClass(), GetClass());
 		DefineComponentConsumer(GetClass(), BuiltInComponents->BoundObjectKey);
 		DefineComponentConsumer(GetClass(), BuiltInComponents->BoundObject);
@@ -430,6 +466,7 @@ void UMovieSceneAnimMixerSystem::OnRun(FSystemTaskPrerequisites& InPrerequisites
 
 		FEntityTaskBuilder()
 		.ReadEntityIDs()
+		.Read(BuiltInComponents->InstanceHandle)
 		.Read(BuiltInComponents->BoundObjectKey)
 		.Write(AnimMixerComponents->Target)
 		.Write(AnimMixerComponents->MixerEntry)
@@ -441,6 +478,7 @@ void UMovieSceneAnimMixerSystem::OnRun(FSystemTaskPrerequisites& InPrerequisites
 		[this, &NewMixers, BuiltInComponents, AnimMixerComponents]
 		(
 			FMovieSceneEntityID EntityID,
+			FInstanceHandle InstanceHandle,
 			FObjectKey BoundObjectKey,
 			TInstancedStruct<FMovieSceneMixedAnimationTarget>& Target,
 			TSharedPtr<FMovieSceneAnimMixerEntry>& InOutMixerEntry,
@@ -471,6 +509,7 @@ void UMovieSceneAnimMixerSystem::OnRun(FSystemTaskPrerequisites& InPrerequisites
 
 				FMovieSceneAnimMixerEntry* Entry = InOutMixerEntry.Get();
 
+				Entry->InstanceHandle = InstanceHandle;
 				Entry->EntityID = EntityID;
 				Entry->EvalTask = Task;
 				Entry->Priority = Priority;
@@ -626,7 +665,7 @@ void UMovieSceneAnimMixerSystem::OnSchedulePersistentTasks(UE::MovieScene::IEnti
 		.Read(AnimMixerComponents->MeshComponent)
 		.Read(AnimMixerComponents->Target)
 		.Write(AnimMixerComponents->MixerTask)
-		.Schedule_PerEntity<FEvaluateAnimMixers>(&Linker->EntityManager, TaskScheduler, &Mixers, Linker->FindSystem<UMovieSceneRootMotionSystem>());
+		.Schedule_PerEntity<FEvaluateAnimMixers>(&Linker->EntityManager, TaskScheduler, &Mixers, Linker);
 
 	TaskScheduler->AddPrerequisite(UpdateTask, MixTask);
 }
diff --git a/Engine/Plugins/Experimental/MovieSceneAnimMixer/Source/MovieSceneAnimMixer/Private/Systems/MovieSceneMixedSkeletalAnimationSystem.cpp b/Engine/Plugins/Experimental/MovieSceneAnimMixer/Source/MovieSceneAnimMixer/Private/Systems/MovieSceneMixedSkeletalAnimationSystem.cpp
--- a/Engine/Plugins/Experimental/MovieSceneAnimMixer/Source/MovieSceneAnimMixer/Private/Systems/MovieSceneMixedSkeletalAnimationSystem.cpp
+++ b/Engine/Plugins/Experimental/MovieSceneAnimMixer/Source/MovieSceneAnimMixer/Private/Systems/MovieSceneMixedSkeletalAnimationSystem.cpp
@@ -29,6 +29,7 @@
 #include "Systems/MovieSceneComponentTransformSystem.h"
 #include "Systems/MovieSceneQuaternionInterpolationRotationSystem.h"
 #include "Systems/ByteChannelEvaluatorSystem.h"
+#include "Systems/MoviesceneTransformOriginSystem.h"
 #include "EntitySystem/MovieSceneBoundObjectInstantiator.h"
 #include "EntitySystem/MovieSceneBoundSceneComponentInstantiator.h"
 #include "EntitySystem/MovieSceneEntityMutations.h"
@@ -136,12 +137,13 @@ struct FSkeletalAnimMixerMutation : IMovieSceneEntityMutation
 
 
 				bool bHasRootMotionOverride = false;
-				UMovieSceneSkeletalAnimationTrack* Track = SkeletalAnim.Section->GetTypedOuter<UMovieSceneSkeletalAnimationTrack>();
+				UMovieSceneCommonAnimationTrack* Track = SkeletalAnim.Section->GetTypedOuter<UMovieSceneCommonAnimationTrack>();
+				UMovieSceneSkeletalAnimationTrack* SkelAnimTrack = Cast<UMovieSceneSkeletalAnimationTrack>(Track);
 
 				ESwapRootBone LegacySwapRootBone = SkeletalAnim.Section->Params.SwapRootBone;
-				if (LegacySwapRootBone == ESwapRootBone::SwapRootBone_None && Track)
+				if (LegacySwapRootBone == ESwapRootBone::SwapRootBone_None && SkelAnimTrack)
 				{
-					LegacySwapRootBone = Track->SwapRootBone;
+					LegacySwapRootBone = SkelAnimTrack->SwapRootBone;
 				}
 				OutRootMotionSettings[Index].LegacySwapRootBone = LegacySwapRootBone;
 
@@ -165,7 +167,6 @@ struct FSkeletalAnimMixerMutation : IMovieSceneEntityMutation
 					}
 				}
 
-
 				if (Track)
 				{
 					if (Track->RootMotionParams.bRootMotionsDirty) 
@@ -217,7 +218,6 @@ struct FGatherMixableRootMotion
 				RootMotionSettings.RootMotionOverride->RootMotion = RootMotionParams.Transform.GetValue();
 				RootMotionSettings.RootMotionOverride->bBlendFirstChildOfRoot = RootMotionParams.bBlendFirstChildOfRoot;
 				RootMotionSettings.RootMotionOverride->ChildBoneIndex = RootMotionParams.ChildBoneIndex;
-				RootMotionSettings.RootMotionOverride->PreviousTransform = RootMotionParams.PreviousTransform.GetValue();
 			}
 		}
 	}
@@ -229,10 +229,12 @@ struct FGatherMixableRootMotion
 struct FGatherMixableSkeletalAnimations
 {
 	const FInstanceRegistry* InstanceRegistry;
+	const UMovieSceneTransformOriginSystem* TransformOriginSystem;
 
-	FGatherMixableSkeletalAnimations(const FInstanceRegistry* InInstanceRegistry)
-		: InstanceRegistry(InInstanceRegistry)
+	FGatherMixableSkeletalAnimations(UMovieSceneEntitySystemLinker* InLinker)
+		: InstanceRegistry(InLinker->GetInstanceRegistry())
 	{
+		TransformOriginSystem = InLinker->FindSystem<UMovieSceneTransformOriginSystem>();
 	}
 
 	void ForEachAllocation(
@@ -242,7 +244,7 @@ struct FGatherMixableSkeletalAnimations
 		TRead<UObject*> BoundObjects,
 		TRead<FMovieSceneSkeletalAnimationComponentData> SkeletalAnimations,
 		TReadOptional<FFrameTime> OptionalEvalTimes,
-		TReadOptional<FMovieSceneRootMotionSettings> OptionalRootMotionSettings,
+		TWriteOptional<FMovieSceneRootMotionSettings> OptionalRootMotionSettings,
 		const double* OptLocationX, const double* OptLocationY, const double* OptLocationZ,
 		const double* OptRotationX, const double* OptRotationY, const double* OptRotationZ,
 		const uint8* OptSpaces,
@@ -366,22 +368,20 @@ struct FGatherMixableSkeletalAnimations
 			TSharedPtr<FMovieSceneSkeletalAnimationEvaluationTask> AnimTask = StaticCastSharedPtr<FMovieSceneSkeletalAnimationEvaluationTask>(EvalTask[Index]);
 			AnimTask->AnimationData = AnimData;
 
-			FVector  CurrentLocation = AnimTask->RootBaseTransform.GetTranslation();
-			FRotator CurrentRotation = AnimTask->RootBaseTransform.GetRotation().Rotator();
-
-			if (OptLocationX) { CurrentLocation.X = OptLocationX[Index]; }
-			if (OptLocationY) { CurrentLocation.Y = OptLocationY[Index]; }
-			if (OptLocationZ) { CurrentLocation.Z = OptLocationZ[Index]; }
-			if (OptRotationX) { CurrentRotation.Roll  = OptRotationX[Index]; }
-			if (OptRotationY) { CurrentRotation.Pitch = OptRotationY[Index]; }
-			if (OptRotationZ) { CurrentRotation.Yaw   = OptRotationZ[Index]; }
-
-			AnimTask->RootBaseTransform.SetRotation(CurrentRotation.Quaternion());
-			AnimTask->RootBaseTransform.SetLocation(CurrentLocation);
-
-			if (OptSpaces)
+			if (OptionalRootMotionSettings)
 			{
-				AnimTask->RootMotionSpace = (EMovieSceneRootMotionSpace)(OptSpaces[Index]);
+				FVector  CurrentLocation = OptionalRootMotionSettings[Index].RootBaseLocation;
+				FRotator CurrentRotation = OptionalRootMotionSettings[Index].RootBaseRotation.Rotator();
+
+				if (OptLocationX) { CurrentLocation.X = OptLocationX[Index]; }
+				if (OptLocationY) { CurrentLocation.Y = OptLocationY[Index]; }
+				if (OptLocationZ) { CurrentLocation.Z = OptLocationZ[Index]; }
+				if (OptRotationX) { CurrentRotation.Roll  = OptRotationX[Index]; }
+				if (OptRotationY) { CurrentRotation.Pitch = OptRotationY[Index]; }
+				if (OptRotationZ) { CurrentRotation.Yaw   = OptRotationZ[Index]; }
+
+				OptionalRootMotionSettings[Index].RootBaseLocation = CurrentLocation;
+				OptionalRootMotionSettings[Index].RootBaseRotation = CurrentRotation.Quaternion();
 			}
 		}
 	}
@@ -494,14 +494,7 @@ void FMovieSceneSkeletalAnimationEvaluationTask::Execute(UE::AnimNext::FEvaluati
 			FTransformAnimationAttribute* RootMotionAttribute        = Keyframe.Attributes.FindOrAdd<FTransformAnimationAttribute>(AnimMixerComponents->RootTransformAttributeId);
 			FFloatAnimationAttribute*     RootMotionWeight           = Keyframe.Attributes.FindOrAdd<FFloatAnimationAttribute>(AnimMixerComponents->RootTransformWeightAttributeId);
 
-			if (RootMotionSpace == EMovieSceneRootMotionSpace::AnimationSpace)
-			{
-				RootMotionAttribute->Value = RootBoneTransform.GetValue() * RootBaseTransform;
-			}
-			else
-			{
-				RootMotionAttribute->Value = RootBaseTransform * RootBoneTransform.GetValue();
-			}
+			RootMotionAttribute->Value = RootBoneTransform.GetValue();
 			RootMotionWeight->Value = 1.f;
 		}
 
@@ -534,6 +527,7 @@ UMovieSceneMixedSkeletalAnimationSystem::UMovieSceneMixedSkeletalAnimationSystem
 		DefineImplicitPrerequisite(UByteChannelEvaluatorSystem::StaticClass(), GetClass());
 		DefineImplicitPrerequisite(UMovieSceneGenericBoundObjectInstantiator::StaticClass(), GetClass());
 		DefineImplicitPrerequisite(UMovieSceneBoundSceneComponentInstantiator::StaticClass(), GetClass());
+		DefineImplicitPrerequisite(UMovieSceneTransformOriginSystem::StaticClass(), GetClass());
 		DefineComponentConsumer(GetClass(), BuiltInComponents->BoundObject);
 		DefineComponentConsumer(GetClass(), TrackComponents->SkeletalAnimation);
 		DefineComponentProducer(GetClass(), TrackComponents->Tags.AnimMixerPoseProducer);
@@ -597,7 +591,7 @@ void UMovieSceneMixedSkeletalAnimationSystem::OnSchedulePersistentTasks(UE::Movi
 		.Read(BuiltInComponents->BoundObject)
 		.Read(TrackComponents->SkeletalAnimation)
 		.ReadOptional(BuiltInComponents->EvalTime)
-		.ReadOptional(AnimMixerComponents->RootMotionSettings)
+		.WriteOptional(AnimMixerComponents->RootMotionSettings)
 		.ReadOptional(BuiltInComponents->DoubleResult[0]) // Root Offset X
 		.ReadOptional(BuiltInComponents->DoubleResult[1]) // Root Offset Y
 		.ReadOptional(BuiltInComponents->DoubleResult[2]) // Root Offset Z
@@ -609,7 +603,7 @@ void UMovieSceneMixedSkeletalAnimationSystem::OnSchedulePersistentTasks(UE::Movi
 		.FilterAll( { TrackComponents->Tags.AnimMixerPoseProducer })
 		.FilterNone({ BuiltInComponents->Tags.Ignored})
 		.Schedule_PerAllocation<FGatherMixableSkeletalAnimations>(&Linker->EntityManager, TaskScheduler,
-			Linker->GetInstanceRegistry());
+			Linker);
 
 	TaskScheduler->AddPrerequisite(GatherRootMotionTask, GatherTask);
 }
diff --git a/Engine/Plugins/Experimental/MovieSceneAnimMixer/Source/MovieSceneAnimMixer/Private/Systems/MovieSceneRootMotionSystem.cpp b/Engine/Plugins/Experimental/MovieSceneAnimMixer/Source/MovieSceneAnimMixer/Private/Systems/MovieSceneRootMotionSystem.cpp
--- a/Engine/Plugins/Experimental/MovieSceneAnimMixer/Source/MovieSceneAnimMixer/Private/Systems/MovieSceneRootMotionSystem.cpp
+++ b/Engine/Plugins/Experimental/MovieSceneAnimMixer/Source/MovieSceneAnimMixer/Private/Systems/MovieSceneRootMotionSystem.cpp
@@ -173,14 +173,17 @@ void FMovieSceneMixerRootMotionComponentData::Initialize()
 	}
 }
 
-FAnimNextConvertRootMotionToWorldSpaceTask::FAnimNextConvertRootMotionToWorldSpaceTask(const TSharedPtr<FMovieSceneMixerRootMotionComponentData>& InRootMotionData)
+FAnimNextConvertRootMotionToWorldSpaceTask::FAnimNextConvertRootMotionToWorldSpaceTask(const TSharedPtr<FMovieSceneMixerRootMotionComponentData>& InRootMotionData, const FTransform& InTransformOrigin, const FTransform& InBaseTransform, ESpaceConversions InConversions)
 	: WeakRootMotionData(InRootMotionData)
+	, TransformOrigin(InTransformOrigin)
+	, BaseTransform(InBaseTransform)
+	, Conversions(InConversions)
 {
 }
 
-FAnimNextConvertRootMotionToWorldSpaceTask FAnimNextConvertRootMotionToWorldSpaceTask::Make(const TSharedPtr<FMovieSceneMixerRootMotionComponentData>& InRootMotionData)
+FAnimNextConvertRootMotionToWorldSpaceTask FAnimNextConvertRootMotionToWorldSpaceTask::Make(const TSharedPtr<FMovieSceneMixerRootMotionComponentData>& InRootMotionData, const FTransform& InTransformOrigin, const FTransform& InBaseTransform, ESpaceConversions InConversions)
 {
-	return FAnimNextConvertRootMotionToWorldSpaceTask(InRootMotionData);
+	return FAnimNextConvertRootMotionToWorldSpaceTask(InRootMotionData, InTransformOrigin, InBaseTransform, InConversions);
 }
 
 void FAnimNextConvertRootMotionToWorldSpaceTask::Execute(UE::AnimNext::FEvaluationVM& VM) const
@@ -199,24 +202,50 @@ void FAnimNextConvertRootMotionToWorldSpaceTask::Execute(UE::AnimNext::FEvaluati
 		return;
 	}
 
+	const FAnimMixerComponentTypes* AnimMixerComponents = FAnimMixerComponentTypes::Get();
+
+	FTransformAnimationAttribute* RootMotionAttribute = (*Keyframe)->Attributes.Find<FTransformAnimationAttribute>(AnimMixerComponents->RootTransformAttributeId);
+	if (!RootMotionAttribute)
+	{
+		return;
+	}
+
+	FTransform RelativeTransform = RootMotionAttribute->Value;
 
-	FAnimMixerComponentTypes* AnimMixerComponents = FAnimMixerComponentTypes::Get();
 	TSharedPtr<FMovieSceneMixerRootMotionComponentData> RootMotionData = WeakRootMotionData.Pin();
 	if (RootMotionData)
 	{
-		if (FTransformAnimationAttribute* RootMotionAttribute = (*Keyframe)->Attributes.Find<FTransformAnimationAttribute>(AnimMixerComponents->RootTransformAttributeId))
+		if (EnumHasAnyFlags(Conversions, ESpaceConversions::ComponentToActorRotation))
 		{
-			FTransform RelativeTransform = RootMotionAttribute->Value;
-
 			TOptional<FQuat> InverseMeshToActorRotation = RootMotionData->GetInverseMeshToActorRotation();
 			if (InverseMeshToActorRotation.IsSet())
 			{
 				RelativeTransform.SetTranslation(InverseMeshToActorRotation.GetValue() * RelativeTransform.GetTranslation());
 			}
+		}
 
-			RootMotionAttribute->Value = RelativeTransform * RootMotionData->ActorTransform;
+		if (EnumHasAnyFlags(Conversions, ESpaceConversions::PreBaseTransform))
+		{
+			RelativeTransform = RelativeTransform * BaseTransform;
+		}
+		else if (EnumHasAnyFlags(Conversions, ESpaceConversions::PostBaseTransform))
+		{
+			RelativeTransform = BaseTransform * RelativeTransform;
+		}
+
+		if (EnumHasAnyFlags(Conversions, ESpaceConversions::AnimationToWorld))
+		{
+			RelativeTransform *= RootMotionData->ActorTransform;
 		}
 	}
+
+	// Doesn't really make sense to mix both AnimationToWorld & TransformOriginToWorld
+	if (EnumHasAnyFlags(Conversions, ESpaceConversions::TransformOriginToWorld))
+	{
+		RelativeTransform *= TransformOrigin;
+	}
+
+	RootMotionAttribute->Value = RelativeTransform;
 }
 
 FAnimNextStoreRootTransformTask::FAnimNextStoreRootTransformTask(const TSharedPtr<FMovieSceneMixerRootMotionComponentData>& InRootMotionData, bool bInComponentHasKeyedTransform, bool bInRootComponentHasKeyedTransform)
@@ -551,6 +580,7 @@ void UMovieSceneRootMotionSystem::OnRun(FSystemTaskPrerequisites& InPrerequisite
 	FEntityTaskBuilder()
 	.Read(BuiltInComponents->BoundObjectKey)
 	.FilterAll({ TrackComponents->ComponentTransform.PropertyTag, BuiltInComponents->CustomPropertyIndex })
+	.FilterNone({ BuiltInComponents->Tags.NeedsUnlink })
 	.Iterate_PerAllocation(&Linker->EntityManager,
 		[this](const FEntityAllocation* Allocation, const FObjectKey* RootMotionTargets)
 		{
diff --git a/Engine/Plugins/Experimental/MovieSceneAnimMixer/Source/MovieSceneAnimMixer/Public/Systems/MovieSceneAnimMixerSystem.h b/Engine/Plugins/Experimental/MovieSceneAnimMixer/Source/MovieSceneAnimMixer/Public/Systems/MovieSceneAnimMixerSystem.h
--- a/Engine/Plugins/Experimental/MovieSceneAnimMixer/Source/MovieSceneAnimMixer/Public/Systems/MovieSceneAnimMixerSystem.h
+++ b/Engine/Plugins/Experimental/MovieSceneAnimMixer/Source/MovieSceneAnimMixer/Public/Systems/MovieSceneAnimMixerSystem.h
@@ -39,6 +39,7 @@ struct FMovieSceneAnimMixerEntry
 	bool bAdditive = false;
 	bool bRequiresBlend = true;
 	UE::MovieScene::FMovieSceneEntityID EntityID;
+	UE::MovieScene::FInstanceHandle InstanceHandle;
 
 
 	FORCEINLINE bool operator<(const FMovieSceneAnimMixerEntry& RHS) const
diff --git a/Engine/Plugins/Experimental/MovieSceneAnimMixer/Source/MovieSceneAnimMixer/Public/Systems/MovieSceneMixedSkeletalAnimationSystem.h b/Engine/Plugins/Experimental/MovieSceneAnimMixer/Source/MovieSceneAnimMixer/Public/Systems/MovieSceneMixedSkeletalAnimationSystem.h
--- a/Engine/Plugins/Experimental/MovieSceneAnimMixer/Source/MovieSceneAnimMixer/Public/Systems/MovieSceneMixedSkeletalAnimationSystem.h
+++ b/Engine/Plugins/Experimental/MovieSceneAnimMixer/Source/MovieSceneAnimMixer/Public/Systems/MovieSceneMixedSkeletalAnimationSystem.h
@@ -55,12 +55,6 @@ struct FMovieSceneSkeletalAnimationEvaluationTask : public FAnimNextEvaluationTa
 
 	UPROPERTY()
 	FMixedAnimSkeletalAnimationData AnimationData;
-
-	UPROPERTY()
-	FTransform RootBaseTransform;
-
-	UPROPERTY()
-	EMovieSceneRootMotionSpace RootMotionSpace = EMovieSceneRootMotionSpace::AnimationSpace;
 };
 
 // System to handle creating evaluation tasks from skeletal animation track sections for the anim mixer.
diff --git a/Engine/Plugins/Experimental/MovieSceneAnimMixer/Source/MovieSceneAnimMixer/Public/Systems/MovieSceneRootMotionSystem.h b/Engine/Plugins/Experimental/MovieSceneAnimMixer/Source/MovieSceneAnimMixer/Public/Systems/MovieSceneRootMotionSystem.h
--- a/Engine/Plugins/Experimental/MovieSceneAnimMixer/Source/MovieSceneAnimMixer/Public/Systems/MovieSceneRootMotionSystem.h
+++ b/Engine/Plugins/Experimental/MovieSceneAnimMixer/Source/MovieSceneAnimMixer/Public/Systems/MovieSceneRootMotionSystem.h
@@ -23,10 +23,7 @@ struct FSkeletalAnimationRootMotionOverride
 	
 	UPROPERTY()
 	FTransform RootMotion = FTransform::Identity;
-	
-	UPROPERTY()
-	FTransform PreviousTransform = FTransform::Identity;
-	
+
 	UPROPERTY()
 	int32 ChildBoneIndex = INDEX_NONE;
 
@@ -40,6 +37,9 @@ struct FMovieSceneRootMotionSettings
 {
 	TOptional<FSkeletalAnimationRootMotionOverride> RootMotionOverride;
 
+	FVector RootBaseLocation = FVector(EForceInit::ForceInit);
+	FQuat   RootBaseRotation = FQuat(EForceInit::ForceInit);
+
 	// What space to apply root motion in. Defaults to animation space.
 	EMovieSceneRootMotionSpace RootMotionSpace = EMovieSceneRootMotionSpace::AnimationSpace;
 
@@ -79,7 +79,7 @@ private:
 	mutable FTransactionallySafeRWLock RootMotionLock;
 };
 
-// Task that converts the root motion attribute on the top pose of the pose stack to world space by adding on the blended transform track value.
+// Task that converts the root motion attribute on the top pose of the pose stack to world space by adding on the actor transformation, root base transform, and/or transform origin.
 USTRUCT()
 struct FAnimNextConvertRootMotionToWorldSpaceTask : public FAnimNextEvaluationTask
 {
@@ -87,15 +87,39 @@ struct FAnimNextConvertRootMotionToWorldSpaceTask : public FAnimNextEvaluationTa
 
 	DECLARE_ANIM_EVALUATION_TASK(FAnimNextConvertRootMotionToWorldSpaceTask)
 
+	/** Enum specifying which conversions to perform */
+	enum class ESpaceConversions : uint8
+	{
+		None = 0,
+		/** Convert the root motion from animation space to world space */
+		AnimationToWorld = 1 << 0,
+		/** Convert the root motion from transform origin space to world space (used when there is no transform track in Sequencer) */
+		TransformOriginToWorld = 1 << 1,
+		/** Convert the root motion from component -> actor space using the inverse component rotation only */
+		ComponentToActorRotation = 1 << 2,
+		/** Apply RootBaseTransform as Root = Root * RootBaseTransform */
+		PreBaseTransform = 1 << 3,
+		/** Apply RootBaseTransform as Root = RootBaseTransform * Root */
+		PostBaseTransform = 1 << 4,
+	};
+
 	FAnimNextConvertRootMotionToWorldSpaceTask() = default;
-	FAnimNextConvertRootMotionToWorldSpaceTask(const TSharedPtr<FMovieSceneMixerRootMotionComponentData>& InRootMotionData);
+	FAnimNextConvertRootMotionToWorldSpaceTask(const TSharedPtr<FMovieSceneMixerRootMotionComponentData>& InRootMotionData, const FTransform& InTransformOrigin, const FTransform& InBaseTransform, ESpaceConversions InConversions);
 
-	static FAnimNextConvertRootMotionToWorldSpaceTask Make(const TSharedPtr<FMovieSceneMixerRootMotionComponentData>& InRootMotionData);
+	static FAnimNextConvertRootMotionToWorldSpaceTask Make(const TSharedPtr<FMovieSceneMixerRootMotionComponentData>& InRootMotionData, const FTransform& InTransformOrigin, const FTransform& InBaseTransform, ESpaceConversions InConversions);
 
 	virtual void Execute(UE::AnimNext::FEvaluationVM& VM) const override;
 
+	/** Weak pointer to the root motion data for all mixed animations. May be null if only a transform origin transformation is required. */
 	TWeakPtr<FMovieSceneMixerRootMotionComponentData> WeakRootMotionData;
+	/** Base transformation to apply to the root in Actor space, before transform origins. Can be used in place of a transform track. */
+	FTransform BaseTransform;
+	/** Transform origin to apply to the root, if Conversions & ESpaceConversions::TransformOriginToWorld */
+	FTransform TransformOrigin;
+	/** Bitmask specifying the conversions to perform */
+	ESpaceConversions Conversions;
 };
+ENUM_CLASS_FLAGS(FAnimNextConvertRootMotionToWorldSpaceTask::ESpaceConversions);
 
 // Task that gets the final mixed root transform and stores it in the root motion data for later application.
 // TODO: it's not ideal that we're writing things outside of the animation system during an evaluation task. 
diff --git a/Engine/Plugins/Experimental/MovieScenePoseSearchTracks/Source/MovieScenePoseSearchTracks/Private/Sections/MovieSceneStitchAnimSection.cpp b/Engine/Plugins/Experimental/MovieScenePoseSearchTracks/Source/MovieScenePoseSearchTracks/Private/Sections/MovieSceneStitchAnimSection.cpp
--- a/Engine/Plugins/Experimental/MovieScenePoseSearchTracks/Source/MovieScenePoseSearchTracks/Private/Sections/MovieSceneStitchAnimSection.cpp
+++ b/Engine/Plugins/Experimental/MovieScenePoseSearchTracks/Source/MovieScenePoseSearchTracks/Private/Sections/MovieSceneStitchAnimSection.cpp
@@ -143,14 +143,17 @@ void FMovieSceneStitchAnimEvaluationTask::Execute(UE::AnimNext::FEvaluationVM& V
 	AnimSpace.Current = AnimSequence->ExtractRootTrackTransform(FAnimExtractContext((double)AnimationTime), nullptr);
 	AnimSpace.Target = AnimSequence->ExtractRootTrackTransform(FAnimExtractContext((double)(AnimationTime + (TimeToTarget * MatchedAssetPlayRate))), nullptr);
 
+	FTransform TargetTransformWorldSpace = StitchData.TargetTransform;
 	if (InverseMeshToActorRotation)
 	{
 		AnimSpace.Initial.SetTranslation(InverseMeshToActorRotation.GetValue() * AnimSpace.Initial.GetTranslation());
 		AnimSpace.Current.SetTranslation(InverseMeshToActorRotation.GetValue() * AnimSpace.Current.GetTranslation());
 		AnimSpace.Target.SetTranslation(InverseMeshToActorRotation.GetValue() * AnimSpace.Target.GetTranslation());
+		TargetTransformWorldSpace.SetRotation(InverseMeshToActorRotation.GetValue() * TargetTransformWorldSpace.GetRotation());
+		TargetTransformWorldSpace.NormalizeRotation();
 	}
 
-	FTransform TargetTransformWorldSpace = StitchData.TargetTransform * SequenceTransformOrigin;
+	TargetTransformWorldSpace *= SequenceTransformOrigin;
 
 	FTransform WorldSpaceProjectedFromStart = AnimSpace.Current.GetRelativeTransform(AnimSpace.Initial) * InitialRootTransform;
 	FTransform WorldSpaceProjectedFromTarget = AnimSpace.Current.GetRelativeTransform(AnimSpace.Target) * TargetTransformWorldSpace;
@@ -265,6 +268,9 @@ void UMovieSceneStitchAnimSection::ImportEntityImpl(UMovieSceneEntitySystemLinke
 	TSharedPtr<FMovieSceneStitchAnimEvaluationTask> EvalTask = MakeShared<FMovieSceneStitchAnimEvaluationTask>();
 	EvalTask->StitchData = StitchData;
 
+	FMovieSceneRootMotionSettings RootMotionSettings;
+	RootMotionSettings.RootMotionSpace = EMovieSceneRootMotionSpace::WorldSpace;
+
 	OutImportedEntity->AddBuilder(
 		FEntityBuilder()
 		.Add(PoseSearchTrackComponents->StitchAnim, StitchData)
@@ -275,7 +281,7 @@ void UMovieSceneStitchAnimSection::ImportEntityImpl(UMovieSceneEntitySystemLinke
 		.Add(AnimMixerComponents->Priority, MixedAnimationPriority)
 		.Add(AnimMixerComponents->Target, MixedAnimationTarget)
 		.Add(AnimMixerComponents->Task, EvalTask)
-		.Add(AnimMixerComponents->RootMotionSettings, FMovieSceneRootMotionSettings{TOptional<FSkeletalAnimationRootMotionOverride>(), EMovieSceneRootMotionSpace::WorldSpace})
+		.Add(AnimMixerComponents->RootMotionSettings, RootMotionSettings)
 		.AddTag(AnimMixerComponents->Tags.RequiresBlending)
 	);
 }
diff --git a/Engine/Plugins/Experimental/MovieScenePoseSearchTracks/Source/MovieScenePoseSearchTracks/Private/Systems/MovieSceneStitchAnimSystem.cpp b/Engine/Plugins/Experimental/MovieScenePoseSearchTracks/Source/MovieScenePoseSearchTracks/Private/Systems/MovieSceneStitchAnimSystem.cpp
--- a/Engine/Plugins/Experimental/MovieScenePoseSearchTracks/Source/MovieScenePoseSearchTracks/Private/Systems/MovieSceneStitchAnimSystem.cpp
+++ b/Engine/Plugins/Experimental/MovieScenePoseSearchTracks/Source/MovieScenePoseSearchTracks/Private/Systems/MovieSceneStitchAnimSystem.cpp
@@ -37,55 +37,12 @@ namespace UE::MovieScene
 	struct FUpdateStitchAnimTask
 	{
 		const UMovieSceneEntitySystemLinker* EntityLinker;
+		const UMovieSceneTransformOriginSystem* TransformOriginSystem;
 
 		FUpdateStitchAnimTask(const UMovieSceneEntitySystemLinker* InLinker)
 			: EntityLinker(InLinker)
 		{
-		}
-
-		FTransform GetTransformOrigin(FInstanceHandle InstanceHandle) const
-		{
-			FTransform TransformOrigin = FTransform::Identity;
-
-			if (!EntityLinker)
-			{
-				return TransformOrigin;
-			}
-
-			const FSequenceInstance& Instance = EntityLinker->GetInstanceRegistry()->GetInstance(InstanceHandle);
-			TSharedRef<FSharedPlaybackState> SharedPlaybackState = Instance.GetSharedPlaybackState();
-			const IMovieScenePlaybackClient* Client = SharedPlaybackState->FindCapability<IMovieScenePlaybackClient>();
-			FMovieSceneSequenceID SequenceID = Instance.GetSequenceID();
-			const UObject* InstanceData = Client ? Client->GetInstanceData() : nullptr;
-			const IMovieSceneTransformOrigin* RawInterface = Cast<const IMovieSceneTransformOrigin>(InstanceData);
-
-			const bool bHasInterface = RawInterface || (InstanceData && InstanceData->GetClass()->ImplementsInterface(UMovieSceneTransformOrigin::StaticClass()));
-			if (bHasInterface)
-			{
-				// Retrieve the current origin
-				TransformOrigin = RawInterface ? RawInterface->GetTransformOrigin() : IMovieSceneTransformOrigin::Execute_BP_GetTransformOrigin(InstanceData);
-			}
-
-			const UMovieSceneTransformOriginSystem* TransformOriginSystem = EntityLinker->FindSystem<UMovieSceneTransformOriginSystem>();
-
-			if (!TransformOriginSystem)
-			{
-				return TransformOrigin;
-			}
-
-			const TSparseArray<FTransform>& TransformOrigins = TransformOriginSystem->GetTransformOriginsByInstanceID();
-			const TMap<FMovieSceneSequenceID, UE::MovieScene::FInstanceHandle> SequenceIDToInstanceHandle = TransformOriginSystem->GetSequenceIDToInstanceHandle();
-
-			if (SequenceIDToInstanceHandle.Contains(SequenceID))
-			{
-				const UE::MovieScene::FInstanceHandle CurrentHandle = SequenceIDToInstanceHandle[SequenceID];
-				if (TransformOrigins.IsValidIndex(CurrentHandle.InstanceID))
-				{
-					TransformOrigin = TransformOrigins[CurrentHandle.InstanceID];
-				}
-			}
-
-			return TransformOrigin;
+			TransformOriginSystem = EntityLinker->FindSystem<UMovieSceneTransformOriginSystem>();
 		}
 
 		// Chooses a context object for motion matching calls based on the anim target for the track
@@ -178,9 +135,10 @@ namespace UE::MovieScene
 					// We use the previous time for the initial time, as it will match the transform from the actor which hasn't
 					// yet been set for this frame.
 					AnimTask->InitialTime = PreviousTimeSeconds;
-					AnimTask->SequenceTransformOrigin = StitchAnim.TargetTransformSpace == EMovieSceneRootMotionSpace::AnimationSpace
-						? GetTransformOrigin(InstanceHandle)
-						: FTransform::Identity;
+					if (StitchAnim.TargetTransformSpace == EMovieSceneRootMotionSpace::AnimationSpace && TransformOriginSystem)
+					{
+						TransformOriginSystem->GetTransformOrigin(InstanceHandle, AnimTask->SequenceTransformOrigin);
+					}
 
 					const AActor* Actor = Cast<AActor>(BoundObject);
 					if (!Actor)
@@ -234,6 +192,7 @@ UMovieSceneStitchAnimSystem::UMovieSceneStitchAnimSystem(const FObjectInitialize
 		DefineImplicitPrerequisite(UMovieSceneEvalTimeSystem::StaticClass(), GetClass());
 		DefineImplicitPrerequisite(UMovieSceneGenericBoundObjectInstantiator::StaticClass(), GetClass());
 		DefineImplicitPrerequisite(UMovieSceneBoundSceneComponentInstantiator::StaticClass(), GetClass());
+		DefineImplicitPrerequisite(UMovieSceneTransformOriginSystem::StaticClass(), GetClass());
 		DefineComponentConsumer(GetClass(), PoseSearchTrackComponents->StitchAnim);
 		DefineImplicitPrerequisite(GetClass(), UMovieSceneRestorePreAnimatedStateSystem::StaticClass());
 		DefineComponentProducer(GetClass(), AnimMixerComponents->Task);
diff --git a/Engine/Source/Editor/MovieSceneTools/Private/TrackEditors/CommonAnimationTrackEditor.cpp b/Engine/Source/Editor/MovieSceneTools/Private/TrackEditors/CommonAnimationTrackEditor.cpp
--- a/Engine/Source/Editor/MovieSceneTools/Private/TrackEditors/CommonAnimationTrackEditor.cpp
+++ b/Engine/Source/Editor/MovieSceneTools/Private/TrackEditors/CommonAnimationTrackEditor.cpp
@@ -887,6 +887,339 @@ void FCommonAnimationSection::CustomizePropertiesDetailsView(TSharedRef<IDetails
 		FOnGetPropertyTypeCustomizationInstance::CreateLambda([=]() { return MakeShared<FMovieSceneSkeletalAnimationParamsDetailCustomization>(InParams); }));
 }
 
+void FCommonAnimationSection::BuildSectionContextMenu(FMenuBuilder& MenuBuilder, const FGuid& ObjectBinding)
+{
+	const TSharedPtr<ISequencer> SequencerPtr = Sequencer.Pin();
+	if (!SequencerPtr.IsValid())
+	{
+		return;
+	}
+
+	UMovieSceneCommonAnimationTrack* const Track = Section.GetTypedOuter<UMovieSceneCommonAnimationTrack>();
+	if (!IsValid(Track))
+	{
+		return;
+	}
+
+	// Can't pick the object that this track binds
+	USkeleton* const Skeleton = UE::Sequencer::FCommonAnimationTrackEditor::AcquireSkeletonFromObjectGuid(ObjectBinding, SequencerPtr);
+	if (!IsValid(Skeleton))
+	{
+		return;
+	}
+
+	const int32 NumBones = Skeleton->GetReferenceSkeleton().GetNum();
+	TArray<FName> BoneNames;
+	for (int32 BoneIndex = 0; BoneIndex < NumBones; ++BoneIndex)
+	{
+		BoneNames.Add(Skeleton->GetReferenceSkeleton().GetBoneName(BoneIndex));
+	}
+
+	auto MatchToBone = [this, ObjectBinding, Skeleton](bool bMatchPrevious, int32 Index)
+	{
+		return FUIAction(
+			FExecuteAction::CreateLambda([this, ObjectBinding, Skeleton, bMatchPrevious, Index]
+				{
+					const TSharedPtr<ISequencer> SequencerPtr = Sequencer.Pin();
+					if (!SequencerPtr.IsValid())
+					{
+						return;
+					}
+
+					UMovieSceneCommonAnimationTrack* const Track = Section.GetTypedOuter<UMovieSceneCommonAnimationTrack>();
+					if (!IsValid(Track))
+					{
+						return;
+					}
+
+					FScopedTransaction MatchSection(LOCTEXT("MatchSectionByBone_Transaction", "Match Section By Bone"));
+					Section.Modify();	
+					Section.bMatchWithPrevious = bMatchPrevious;
+					if (Index >= 0)
+					{
+						USkeletalMeshComponent* const SkelMeshComp = UE::Sequencer::FCommonAnimationTrackEditor::AcquireSkeletalMeshFromObjectGuid(ObjectBinding, SequencerPtr);
+
+						FName Name = Skeleton->GetReferenceSkeleton().GetBoneName(Index);
+						Section.MatchSectionByBoneTransform(SkelMeshComp, SequencerPtr->GetLocalTime().Time, SequencerPtr->GetLocalTime().Rate, Name);
+					}
+					else
+					{
+						Section.ClearMatchedOffsetTransforms();
+					}
+					SequencerPtr->NotifyMovieSceneDataChanged(EMovieSceneDataChangeType::TrackValueChanged);
+
+				}
+			),
+			FCanExecuteAction::CreateLambda([this]() -> bool 
+				{ 
+					return Sequencer.IsValid(); 
+				}),
+			FIsActionChecked::CreateLambda([this, ObjectBinding, Index]() -> bool
+				{
+					const TSharedPtr<ISequencer> SequencerPtr = Sequencer.Pin();
+					if (!SequencerPtr.IsValid())
+					{
+						return false;
+					}
+
+					USkeleton* const Skeleton = UE::Sequencer::FCommonAnimationTrackEditor::AcquireSkeletonFromObjectGuid(ObjectBinding, SequencerPtr);
+					if (!IsValid(Skeleton))
+					{
+						return false;
+					}
+
+					if (Index >= 0)
+					{
+						FName Name = Skeleton->GetReferenceSkeleton().GetBoneName(Index);
+						return (Section.MatchedBoneName == Name);
+					}
+
+					return (Section.MatchedBoneName == NAME_None);
+				})
+			);
+	};
+
+	MenuBuilder.BeginSection(NAME_None, LOCTEXT("MotionBlendingOptions", "Motion Blending Options"));
+	{
+		MenuBuilder.AddSubMenu(
+			LOCTEXT("MatchWithThisBoneInPreviousClip", "Match With This Bone In Previous Clip"), LOCTEXT("MatchWithThisBoneInPreviousClip_Tooltip", "Match This Bone With Previous Clip At Current Frame"),
+			FNewMenuDelegate::CreateLambda([MatchToBone, BoneNames](FMenuBuilder& SubMenuBuilder)
+				{
+					int32 Index = -1;
+					FText NoNameText = LOCTEXT("TurnOffBoneMatching", "Turn Off Matching");
+					FText NoNameTooltipText = LOCTEXT("TurnOffMatchingTooltip", "Turn Off Any Bone Matching");
+					SubMenuBuilder.AddMenuEntry(
+						NoNameText, NoNameTooltipText,
+						FSlateIcon(), MatchToBone(true, Index++), NAME_None, EUserInterfaceActionType::RadioButton);
+
+					for (const FName& BoneName : BoneNames)
+					{
+						FText Name = FText::FromName(BoneName);
+						FText Text = FText::Format(LOCTEXT("BoneNameSelect", "{0}"), Name);
+						FText TooltipText = FText::Format(LOCTEXT("BoneNameSelectTooltip", "Match To This Bone {0}"), Name);
+						SubMenuBuilder.AddMenuEntry(
+							Text, TooltipText,
+							FSlateIcon(), MatchToBone(true, Index++), NAME_None, EUserInterfaceActionType::RadioButton);
+					}
+				}));
+
+		MenuBuilder.AddSubMenu(
+			LOCTEXT("MatchWithThisBoneInNextClip", "Match With This Bone In Next Clip"), LOCTEXT("MatchWithThisBoneInNextClip_Tooltip", "Match This Bone With Next Clip At Current Frame"),
+			FNewMenuDelegate::CreateLambda([MatchToBone, BoneNames](FMenuBuilder& SubMenuBuilder)
+				{
+					int32 Index = -1;
+					
+					FText NoNameText = LOCTEXT("TurnOffBoneMatching", "Turn Off Matching");
+					FText NoNameTooltipText = LOCTEXT("TurnOffMatchingTooltip", "Turn Off Any Bone Matching");
+					SubMenuBuilder.AddMenuEntry(
+						NoNameText, NoNameTooltipText,
+						FSlateIcon(), MatchToBone(false, Index++), NAME_None, EUserInterfaceActionType::RadioButton);
+
+					for (const FName& BoneName : BoneNames)
+					{
+						FText Name = FText::FromName(BoneName);
+						FText Text = FText::Format(LOCTEXT("BoneNameSelect", "{0}"), Name);
+						FText TooltipText = FText::Format(LOCTEXT("BoneNameSelectTooltip", "Match To This Bone {0}"), Name);
+						SubMenuBuilder.AddMenuEntry(
+							Text, TooltipText,
+							FSlateIcon(), MatchToBone(false, Index++), NAME_None, EUserInterfaceActionType::RadioButton);
+					}
+				}));
+
+		MenuBuilder.AddMenuEntry(
+			LOCTEXT("MatchTranslation", "Match X and Y Translation"),
+			LOCTEXT("MatchTranslationTooltip", "Match the Translation to the Specified Bone"),
+			FSlateIcon(),
+			FUIAction(
+				FExecuteAction::CreateLambda([this, ObjectBinding]()
+					{
+						const TSharedPtr<ISequencer> SequencerPtr = Sequencer.Pin();
+						if (!SequencerPtr.IsValid())
+						{
+							return;
+						}
+
+						USkeletalMeshComponent* const SkelMeshComp = UE::Sequencer::FCommonAnimationTrackEditor::AcquireSkeletalMeshFromObjectGuid(ObjectBinding, SequencerPtr);
+
+						FScopedTransaction MatchTransaction(LOCTEXT("MatchTranslation_Transaction", "Match Translation"));
+						Section.Modify();	
+						Section.ToggleMatchTranslation();
+						Section.MatchSectionByBoneTransform(SkelMeshComp, SequencerPtr->GetLocalTime().Time, SequencerPtr->GetLocalTime().Rate, Section.MatchedBoneName);
+						SequencerPtr->NotifyMovieSceneDataChanged(EMovieSceneDataChangeType::TrackValueChanged);
+					}),
+				FCanExecuteAction::CreateLambda([]() -> bool
+					{
+						return true;
+					}),
+				FIsActionChecked::CreateLambda([this]() -> bool
+					{
+						return Section.bMatchTranslation;
+					})),
+			NAME_None,
+			EUserInterfaceActionType::ToggleButton);
+
+		MenuBuilder.AddMenuEntry(
+			LOCTEXT("MatchZHeight", "Match Z Height"),
+			LOCTEXT("MatchZHeightTooltip", "Match the Z Height, may want this off for better matching"),
+			FSlateIcon(),
+			FUIAction(
+				FExecuteAction::CreateLambda([this, ObjectBinding]()
+					{
+						const TSharedPtr<ISequencer> SequencerPtr = Sequencer.Pin();
+						if (!SequencerPtr.IsValid())
+						{
+							return;
+						}
+
+						USkeletalMeshComponent* const SkelMeshComp = UE::Sequencer::FCommonAnimationTrackEditor::AcquireSkeletalMeshFromObjectGuid(ObjectBinding, SequencerPtr);
+
+						FScopedTransaction MatchTransaction(LOCTEXT("MatchZHeight_Transaction", "Match Z Height"));
+						Section.Modify();
+						Section.ToggleMatchIncludeZHeight(); 
+						Section.MatchSectionByBoneTransform(SkelMeshComp, SequencerPtr->GetLocalTime().Time, SequencerPtr->GetLocalTime().Rate, Section.MatchedBoneName);
+						SequencerPtr->NotifyMovieSceneDataChanged(EMovieSceneDataChangeType::TrackValueChanged);
+					}),
+					FCanExecuteAction::CreateLambda([]() -> bool
+						{
+							return true;
+						}),
+					FIsActionChecked::CreateLambda([this]() -> bool
+						{
+							return Section.bMatchIncludeZHeight;
+						})),
+			NAME_None,
+			EUserInterfaceActionType::ToggleButton
+		);
+
+		MenuBuilder.AddMenuEntry(
+			LOCTEXT("MatchYawRotation", "Match Yaw Rotation"),
+			LOCTEXT("MatchYawRotationTooltip", "Match the Yaw Rotation, may want this off for better matching"),
+			FSlateIcon(),
+			FUIAction(
+				FExecuteAction::CreateLambda([this, ObjectBinding]()
+					{
+						const TSharedPtr<ISequencer> SequencerPtr = Sequencer.Pin();
+						if (!SequencerPtr.IsValid())
+						{
+							return;
+						}
+
+						USkeletalMeshComponent* const SkelMeshComp = UE::Sequencer::FCommonAnimationTrackEditor::AcquireSkeletalMeshFromObjectGuid(ObjectBinding, SequencerPtr);
+
+						FScopedTransaction MatchTransaction(LOCTEXT("MatchYawRotation_Transaction", "Match Yaw Rotation"));
+						Section.Modify();
+						Section.ToggleMatchIncludeYawRotation();
+						Section.MatchSectionByBoneTransform(SkelMeshComp, SequencerPtr->GetLocalTime().Time, SequencerPtr->GetLocalTime().Rate, Section.MatchedBoneName);
+						SequencerPtr->NotifyMovieSceneDataChanged(EMovieSceneDataChangeType::TrackValueChanged); }),
+					FCanExecuteAction::CreateLambda([]() -> bool
+						{
+							return true;
+						}),
+					FIsActionChecked::CreateLambda([this]() -> bool
+						{
+							return Section.bMatchRotationYaw;
+						})),
+				NAME_None,
+				EUserInterfaceActionType::ToggleButton);
+
+		MenuBuilder.AddMenuEntry(
+			LOCTEXT("MatchPitchRotation", "Match Pitch Rotation"),
+			LOCTEXT("MatchPitchRotationTooltip", "Match the Pitch Rotation, may want this off for better matching"),
+			FSlateIcon(),
+			FUIAction(
+				FExecuteAction::CreateLambda([this, ObjectBinding]()
+					{
+						const TSharedPtr<ISequencer> SequencerPtr = Sequencer.Pin();
+						if (!SequencerPtr.IsValid())
+						{
+							return;
+						}
+
+						USkeletalMeshComponent* const SkelMeshComp = UE::Sequencer::FCommonAnimationTrackEditor::AcquireSkeletalMeshFromObjectGuid(ObjectBinding, SequencerPtr);
+
+						FScopedTransaction MatchTransaction(LOCTEXT("MatchPitchRotation_Transaction", "Match Pitch Rotation"));
+						Section.Modify();
+						Section.ToggleMatchIncludePitchRotation();
+						Section.MatchSectionByBoneTransform(SkelMeshComp, SequencerPtr->GetLocalTime().Time, SequencerPtr->GetLocalTime().Rate, Section.MatchedBoneName);
+						SequencerPtr->NotifyMovieSceneDataChanged(EMovieSceneDataChangeType::TrackValueChanged);
+					}),
+					FCanExecuteAction::CreateLambda([]() -> bool
+						{
+							return true;
+						}),
+					FIsActionChecked::CreateLambda([this]() -> bool
+						{
+							return Section.bMatchRotationPitch;
+						})),
+			NAME_None,
+			EUserInterfaceActionType::ToggleButton);
+
+		MenuBuilder.AddMenuEntry(
+			LOCTEXT("MatchRollRotation", "Match Roll Rotation"),
+			LOCTEXT("MatchRollRotationTooltip", "Match the Roll Rotation, may want this off for better matching"),
+			FSlateIcon(),
+			FUIAction(
+				FExecuteAction::CreateLambda([this, ObjectBinding]()
+					{
+						const TSharedPtr<ISequencer> SequencerPtr = Sequencer.Pin();
+						if (!SequencerPtr.IsValid())
+						{
+							return;
+						}
+
+						USkeletalMeshComponent* const SkelMeshComp = UE::Sequencer::FCommonAnimationTrackEditor::AcquireSkeletalMeshFromObjectGuid(ObjectBinding, SequencerPtr);
+
+						FScopedTransaction MatchTransaction(LOCTEXT("MatchRollRotation_Transaction", "Match Roll Rotation"));
+						Section.Modify();
+						Section.ToggleMatchIncludeRollRotation();
+						Section.MatchSectionByBoneTransform(SkelMeshComp, SequencerPtr->GetLocalTime().Time, SequencerPtr->GetLocalTime().Rate, Section.MatchedBoneName);
+						SequencerPtr->NotifyMovieSceneDataChanged(EMovieSceneDataChangeType::TrackValueChanged);
+					}),
+				FCanExecuteAction::CreateLambda([]() -> bool
+					{
+						return true;
+					}),
+				FIsActionChecked::CreateLambda([this]() -> bool
+					{
+						return Section.bMatchRotationRoll;
+					})),
+			NAME_None,
+			EUserInterfaceActionType::ToggleButton);
+	}
+	MenuBuilder.EndSection();
+
+	MenuBuilder.BeginSection(NAME_None, LOCTEXT("SkelAnimSectionDisplay", "Display"));
+	{
+		MenuBuilder.AddMenuEntry(
+			NSLOCTEXT("Sequencer", "ShowSkeletons", "Show Skeleton"),
+			NSLOCTEXT("Sequencer", "ShowSkeletonsTooltip", "Show A Skeleton for this Section."),
+			FSlateIcon(),
+			FUIAction(
+				FExecuteAction::CreateLambda([this]()
+					{
+						const TSharedPtr<ISequencer> SequencerPtr = Sequencer.Pin();
+						if (!SequencerPtr.IsValid())
+						{
+							return;
+						}
+
+						Section.ToggleShowSkeleton();
+						SequencerPtr->NotifyMovieSceneDataChanged(EMovieSceneDataChangeType::TrackValueChanged);
+					}),
+				FCanExecuteAction::CreateLambda([this]() -> bool
+					{
+						return Sequencer.IsValid();
+					}),
+				FIsActionChecked::CreateLambda([this]() -> bool
+					{
+						return Section.bShowSkeleton;
+					})),
+				NAME_None,
+				EUserInterfaceActionType::ToggleButton);
+	}
+	MenuBuilder.EndSection();
+}
+
 USkeletalMeshComponent* FCommonAnimationTrackEditor::AcquireSkeletalMeshFromObjectGuid(const FGuid& Guid, TSharedPtr<ISequencer> SequencerPtr)
 {
 	UObject* BoundObject = SequencerPtr.IsValid() ? SequencerPtr->FindSpawnedObjectOrTemplate(Guid) : nullptr;
diff --git a/Engine/Source/Editor/MovieSceneTools/Private/TrackEditors/SkeletalAnimationTrackEditor.cpp b/Engine/Source/Editor/MovieSceneTools/Private/TrackEditors/SkeletalAnimationTrackEditor.cpp
--- a/Engine/Source/Editor/MovieSceneTools/Private/TrackEditors/SkeletalAnimationTrackEditor.cpp
+++ b/Engine/Source/Editor/MovieSceneTools/Private/TrackEditors/SkeletalAnimationTrackEditor.cpp
@@ -229,339 +229,6 @@ FSkeletalAnimationSection::FSkeletalAnimationSection( UMovieSceneSection& InSect
 {
 }
 
-void FSkeletalAnimationSection::BuildSectionContextMenu(FMenuBuilder& MenuBuilder, const FGuid& ObjectBinding)
-{
-	const TSharedPtr<ISequencer> SequencerPtr = Sequencer.Pin();
-	if (!SequencerPtr.IsValid())
-	{
-		return;
-	}
-
-	UMovieSceneSkeletalAnimationTrack* const Track = Section.GetTypedOuter<UMovieSceneSkeletalAnimationTrack>();
-	if (!IsValid(Track))
-	{
-		return;
-	}
-
-	// Can't pick the object that this track binds
-	USkeleton* const Skeleton = UE::Sequencer::FCommonAnimationTrackEditor::AcquireSkeletonFromObjectGuid(ObjectBinding, SequencerPtr);
-	if (!IsValid(Skeleton))
-	{
-		return;
-	}
-
-	const int32 NumBones = Skeleton->GetReferenceSkeleton().GetNum();
-	TArray<FName> BoneNames;
-	for (int32 BoneIndex = 0; BoneIndex < NumBones; ++BoneIndex)
-	{
-		BoneNames.Add(Skeleton->GetReferenceSkeleton().GetBoneName(BoneIndex));
-	}
-
-	auto MatchToBone = [this, ObjectBinding, Skeleton](bool bMatchPrevious, int32 Index)
-	{
-		return FUIAction(
-			FExecuteAction::CreateLambda([this, ObjectBinding, Skeleton, bMatchPrevious, Index]
-				{
-					const TSharedPtr<ISequencer> SequencerPtr = Sequencer.Pin();
-					if (!SequencerPtr.IsValid())
-					{
-						return;
-					}
-
-					UMovieSceneSkeletalAnimationTrack* const Track = Section.GetTypedOuter<UMovieSceneSkeletalAnimationTrack>();
-					if (!IsValid(Track))
-					{
-						return;
-					}
-
-					FScopedTransaction MatchSection(LOCTEXT("MatchSectionByBone_Transaction", "Match Section By Bone"));
-					Section.Modify();	
-					Section.bMatchWithPrevious = bMatchPrevious;
-					if (Index >= 0)
-					{
-						USkeletalMeshComponent* const SkelMeshComp = UE::Sequencer::FCommonAnimationTrackEditor::AcquireSkeletalMeshFromObjectGuid(ObjectBinding, SequencerPtr);
-
-						FName Name = Skeleton->GetReferenceSkeleton().GetBoneName(Index);
-						Section.MatchSectionByBoneTransform(SkelMeshComp, SequencerPtr->GetLocalTime().Time, SequencerPtr->GetLocalTime().Rate, Name);
-					}
-					else
-					{
-						Section.ClearMatchedOffsetTransforms();
-					}
-					SequencerPtr->NotifyMovieSceneDataChanged(EMovieSceneDataChangeType::TrackValueChanged);
-
-				}
-			),
-			FCanExecuteAction::CreateLambda([this]() -> bool 
-				{ 
-					return Sequencer.IsValid(); 
-				}),
-			FIsActionChecked::CreateLambda([this, ObjectBinding, Index]() -> bool
-				{
-					const TSharedPtr<ISequencer> SequencerPtr = Sequencer.Pin();
-					if (!SequencerPtr.IsValid())
-					{
-						return false;
-					}
-
-					USkeleton* const Skeleton = UE::Sequencer::FCommonAnimationTrackEditor::AcquireSkeletonFromObjectGuid(ObjectBinding, SequencerPtr);
-					if (!IsValid(Skeleton))
-					{
-						return false;
-					}
-
-					if (Index >= 0)
-					{
-						FName Name = Skeleton->GetReferenceSkeleton().GetBoneName(Index);
-						return (Section.MatchedBoneName == Name);
-					}
-
-					return (Section.MatchedBoneName == NAME_None);
-				})
-			);
-	};
-
-	MenuBuilder.BeginSection(NAME_None, LOCTEXT("MotionBlendingOptions", "Motion Blending Options"));
-	{
-		MenuBuilder.AddSubMenu(
-			LOCTEXT("MatchWithThisBoneInPreviousClip", "Match With This Bone In Previous Clip"), LOCTEXT("MatchWithThisBoneInPreviousClip_Tooltip", "Match This Bone With Previous Clip At Current Frame"),
-			FNewMenuDelegate::CreateLambda([MatchToBone, BoneNames](FMenuBuilder& SubMenuBuilder)
-				{
-					int32 Index = -1;
-					FText NoNameText = LOCTEXT("TurnOffBoneMatching", "Turn Off Matching");
-					FText NoNameTooltipText = LOCTEXT("TurnOffMatchingTooltip", "Turn Off Any Bone Matching");
-					SubMenuBuilder.AddMenuEntry(
-						NoNameText, NoNameTooltipText,
-						FSlateIcon(), MatchToBone(true, Index++), NAME_None, EUserInterfaceActionType::RadioButton);
-
-					for (const FName& BoneName : BoneNames)
-					{
-						FText Name = FText::FromName(BoneName);
-						FText Text = FText::Format(LOCTEXT("BoneNameSelect", "{0}"), Name);
-						FText TooltipText = FText::Format(LOCTEXT("BoneNameSelectTooltip", "Match To This Bone {0}"), Name);
-						SubMenuBuilder.AddMenuEntry(
-							Text, TooltipText,
-							FSlateIcon(), MatchToBone(true, Index++), NAME_None, EUserInterfaceActionType::RadioButton);
-					}
-				}));
-
-		MenuBuilder.AddSubMenu(
-			LOCTEXT("MatchWithThisBoneInNextClip", "Match With This Bone In Next Clip"), LOCTEXT("MatchWithThisBoneInNextClip_Tooltip", "Match This Bone With Next Clip At Current Frame"),
-			FNewMenuDelegate::CreateLambda([MatchToBone, BoneNames](FMenuBuilder& SubMenuBuilder)
-				{
-					int32 Index = -1;
-					
-					FText NoNameText = LOCTEXT("TurnOffBoneMatching", "Turn Off Matching");
-					FText NoNameTooltipText = LOCTEXT("TurnOffMatchingTooltip", "Turn Off Any Bone Matching");
-					SubMenuBuilder.AddMenuEntry(
-						NoNameText, NoNameTooltipText,
-						FSlateIcon(), MatchToBone(false, Index++), NAME_None, EUserInterfaceActionType::RadioButton);
-
-					for (const FName& BoneName : BoneNames)
-					{
-						FText Name = FText::FromName(BoneName);
-						FText Text = FText::Format(LOCTEXT("BoneNameSelect", "{0}"), Name);
-						FText TooltipText = FText::Format(LOCTEXT("BoneNameSelectTooltip", "Match To This Bone {0}"), Name);
-						SubMenuBuilder.AddMenuEntry(
-							Text, TooltipText,
-							FSlateIcon(), MatchToBone(false, Index++), NAME_None, EUserInterfaceActionType::RadioButton);
-					}
-				}));
-
-		MenuBuilder.AddMenuEntry(
-			LOCTEXT("MatchTranslation", "Match X and Y Translation"),
-			LOCTEXT("MatchTranslationTooltip", "Match the Translation to the Specified Bone"),
-			FSlateIcon(),
-			FUIAction(
-				FExecuteAction::CreateLambda([this, ObjectBinding]()
-					{
-						const TSharedPtr<ISequencer> SequencerPtr = Sequencer.Pin();
-						if (!SequencerPtr.IsValid())
-						{
-							return;
-						}
-
-						USkeletalMeshComponent* const SkelMeshComp = UE::Sequencer::FCommonAnimationTrackEditor::AcquireSkeletalMeshFromObjectGuid(ObjectBinding, SequencerPtr);
-
-						FScopedTransaction MatchTransaction(LOCTEXT("MatchTranslation_Transaction", "Match Translation"));
-						Section.Modify();	
-						Section.ToggleMatchTranslation();
-						Section.MatchSectionByBoneTransform(SkelMeshComp, SequencerPtr->GetLocalTime().Time, SequencerPtr->GetLocalTime().Rate, Section.MatchedBoneName);
-						SequencerPtr->NotifyMovieSceneDataChanged(EMovieSceneDataChangeType::TrackValueChanged);
-					}),
-				FCanExecuteAction::CreateLambda([]() -> bool
-					{
-						return true;
-					}),
-				FIsActionChecked::CreateLambda([this]() -> bool
-					{
-						return Section.bMatchTranslation;
-					})),
-			NAME_None,
-			EUserInterfaceActionType::ToggleButton);
-
-		MenuBuilder.AddMenuEntry(
-			LOCTEXT("MatchZHeight", "Match Z Height"),
-			LOCTEXT("MatchZHeightTooltip", "Match the Z Height, may want this off for better matching"),
-			FSlateIcon(),
-			FUIAction(
-				FExecuteAction::CreateLambda([this, ObjectBinding]()
-					{
-						const TSharedPtr<ISequencer> SequencerPtr = Sequencer.Pin();
-						if (!SequencerPtr.IsValid())
-						{
-							return;
-						}
-
-						USkeletalMeshComponent* const SkelMeshComp = UE::Sequencer::FCommonAnimationTrackEditor::AcquireSkeletalMeshFromObjectGuid(ObjectBinding, SequencerPtr);
-
-						FScopedTransaction MatchTransaction(LOCTEXT("MatchZHeight_Transaction", "Match Z Height"));
-						Section.Modify();
-						Section.ToggleMatchIncludeZHeight(); 
-						Section.MatchSectionByBoneTransform(SkelMeshComp, SequencerPtr->GetLocalTime().Time, SequencerPtr->GetLocalTime().Rate, Section.MatchedBoneName);
-						SequencerPtr->NotifyMovieSceneDataChanged(EMovieSceneDataChangeType::TrackValueChanged);
-					}),
-					FCanExecuteAction::CreateLambda([]() -> bool
-						{
-							return true;
-						}),
-					FIsActionChecked::CreateLambda([this]() -> bool
-						{
-							return Section.bMatchIncludeZHeight;
-						})),
-			NAME_None,
-			EUserInterfaceActionType::ToggleButton
-		);
-
-		MenuBuilder.AddMenuEntry(
-			LOCTEXT("MatchYawRotation", "Match Yaw Rotation"),
-			LOCTEXT("MatchYawRotationTooltip", "Match the Yaw Rotation, may want this off for better matching"),
-			FSlateIcon(),
-			FUIAction(
-				FExecuteAction::CreateLambda([this, ObjectBinding]()
-					{
-						const TSharedPtr<ISequencer> SequencerPtr = Sequencer.Pin();
-						if (!SequencerPtr.IsValid())
-						{
-							return;
-						}
-
-						USkeletalMeshComponent* const SkelMeshComp = UE::Sequencer::FCommonAnimationTrackEditor::AcquireSkeletalMeshFromObjectGuid(ObjectBinding, SequencerPtr);
-
-						FScopedTransaction MatchTransaction(LOCTEXT("MatchYawRotation_Transaction", "Match Yaw Rotation"));
-						Section.Modify();
-						Section.ToggleMatchIncludeYawRotation();
-						Section.MatchSectionByBoneTransform(SkelMeshComp, SequencerPtr->GetLocalTime().Time, SequencerPtr->GetLocalTime().Rate, Section.MatchedBoneName);
-						SequencerPtr->NotifyMovieSceneDataChanged(EMovieSceneDataChangeType::TrackValueChanged); }),
-					FCanExecuteAction::CreateLambda([]() -> bool
-						{
-							return true;
-						}),
-					FIsActionChecked::CreateLambda([this]() -> bool
-						{
-							return Section.bMatchRotationYaw;
-						})),
-				NAME_None,
-				EUserInterfaceActionType::ToggleButton);
-
-		MenuBuilder.AddMenuEntry(
-			LOCTEXT("MatchPitchRotation", "Match Pitch Rotation"),
-			LOCTEXT("MatchPitchRotationTooltip", "Match the Pitch Rotation, may want this off for better matching"),
-			FSlateIcon(),
-			FUIAction(
-				FExecuteAction::CreateLambda([this, ObjectBinding]()
-					{
-						const TSharedPtr<ISequencer> SequencerPtr = Sequencer.Pin();
-						if (!SequencerPtr.IsValid())
-						{
-							return;
-						}
-
-						USkeletalMeshComponent* const SkelMeshComp = UE::Sequencer::FCommonAnimationTrackEditor::AcquireSkeletalMeshFromObjectGuid(ObjectBinding, SequencerPtr);
-
-						FScopedTransaction MatchTransaction(LOCTEXT("MatchPitchRotation_Transaction", "Match Pitch Rotation"));
-						Section.Modify();
-						Section.ToggleMatchIncludePitchRotation();
-						Section.MatchSectionByBoneTransform(SkelMeshComp, SequencerPtr->GetLocalTime().Time, SequencerPtr->GetLocalTime().Rate, Section.MatchedBoneName);
-						SequencerPtr->NotifyMovieSceneDataChanged(EMovieSceneDataChangeType::TrackValueChanged);
-					}),
-					FCanExecuteAction::CreateLambda([]() -> bool
-						{
-							return true;
-						}),
-					FIsActionChecked::CreateLambda([this]() -> bool
-						{
-							return Section.bMatchRotationPitch;
-						})),
-			NAME_None,
-			EUserInterfaceActionType::ToggleButton);
-
-		MenuBuilder.AddMenuEntry(
-			LOCTEXT("MatchRollRotation", "Match Roll Rotation"),
-			LOCTEXT("MatchRollRotationTooltip", "Match the Roll Rotation, may want this off for better matching"),
-			FSlateIcon(),
-			FUIAction(
-				FExecuteAction::CreateLambda([this, ObjectBinding]()
-					{
-						const TSharedPtr<ISequencer> SequencerPtr = Sequencer.Pin();
-						if (!SequencerPtr.IsValid())
-						{
-							return;
-						}
-
-						USkeletalMeshComponent* const SkelMeshComp = UE::Sequencer::FCommonAnimationTrackEditor::AcquireSkeletalMeshFromObjectGuid(ObjectBinding, SequencerPtr);
-
-						FScopedTransaction MatchTransaction(LOCTEXT("MatchRollRotation_Transaction", "Match Roll Rotation"));
-						Section.Modify();
-						Section.ToggleMatchIncludeRollRotation();
-						Section.MatchSectionByBoneTransform(SkelMeshComp, SequencerPtr->GetLocalTime().Time, SequencerPtr->GetLocalTime().Rate, Section.MatchedBoneName);
-						SequencerPtr->NotifyMovieSceneDataChanged(EMovieSceneDataChangeType::TrackValueChanged);
-					}),
-				FCanExecuteAction::CreateLambda([]() -> bool
-					{
-						return true;
-					}),
-				FIsActionChecked::CreateLambda([this]() -> bool
-					{
-						return Section.bMatchRotationRoll;
-					})),
-			NAME_None,
-			EUserInterfaceActionType::ToggleButton);
-	}
-	MenuBuilder.EndSection();
-
-	MenuBuilder.BeginSection(NAME_None, LOCTEXT("SkelAnimSectionDisplay", "Display"));
-	{
-		MenuBuilder.AddMenuEntry(
-			NSLOCTEXT("Sequencer", "ShowSkeletons", "Show Skeleton"),
-			NSLOCTEXT("Sequencer", "ShowSkeletonsTooltip", "Show A Skeleton for this Section."),
-			FSlateIcon(),
-			FUIAction(
-				FExecuteAction::CreateLambda([this]()
-					{
-						const TSharedPtr<ISequencer> SequencerPtr = Sequencer.Pin();
-						if (!SequencerPtr.IsValid())
-						{
-							return;
-						}
-
-						Section.ToggleShowSkeleton();
-						SequencerPtr->NotifyMovieSceneDataChanged(EMovieSceneDataChangeType::TrackValueChanged);
-					}),
-				FCanExecuteAction::CreateLambda([this]() -> bool
-					{
-						return Sequencer.IsValid();
-					}),
-				FIsActionChecked::CreateLambda([this]() -> bool
-					{
-						return Section.bShowSkeleton;
-					})),
-				NAME_None,
-				EUserInterfaceActionType::ToggleButton);
-	}
-	MenuBuilder.EndSection();
-}
-
 void FSkeletalAnimationTrackEditor::BuildTrackContextMenu(FMenuBuilder& MenuBuilder, UMovieSceneTrack* Track)
 {
 	BuildTrackContextMenu_Internal(MenuBuilder, Track, true);
diff --git a/Engine/Source/Editor/MovieSceneTools/Private/TrackEditors/SkeletalAnimationTrackEditor.h b/Engine/Source/Editor/MovieSceneTools/Private/TrackEditors/SkeletalAnimationTrackEditor.h
--- a/Engine/Source/Editor/MovieSceneTools/Private/TrackEditors/SkeletalAnimationTrackEditor.h
+++ b/Engine/Source/Editor/MovieSceneTools/Private/TrackEditors/SkeletalAnimationTrackEditor.h
@@ -47,6 +47,4 @@ public:
 
 	/** Constructor. */
 	FSkeletalAnimationSection( UMovieSceneSection& InSection, TWeakPtr<ISequencer> InSequencer);
-
-	virtual void BuildSectionContextMenu(FMenuBuilder& MenuBuilder, const FGuid& InObjectBinding) override;
 };
diff --git a/Engine/Source/Editor/MovieSceneTools/Public/TrackEditors/CommonAnimationTrackEditor.h b/Engine/Source/Editor/MovieSceneTools/Public/TrackEditors/CommonAnimationTrackEditor.h
--- a/Engine/Source/Editor/MovieSceneTools/Public/TrackEditors/CommonAnimationTrackEditor.h
+++ b/Engine/Source/Editor/MovieSceneTools/Public/TrackEditors/CommonAnimationTrackEditor.h
@@ -170,6 +170,7 @@ public:
 	virtual void BeginDilateSection() override;
 	virtual void DilateSection(const TRange<FFrameNumber>& NewRange, float DilationFactor) override;
 	virtual bool RequestDeleteKeyArea(const TArray<FName>& KeyAreaNamePath) override;
+	virtual void BuildSectionContextMenu(FMenuBuilder& MenuBuilder, const FGuid& InObjectBinding) override;
 
 
 protected:
diff --git a/Engine/Source/Runtime/MovieSceneTracks/Private/Sections/MovieSceneSkeletalAnimationSection.cpp b/Engine/Source/Runtime/MovieSceneTracks/Private/Sections/MovieSceneSkeletalAnimationSection.cpp
--- a/Engine/Source/Runtime/MovieSceneTracks/Private/Sections/MovieSceneSkeletalAnimationSection.cpp
+++ b/Engine/Source/Runtime/MovieSceneTracks/Private/Sections/MovieSceneSkeletalAnimationSection.cpp
@@ -380,7 +380,7 @@ UMovieSceneSection* UMovieSceneSkeletalAnimationSection::SplitSection(FQualified
 
 	TOptional<FTransform> RootTransform;
 	FName BoneName;
-	UMovieSceneSkeletalAnimationTrack* Track = GetTypedOuter<UMovieSceneSkeletalAnimationTrack>();
+	UMovieSceneCommonAnimationTrack* Track = GetTypedOuter<UMovieSceneCommonAnimationTrack>();
 	//if we are doing root motion then get the BoneName 
 	if (Track)
 	{
@@ -620,7 +620,7 @@ void UMovieSceneSkeletalAnimationSection::GetRootMotion(FFrameTime CurrentTime,
 {
 	if (GetRootMotionParams())
 	{
-		if (UMovieSceneSkeletalAnimationTrack* Track = GetTypedOuter<UMovieSceneSkeletalAnimationTrack>())
+		if (UMovieSceneCommonAnimationTrack* Track = GetTypedOuter<UMovieSceneCommonAnimationTrack>())
 		{
 			OutRootMotionParams.Transform = Track->GetRootMotion(CurrentTime);
 			OutRootMotionParams.ChildBoneIndex = TempRootBoneIndex.IsSet() ? TempRootBoneIndex.GetValue() : INDEX_NONE;
@@ -651,7 +651,7 @@ bool UMovieSceneSkeletalAnimationSection::GetRootMotionVelocity(FFrameTime Previ
 
 FMovieSceneSkeletalAnimRootMotionTrackParams*  UMovieSceneSkeletalAnimationSection::GetRootMotionParams() const
 {
-	UMovieSceneSkeletalAnimationTrack* Track = GetTypedOuter<UMovieSceneSkeletalAnimationTrack>();
+	UMovieSceneCommonAnimationTrack* Track = GetTypedOuter<UMovieSceneCommonAnimationTrack>();
 	if (Track)
 	{
 		return &Track->RootMotionParams;
@@ -825,7 +825,7 @@ bool UMovieSceneSkeletalAnimationSection::GetRootMotionTransform(FAnimationPoseD
 
 void UMovieSceneSkeletalAnimationSection::MultiplyOutInverseOnNextClips(FVector PreviousMatchedLocationOffset, FRotator PreviousMatchedRotationOffset)
 {
-	UMovieSceneSkeletalAnimationTrack* Track = GetTypedOuter<UMovieSceneSkeletalAnimationTrack>();
+	UMovieSceneCommonAnimationTrack* Track = GetTypedOuter<UMovieSceneCommonAnimationTrack>();
 	if (Track)
 	{
 		bool bMultiplyOutInverse = false;
@@ -879,7 +879,7 @@ void UMovieSceneSkeletalAnimationSection::MatchSectionByBoneTransform(USkeletalM
 	const FName& BoneName)
 {
 	MatchedBoneName = BoneName;
-	UMovieSceneSkeletalAnimationTrack* Track = GetTypedOuter<UMovieSceneSkeletalAnimationTrack>();
+	UMovieSceneCommonAnimationTrack* Track = GetTypedOuter<UMovieSceneCommonAnimationTrack>();
 	if (Track)
 	{
 		FTransform DiffTransform;
diff --git a/Engine/Source/Runtime/MovieSceneTracks/Private/Systems/MovieSceneTransformOriginSystem.cpp b/Engine/Source/Runtime/MovieSceneTracks/Private/Systems/MovieSceneTransformOriginSystem.cpp
--- a/Engine/Source/Runtime/MovieSceneTracks/Private/Systems/MovieSceneTransformOriginSystem.cpp
+++ b/Engine/Source/Runtime/MovieSceneTracks/Private/Systems/MovieSceneTransformOriginSystem.cpp
@@ -290,6 +290,17 @@ bool UMovieSceneTransformOriginSystem::IsRelevantImpl(UMovieSceneEntitySystemLin
 	return false;
 }
 
+bool UMovieSceneTransformOriginSystem::GetTransformOrigin(UE::MovieScene::FInstanceHandle InstanceHandle, FTransform& OutTransform) const
+{
+	if (TransformOriginsByInstanceID.IsValidIndex(InstanceHandle.InstanceID))
+	{
+		OutTransform = TransformOriginsByInstanceID[InstanceHandle.InstanceID];
+		return true;
+	}
+
+	return false;
+}
+
 void UMovieSceneTransformOriginSystem::OnLink()
 {
 	UMovieSceneTransformOriginInstantiatorSystem* Instantiator = Linker->LinkSystem<UMovieSceneTransformOriginInstantiatorSystem>();
diff --git a/Engine/Source/Runtime/MovieSceneTracks/Private/Tracks/MovieSceneCommonAnimationTrack.cpp b/Engine/Source/Runtime/MovieSceneTracks/Private/Tracks/MovieSceneCommonAnimationTrack.cpp
--- a/Engine/Source/Runtime/MovieSceneTracks/Private/Tracks/MovieSceneCommonAnimationTrack.cpp
+++ b/Engine/Source/Runtime/MovieSceneTracks/Private/Tracks/MovieSceneCommonAnimationTrack.cpp
@@ -1,8 +1,35 @@
 // Copyright Epic Games, Inc. All Rights Reserved.
 
 #include "Tracks/MovieSceneCommonAnimationTrack.h"
+
+#include "Animation/AnimationPoseData.h"
+#include "Animation/AnimInstance.h"
+#include "Animation/AnimSequence.h"
+#include "Animation/AttributesRuntime.h"
+#include "AnimationRuntime.h"
+#include "AnimSequencerInstanceProxy.h"
+#include "BoneContainer.h"
+#include "Components/SkeletalMeshComponent.h"
+#include "MovieScene.h"
+#include "Rendering/SkeletalMeshRenderData.h"
 #include "Sections/MovieSceneSkeletalAnimationSection.h"
+#include "SkeletalDebugRendering.h"
+
+#if WITH_EDITORONLY_DATA
+#include "AnimationBlueprintLibrary.h"
+#include "Misc/QualifiedFrameTime.h"
+#include "Misc/Timecode.h"
+#endif
+
 
+UMovieSceneCommonAnimationTrack::UMovieSceneCommonAnimationTrack(const FObjectInitializer& ObjInit)
+	: Super(ObjInit)
+	, bBlendFirstChildOfRoot(false)
+{
+#if WITH_EDITORONLY_DATA
+	bShowRootMotionTrail = false;
+#endif
+}
 
 UMovieSceneSection* UMovieSceneCommonAnimationTrack::CreateNewSection()
 {
@@ -38,19 +65,1043 @@ bool UMovieSceneCommonAnimationTrack::HasSection(const UMovieSceneSection& Secti
 void UMovieSceneCommonAnimationTrack::AddSection(UMovieSceneSection& Section)
 {
 	AnimationSections.Add(&Section);
+	SetUpRootMotions(true);
 }
 
 void UMovieSceneCommonAnimationTrack::RemoveSection(UMovieSceneSection& Section)
 {
 	AnimationSections.Remove(&Section);
+	SetUpRootMotions(true);
 }
 
 void UMovieSceneCommonAnimationTrack::RemoveSectionAt(int32 SectionIndex)
 {
 	AnimationSections.RemoveAt(SectionIndex);
+	SetUpRootMotions(true);
+}
+
+void UMovieSceneCommonAnimationTrack::UpdateEasing()
+{
+	Super::UpdateEasing();
+	SetRootMotionsDirty();
 }
 
 bool UMovieSceneCommonAnimationTrack::IsEmpty() const
 {
 	return AnimationSections.Num() == 0;
+}
+
+UMovieSceneSection* UMovieSceneCommonAnimationTrack::AddNewAnimationOnRow(FFrameNumber KeyTime, UAnimSequenceBase* AnimSequence, int32 RowIndex)
+{
+	UMovieSceneSkeletalAnimationSection* NewSection = Cast<UMovieSceneSkeletalAnimationSection>(CreateNewSection());
+	{
+		FFrameTime AnimationLength = AnimSequence->GetPlayLength() * GetTypedOuter<UMovieScene>()->GetTickResolution();
+		int32 IFrameNumber = AnimationLength.FrameNumber.Value + (int)(AnimationLength.GetSubFrame() + 0.5f) + 1;
+		NewSection->InitialPlacementOnRow(AnimationSections, KeyTime, IFrameNumber, RowIndex);
+		NewSection->Params.Animation = AnimSequence;
+
+#if WITH_EDITORONLY_DATA
+		FQualifiedFrameTime SourceStartFrameTime;
+		if (UAnimationBlueprintLibrary::EvaluateRootBoneTimecodeAttributesAtTime(NewSection->Params.Animation, 0.0f, SourceStartFrameTime))
+		{
+			NewSection->TimecodeSource.Timecode = SourceStartFrameTime.ToTimecode();
+		}
+#endif
+	}
+
+	AddSection(*NewSection);
+
+	return NewSection;
+}
+
+#if WITH_EDITOR
+void UMovieSceneCommonAnimationTrack::PostEditImport()
+{
+	Super::PostEditImport();
+	RootMotionParams.bRootMotionsDirty = true;
+}
+
+void UMovieSceneCommonAnimationTrack::PostEditUndo()
+{
+	Super::PostEditUndo();
+	RootMotionParams.bRootMotionsDirty = true;
+}
+
+void UMovieSceneCommonAnimationTrack::PostEditChangeProperty(FPropertyChangedEvent& PropertyChangedEvent)
+{
+	Super::PostEditChangeProperty(PropertyChangedEvent);
+	RootMotionParams.bRootMotionsDirty = true;
+}
+EMovieSceneSectionMovedResult UMovieSceneCommonAnimationTrack::OnSectionMoved(UMovieSceneSection& Section, const FMovieSceneSectionMovedParams& Params)
+{
+	if (Params.MoveType == EPropertyChangeType::ValueSet)
+	{
+		RootMotionParams.bRootMotionsDirty = true;
+	}
+	
+	return EMovieSceneSectionMovedResult::None;
+}
+#endif // WITH_EDITOR
+
+void UMovieSceneCommonAnimationTrack::SortSections()
+{
+	AnimationSections.Sort(
+		[](const UMovieSceneSection& A, const UMovieSceneSection& B)
+		{
+			TRange<FFrameNumber> RangeA = A.GetTrueRange();
+			TRange<FFrameNumber> RangeB = B.GetTrueRange();
+			if (RangeA.GetLowerBound().IsOpen())
+			{
+				return true;
+			}
+			else if (RangeB.GetLowerBound().IsOpen())
+			{
+				return false;
+			}
+			return (RangeA.GetLowerBoundValue() < RangeB.GetLowerBoundValue());
+		}
+	);
+}
+
+//expectation is the weights may be unnormalized.
+static void BlendTheseTransformsByWeight(FTransform& OutTransform, const TArray<FTransform>& Transforms,  TArray<float>& Weights)
+{
+	if (Weights.Num() > 0)
+	{
+		float TotalWeight = 0.0f;
+		for (int32 WeightIndex = 0; WeightIndex < Weights.Num(); ++WeightIndex)
+		{
+			TotalWeight += Weights[WeightIndex];
+		}
+		if (!FMath::IsNearlyEqual(TotalWeight, 1.0f))
+		{
+			for (int32 DivideIndex = 0; DivideIndex < Weights.Num(); ++DivideIndex)
+			{
+				Weights[DivideIndex] /= TotalWeight;
+			}
+		}
+	}
+	int32 NumBlends = Transforms.Num();
+	check(Transforms.Num() == Weights.Num());
+	if (NumBlends == 0)
+	{
+		OutTransform = FTransform::Identity;
+	}
+	else if (NumBlends == 1)
+	{
+		OutTransform = Transforms[0];
+	}
+	else
+	{
+		FVector		OutTranslation(0.0f, 0.0f, 0.0f);
+		FVector		OutScale(0.0f, 0.0f, 0.0f);
+
+		//rotation will get set to the first weighted and then made closest to that so linear interp works.
+		FQuat FirstRot = Transforms[0].GetRotation();
+		FQuat		OutRotation(FirstRot.X * Weights[0], FirstRot.Y * Weights[0], FirstRot.Z * Weights[0], FirstRot.W * Weights[0]);
+
+		for (int32 Index = 0; Index < NumBlends; ++Index)
+		{
+			OutTranslation += Transforms[Index].GetTranslation() * Weights[Index];
+			OutScale +=  Transforms[Index].GetScale3D() * Weights[Index];
+			if (Index != 0)
+			{
+				FQuat Quat = Transforms[Index].GetRotation();
+				Quat.EnforceShortestArcWith(FirstRot);
+				Quat *= Weights[Index];
+				OutRotation += Quat;
+			}
+		}
+
+		OutRotation.Normalize();
+		OutTransform = FTransform(OutRotation, OutTranslation, OutScale);
+	}
+}
+
+void UMovieSceneCommonAnimationTrack::SetRootMotionsDirty()
+{
+	RootMotionParams.bRootMotionsDirty = true;
+}
+
+struct FSkelBoneLength
+{
+	FSkelBoneLength(FCompactPoseBoneIndex InPoseIndex, float InBL) :PoseBoneIndex(InPoseIndex), BoneLength(InBL) {};
+	FCompactPoseBoneIndex PoseBoneIndex;
+	float BoneLength; //squared
+};
+
+static void CalculateDistanceMap(USkeletalMeshComponent* SkelMeshComp, UAnimSequenceBase* FirstAnimSeq, UAnimSequenceBase* SecondAnimSeq, float StartFirstAnimTime, float FrameRate,
+	TArray<TArray<float>>& OutDistanceDifferences)
+{
+
+	int32 FirstAnimNumFrames = (FirstAnimSeq->GetPlayLength() - StartFirstAnimTime) * FrameRate + 1;
+	int32 SecondAnimNumFrames = SecondAnimSeq->GetPlayLength() * FrameRate + 1;
+	OutDistanceDifferences.SetNum(FirstAnimNumFrames);
+	float FirstAnimIndex = 0.0f;
+	float FrameRateDiff = 1.0f / FrameRate;
+	FCompactPose FirstAnimPose, SecondAnimPose;
+	FCSPose<FCompactPose> FirstMeshPoses, SecondMeshPoses;
+	FirstAnimPose.ResetToRefPose(SkelMeshComp->GetAnimInstance()->GetRequiredBones());
+	SecondAnimPose.ResetToRefPose(SkelMeshComp->GetAnimInstance()->GetRequiredBones());
+
+	FBlendedCurve FirstOutCurve, SecondOutCurve;
+	UE::Anim::FStackAttributeContainer FirstTempAttributes, SecondTempAttributes;
+	FAnimationPoseData FirstAnimPoseData(FirstAnimPose, FirstOutCurve, FirstTempAttributes);
+	FAnimationPoseData SecondAnimPoseData(SecondAnimPose, SecondOutCurve, SecondTempAttributes);
+
+	//sort by bone lengths just do the first half
+	//this should avoid us overvalueing to many small values.
+	/*
+	TArray<FSkelBoneLength> BoneLengths;
+	BoneLengths.SetNum(FirstAnimPose.GetNumBones());
+	int32 Index = 0;
+	for (FCompactPoseBoneIndex PoseBoneIndex : FirstAnimPose.ForEachBoneIndex())
+	{
+		FTransform LocalTransform = FirstMeshPoses.GetLocalSpaceTransform(PoseBoneIndex);
+		float BoneLengthVal = LocalTransform.GetLocation().SizeSquared();
+		BoneLengths[Index++] = FSkelBoneLength(PoseBoneIndex, BoneLengthVal);
+	}
+	BoneLengths.Sort([](const FSkelBoneLength& Item1, const FSkelBoneLength& Item2) {
+		return Item1.BoneLength > Item2.BoneLength;
+		});
+		*/
+	FBlendedCurve OutCurve;
+	const FBoneContainer& RequiredBones = FirstAnimPoseData.GetPose().GetBoneContainer();
+	for (TArray<float>& FloatArray : OutDistanceDifferences)
+	{
+		FloatArray.SetNum(SecondAnimNumFrames);
+		float FirstAnimTime = FirstAnimIndex * FrameRateDiff + StartFirstAnimTime;
+		FirstAnimIndex += 1.0f;
+		FAnimExtractContext FirstExtractionContext(static_cast<double>(FirstAnimTime), false);
+		FirstAnimSeq->GetAnimationPose(FirstAnimPoseData, FirstExtractionContext);
+		FirstMeshPoses.InitPose(FirstAnimPoseData.GetPose());
+		float SecondAnimIndex = 0.0f;
+		for (float& DistVal : FloatArray)
+		{
+			DistVal = 0.0f;
+			float SecondAnimTime = SecondAnimIndex * FrameRateDiff;
+			SecondAnimIndex += 1.0f;
+			FAnimExtractContext SecondExtractionContext(static_cast<double>(SecondAnimTime), false);
+			SecondAnimSeq->GetAnimationPose(SecondAnimPoseData, SecondExtractionContext);
+			SecondMeshPoses.InitPose(SecondAnimPoseData.GetPose());
+
+			float DiffVal = 0.0f;
+			for (FCompactPoseBoneIndex PoseBoneIndex : FirstAnimPoseData.GetPose().ForEachBoneIndex())
+			{
+				FTransform FirstTransform = FirstMeshPoses.GetComponentSpaceTransform(PoseBoneIndex);
+				FTransform SecondTransform = SecondMeshPoses.GetComponentSpaceTransform(PoseBoneIndex);
+				if (PoseBoneIndex != 0)
+				{
+					DistVal += (FirstTransform.GetTranslation() - SecondTransform.GetTranslation()).SizeSquared();
+				}
+			}
+		}
+	}
+}
+//outer is startanimtime to firstanim->seqlength...
+//inner is 0 to secondanim->seqlength...
+//for this function just find the smallest in the second...
+//return the end anim time
+static float GetBestBlendPointTimeAtStart(UAnimSequenceBase* FirstAnimSeq, UAnimSequenceBase* SecondAnimSeq, float StartFirstAnimTime, float FrameRate,
+	TArray<TArray<float>>& DistanceDifferences)
+{
+
+	//int32 FirstAnimNumFrames = (FirstAnimSeq->SequenceLength - StartFirstAnimTime) * FrameRate + 1;
+	int32 SecondAnimNumFrames = SecondAnimSeq->GetPlayLength() * FrameRate + 1;
+	if (SecondAnimNumFrames <= 0)
+	{
+		return 0.0f;
+	}
+	TArray<float>& Distances = DistanceDifferences[0];
+	float MinVal = Distances[0];
+	int32 SmallIndex = 0;
+	for (int32 Index = 1; Index < SecondAnimNumFrames; ++Index)
+	{
+		float NewMin = Distances[Index];
+		if (NewMin < MinVal)
+		{
+			MinVal = NewMin;
+			SmallIndex = Index;
+		}
+	}
+	return SmallIndex * 1.0f / FrameRate;
+}
+
+TOptional<FTransform>  UMovieSceneCommonAnimationTrack::GetRootMotion(FFrameTime CurrentTime)
+{
+	TOptional<FTransform> Transform;
+
+	UMovieScene* MovieScene = GetTypedOuter<UMovieScene>();
+	if (!MovieScene)
+	{
+		return Transform;
+	}
+	FFrameRate TickResolution = MovieScene->GetTickResolution();
+	
+	if (RootMotionParams.bRootMotionsDirty) 
+	{
+		SetUpRootMotions(true);
+	}
+	if (RootMotionParams.bHaveRootMotion == false)
+	{
+		return Transform;
+	}
+	SortSections();
+	TArray<FTransform> CurrentTransforms;
+	TArray<float> CurrentWeights;
+	TArray<FTransform> CurrentAdditiveTransforms;
+	TArray<float> CurrentAdditiveWeights;
+	FTransform CurrentTransform = FTransform::Identity;
+
+	// Use evaluation field to iterate, as it will take into account 'Evaluate Nearest Section'.
+	for (const FMovieSceneTrackEvaluationFieldEntry& FieldEntry : GetEvaluationField().Entries)
+	{
+		if (FieldEntry.Range.Contains(CurrentTime.FrameNumber))
+		{
+			if (UMovieSceneSkeletalAnimationSection* AnimSection = Cast<UMovieSceneSkeletalAnimationSection>(FieldEntry.Section))
+			{
+				FFrameNumber EvaluationFrame = FieldEntry.ForcedTime == TNumericLimits<int32>::Lowest() ? CurrentTime.FrameNumber : FieldEntry.ForcedTime;
+				if (AnimSection->Params.Animation)
+				{
+					UAnimSequenceBase* ValidAnimSequence = AnimSection->Params.Animation;
+					FMemMark Mark(FMemStack::Get());
+					FCompactPose OutPose;
+					TArray<FBoneIndexType> RequiredBoneIndexArray;
+					RequiredBoneIndexArray.AddUninitialized(ValidAnimSequence->GetSkeleton()->GetReferenceSkeleton().GetNum());
+					for (int32 BoneIndex = 0; BoneIndex < RequiredBoneIndexArray.Num(); ++BoneIndex)
+					{
+						RequiredBoneIndexArray[BoneIndex] = BoneIndex;
+					}
+
+					FBoneContainer BoneContainer(RequiredBoneIndexArray, UE::Anim::FCurveFilterSettings(UE::Anim::ECurveFilterMode::None), *ValidAnimSequence->GetSkeleton());
+					OutPose.ResetToRefPose(BoneContainer);
+					FBlendedCurve OutCurve;
+					OutCurve.InitFrom(BoneContainer);
+					UE::Anim::FStackAttributeContainer TempAttributes;
+					FAnimationPoseData AnimationPoseData(OutPose, OutCurve, TempAttributes);
+					UMovieSceneSkeletalAnimationSection::FRootMotionTransformParam Param;
+					Param.FrameRate = TickResolution;
+					Param.CurrentTime = EvaluationFrame;
+
+					if (AnimSection->GetRootMotionTransform(AnimationPoseData, Param))
+					{
+						if (!Param.bOutIsAdditive)
+						{
+							CurrentTransform = Param.OutTransform * AnimSection->PreviousTransform;
+							CurrentTransforms.Add(CurrentTransform);
+							CurrentWeights.Add(Param.OutWeight);
+						}
+						else
+						{
+							CurrentAdditiveTransforms.Add(Param.OutTransform);
+							CurrentAdditiveWeights.Add(Param.OutWeight);
+						}
+					}
+				}
+			}
+		}
+	}
+
+	BlendTheseTransformsByWeight(CurrentTransform, CurrentTransforms, CurrentWeights);
+	//now handle additive onto the current
+	if (CurrentAdditiveWeights.Num() > 0)
+	{
+		FTransform AdditiveTransform;
+		BlendTheseTransformsByWeight(AdditiveTransform, CurrentAdditiveTransforms, CurrentAdditiveWeights);
+		const ScalarRegister VBlendWeight(1.0f);
+		FTransform::BlendFromIdentityAndAccumulate(CurrentTransform, AdditiveTransform, VBlendWeight);
+	}
+	Transform = CurrentTransform;
+	return Transform;
+
+}
+
+void UMovieSceneCommonAnimationTrack::SetUpRootMotions(bool bForce)
+{
+	UMovieScene* MovieScene = GetTypedOuter<UMovieScene>();
+	if (!MovieScene)
+	{
+		return;
+	}
+	
+	if (bForce || RootMotionParams.bRootMotionsDirty)
+	{
+		RootMotionParams.bRootMotionsDirty = false;
+		RootMotionParams.bHaveRootMotion = false;
+		const FFrameRate MinDisplayRate(60, 1);
+
+		FFrameRate DisplayRate =  MovieScene->GetDisplayRate().AsDecimal() < MinDisplayRate.AsDecimal() ? MinDisplayRate : MovieScene->GetDisplayRate();
+		FFrameRate TickResolution = MovieScene->GetTickResolution();
+		FFrameTime FrameTick = FFrameTime(FMath::Max(1, TickResolution.AsFrameNumber(1.0).Value / DisplayRate.AsFrameNumber(1.0).Value));
+		if (AnimationSections.Num() == 0 || FrameTick.FrameNumber.Value == 0)
+		{
+#if WITH_EDITORONLY_DATA
+			RootMotionParams.RootTransforms.SetNum(0);
+#endif
+			return;
+		}
+		
+		SortSections();
+		//Set the TempOffset.
+		FTransform InitialTransform = FTransform::Identity;
+		UMovieSceneSkeletalAnimationSection* PrevAnimSection = nullptr;
+		//valid anim sequence to use to calculate bones.
+		UAnimSequenceBase* ValidAnimSequence = nullptr;
+		//if no transforms have offsets then don't do root motion caching.
+		bool bAnySectionsHaveOffset = false;
+		TArray<UMovieSceneSkeletalAnimationSection*> AnimSections;
+		for (UMovieSceneSection* Section : AnimationSections)
+		{
+			UMovieSceneSkeletalAnimationSection* AnimSection = Cast<UMovieSceneSkeletalAnimationSection>(Section);
+			if (AnimSection)
+			{
+				AnimSections.Add(AnimSection);
+				if (AnimSection->StartLocationOffset.IsNearlyZero() == false || AnimSection->StartRotationOffset.IsNearlyZero() == false ||
+					AnimSection->MatchedLocationOffset.IsNearlyZero() == false || AnimSection->MatchedRotationOffset.IsNearlyZero() == false)
+				{
+					bAnySectionsHaveOffset = true;
+				}
+				if (ValidAnimSequence == nullptr)
+				{
+					ValidAnimSequence = AnimSection->Params.Animation;
+				}
+				if (PrevAnimSection)
+				{
+					if (UAnimSequenceBase* PrevAnimSequence = PrevAnimSection->Params.Animation)
+					{	
+						if (bAnySectionsHaveOffset)
+						{ 
+							RootMotionParams.RootMotionStartOffset = AnimSection->GetRootMotionStartOffset();
+							FMemMark Mark(FMemStack::Get());
+							FCompactPose OutPose;
+							TArray<FBoneIndexType> RequiredBoneIndexArray;
+							RequiredBoneIndexArray.AddUninitialized(PrevAnimSequence->GetSkeleton()->GetReferenceSkeleton().GetNum());
+							for (int32 BoneIndex = 0; BoneIndex < RequiredBoneIndexArray.Num(); ++BoneIndex)
+							{
+								RequiredBoneIndexArray[BoneIndex] = BoneIndex;
+							}
+
+							FBoneContainer BoneContainer(RequiredBoneIndexArray, UE::Anim::FCurveFilterSettings(UE::Anim::ECurveFilterMode::None), *PrevAnimSequence->GetSkeleton());
+							OutPose.ResetToRefPose(BoneContainer);
+							FBlendedCurve OutCurve;
+							OutCurve.InitFrom(BoneContainer);
+
+							UE::Anim::FStackAttributeContainer TempAttributes;
+							FAnimationPoseData AnimationPoseData(OutPose, OutCurve, TempAttributes);
+
+							UMovieSceneSkeletalAnimationSection::FRootMotionTransformParam Param;
+							Param.FrameRate = MovieScene->GetTickResolution();
+							Param.CurrentTime = AnimSection->GetRange().GetLowerBoundValue();
+
+							if (PrevAnimSection->GetRootMotionTransform(AnimationPoseData, Param))
+							{
+								AnimSection->PreviousTransform = Param.OutPoseTransform.GetRelativeTransformReverse(Param.OutTransform);
+								AnimSection->PreviousTransform = AnimSection->PreviousTransform * InitialTransform;
+
+							}
+							else
+							{
+								AnimSection->PreviousTransform = FTransform::Identity;
+							}
+						}
+						else
+						{
+							AnimSection->PreviousTransform = FTransform::Identity;
+						}
+					}
+					else
+					{
+						AnimSection->PreviousTransform = FTransform::Identity;
+					}
+					InitialTransform = AnimSection->PreviousTransform;
+
+				}
+				else
+				{
+					AnimSection->PreviousTransform = FTransform::Identity;
+				}
+				PrevAnimSection = AnimSection;
+				AnimSection->SetBoneIndexForRootMotionCalculations(bBlendFirstChildOfRoot);
+			}
+		}
+
+		if (AnimSections.Num() == 0)
+		{
+#if WITH_EDITORONLY_DATA
+			RootMotionParams.RootTransforms.SetNum(0);
+#endif			
+			return;
+		}
+
+		if (bAnySectionsHaveOffset == false && !ShouldUseRootMotions())
+		{
+#if WITH_EDITORONLY_DATA
+			RootMotionParams.RootTransforms.SetNum(0);
+#endif			
+			return;
+		}
+		RootMotionParams.bHaveRootMotion = true;
+		RootMotionParams.StartFrame = AnimSections[0]->GetInclusiveStartFrame();
+		RootMotionParams.StartFrame = AnimSections.Last()->GetExclusiveEndFrame() - 1;
+		RootMotionParams.FrameTick = FrameTick;		
+		
+#if WITH_EDITORONLY_DATA
+		if (RootMotionParams.bCacheRootTransforms == false)
+		{
+			return;
+		}
+		//set up pose from valid anim sequences.
+		FMemMark Mark(FMemStack::Get());
+		FCompactPose OutPose;
+
+		if (ValidAnimSequence)
+		{
+			TArray<FBoneIndexType> RequiredBoneIndexArray;
+			const UE::Anim::FCurveFilterSettings CurveFilterSettings;
+			RequiredBoneIndexArray.AddUninitialized(ValidAnimSequence->GetSkeleton()->GetReferenceSkeleton().GetNum());
+			for (int32 BoneIndex = 0; BoneIndex < RequiredBoneIndexArray.Num(); ++BoneIndex)
+			{
+				RequiredBoneIndexArray[BoneIndex] = BoneIndex;
+			}
+
+			FBoneContainer BoneContainer(RequiredBoneIndexArray, CurveFilterSettings, *ValidAnimSequence->GetSkeleton());
+			OutPose.ResetToRefPose(BoneContainer);
+			FBlendedCurve OutCurve;
+			OutCurve.InitFrom(BoneContainer);
+			TArray< UMovieSceneSkeletalAnimationSection*> SectionsAtCurrentTime;
+			int32 NumTotal = (RootMotionParams.EndFrame.FrameNumber.Value - RootMotionParams.StartFrame.FrameNumber.Value) / (RootMotionParams.FrameTick.FrameNumber.Value) + 1;
+			RootMotionParams.RootTransforms.SetNum(NumTotal);
+
+			TArray<FTransform> CurrentTransforms;
+			TArray<float> CurrentWeights;
+			TArray<FTransform> CurrentAdditiveTransforms;
+			TArray<float> CurrentAdditiveWeights;
+			FFrameTime PreviousFrame = RootMotionParams.StartFrame;
+			int32 Index = 0;
+			for (FFrameTime FrameNumber = RootMotionParams.StartFrame; FrameNumber <= RootMotionParams.EndFrame; FrameNumber += RootMotionParams.FrameTick)
+			{
+				CurrentTransforms.SetNum(0);
+				CurrentWeights.SetNum(0);
+				FTransform CurrentTransform(FTransform::Identity), ParentTransform(FTransform::Identity);
+
+				UMovieSceneSkeletalAnimationSection* PrevSection = nullptr;
+				for (UMovieSceneSection* Section : AnimationSections)
+				{
+					UMovieSceneSkeletalAnimationSection* AnimSection = Cast<UMovieSceneSkeletalAnimationSection>(Section);
+					if (AnimSection && AnimSection->GetRange().Contains(FrameNumber.FrameNumber))
+					{
+						UE::Anim::FStackAttributeContainer TempAttributes;
+						FAnimationPoseData AnimationPoseData(OutPose, OutCurve, TempAttributes);
+						UMovieSceneSkeletalAnimationSection::FRootMotionTransformParam Param;
+						Param.FrameRate = TickResolution;
+						Param.CurrentTime = FrameNumber.FrameNumber;
+
+						if (AnimSection->GetRootMotionTransform(AnimationPoseData,Param))
+						{
+							if (!Param.bOutIsAdditive)
+							{
+								CurrentTransform = Param.OutTransform * AnimSection->PreviousTransform;
+								CurrentTransforms.Add(CurrentTransform);
+								CurrentWeights.Add(Param.OutWeight);
+							}
+							else
+							{
+								CurrentAdditiveTransforms.Add(Param.OutTransform);
+								CurrentAdditiveWeights.Add(Param.OutWeight);
+							}
+						}
+						PrevSection = AnimSection;
+					}
+				}
+
+				BlendTheseTransformsByWeight(CurrentTransform, CurrentTransforms, CurrentWeights);
+				//now handle additive onto the current
+				if (CurrentAdditiveWeights.Num() > 0)
+				{
+					FTransform AdditiveTransform;
+					BlendTheseTransformsByWeight(AdditiveTransform, CurrentAdditiveTransforms, CurrentAdditiveWeights);
+					const ScalarRegister VBlendWeight(1.0f);
+					FTransform::BlendFromIdentityAndAccumulate(CurrentTransform, AdditiveTransform, VBlendWeight);
+				}
+				RootMotionParams.RootTransforms[Index] = CurrentTransform;
+				++Index;
+				PreviousFrame = FrameNumber;
+
+			}
+		}
+		else //no valid anim sequence just clear out
+		{
+			RootMotionParams.RootTransforms.SetNum(0);
+		}
+#endif
+	}
+
+}
+
+static FTransform GetTransformForBoneRelativeToIndex(UAnimSequence* AnimSequence, USkeletalMeshComponent* MeshComponent, const FName& InBoneName,
+	const FCompactPoseBoneIndex& ParentCPIndex, double Seconds)
+{
+	FTransform  WorldTransform = FTransform::Identity;
+	//AnimSequence->GetBoneTransform doesn't seem to be as accurate as GetAnimationPose
+	FMemMark Mark(FMemStack::Get());
+	FCompactPose OutPose;
+	TArray<FBoneIndexType> RequiredBoneIndexArray;
+	const UE::Anim::FCurveFilterSettings CurveFilterSettings;
+	RequiredBoneIndexArray.AddUninitialized(AnimSequence->GetSkeleton()->GetReferenceSkeleton().GetNum());
+	for (int32 BoneIndex = 0; BoneIndex < RequiredBoneIndexArray.Num(); ++BoneIndex)
+	{
+		RequiredBoneIndexArray[BoneIndex] = BoneIndex;
+	}
+
+	FBoneContainer BoneContainer(RequiredBoneIndexArray, CurveFilterSettings, *AnimSequence->GetSkeleton());
+	OutPose.ResetToRefPose(BoneContainer);
+	FBlendedCurve OutCurve;
+	OutCurve.InitFrom(BoneContainer);
+	UE::Anim::FStackAttributeContainer TempAttributes;
+	FAnimationPoseData AnimationPoseData(OutPose, OutCurve, TempAttributes);
+	FAnimExtractContext ExtractionContext(Seconds, false);
+#if WITH_EDITOR
+	ExtractionContext.bIgnoreRootLock = true;
+#endif
+	AnimSequence->GetAnimationPose(AnimationPoseData, ExtractionContext);
+	int32 MeshIndex = AnimationPoseData.GetPose().GetBoneContainer().GetPoseBoneIndexForBoneName(InBoneName);
+	if (MeshIndex != INDEX_NONE)
+	{
+		FCompactPoseBoneIndex CPIndex = AnimationPoseData.GetPose().GetBoneContainer().MakeCompactPoseIndex(FMeshPoseBoneIndex(MeshIndex));
+		if (CPIndex != INDEX_NONE )
+		{
+			FTransform BoneTransform = FTransform::Identity;
+			WorldTransform *= BoneTransform;
+			do
+			{
+				BoneTransform = AnimationPoseData.GetPose()[CPIndex];
+				WorldTransform *= BoneTransform;
+				if (CPIndex == ParentCPIndex)  //if we are the parent then we stop
+				{
+					CPIndex = FCompactPoseBoneIndex(INDEX_NONE);
+				}
+				else
+				{
+					CPIndex = AnimationPoseData.GetPose().GetBoneContainer().GetParentBoneIndex(CPIndex);
+				}
+			} while (CPIndex.IsValid());
+		}
+	}
+	return WorldTransform;
+}
+
+enum class ESkelAnimRotationOrder 
+{
+	XYZ,
+	XZY,
+	YXZ,
+	YZX,
+	ZXY,
+	ZYX
+};
+
+static FQuat QuatFromEuler(const FVector& XYZAnglesInDegrees, ESkelAnimRotationOrder RotationOrder)
+{
+	float X = FMath::DegreesToRadians(XYZAnglesInDegrees.X);
+	float Y = FMath::DegreesToRadians(XYZAnglesInDegrees.Y);
+	float Z = FMath::DegreesToRadians(XYZAnglesInDegrees.Z);
+
+	float CosX = FMath::Cos(X * 0.5f);
+	float CosY = FMath::Cos(Y * 0.5f);
+	float CosZ = FMath::Cos(Z * 0.5f);
+
+	float SinX = FMath::Sin(X * 0.5f);
+	float SinY = FMath::Sin(Y * 0.5f);
+	float SinZ = FMath::Sin(Z * 0.5f);
+
+	if (RotationOrder == ESkelAnimRotationOrder::XYZ)
+	{
+		return FQuat(SinX * CosY * CosZ - CosX * SinY * SinZ,
+			CosX * SinY * CosZ + SinX * CosY * SinZ,
+			CosX * CosY * SinZ - SinX * SinY * CosZ,
+			CosX * CosY * CosZ + SinX * SinY * SinZ);
+
+	}
+	else if (RotationOrder == ESkelAnimRotationOrder::XZY)
+	{
+		return FQuat(SinX * CosY * CosZ + CosX * SinY * SinZ,
+			CosX * SinY * CosZ + SinX * CosY * SinZ,
+			CosX * CosY * SinZ - SinX * SinY * CosZ,
+			CosX * CosY * CosZ - SinX * SinY * SinZ);
+
+	}
+	else if (RotationOrder == ESkelAnimRotationOrder::YXZ)
+	{
+		return FQuat(SinX * CosY * CosZ - CosX * SinY * SinZ,
+			CosX * SinY * CosZ + SinX * CosY * SinZ,
+			CosX * CosY * SinZ + SinX * SinY * CosZ,
+			CosX * CosY * CosZ - SinX * SinY * SinZ);
+
+	}
+	else if (RotationOrder == ESkelAnimRotationOrder::YZX)
+	{
+		return FQuat(SinX * CosY * CosZ - CosX * SinY * SinZ,
+			CosX * SinY * CosZ - SinX * CosY * SinZ,
+			CosX * CosY * SinZ + SinX * SinY * CosZ,
+			CosX * CosY * CosZ + SinX * SinY * SinZ);
+	}
+	else if (RotationOrder == ESkelAnimRotationOrder::ZXY)
+	{
+		return FQuat(SinX * CosY * CosZ + CosX * SinY * SinZ,
+			CosX * SinY * CosZ - SinX * CosY * SinZ,
+			CosX * CosY * SinZ - SinX * SinY * CosZ,
+			CosX * CosY * CosZ + SinX * SinY * SinZ);
+
+	}
+	else if (RotationOrder == ESkelAnimRotationOrder::ZYX)
+	{
+		return FQuat(SinX * CosY * CosZ + CosX * SinY * SinZ,
+			CosX * SinY * CosZ - SinX * CosY * SinZ,
+			CosX * CosY * SinZ + SinX * SinY * CosZ,
+			CosX * CosY * CosZ - SinX * SinY * SinZ);
+
+	}
+
+	// should not happen
+	return FQuat::Identity;
+}
+
+static FVector EulerFromQuat(const FQuat& Rotation, ESkelAnimRotationOrder RotationOrder)
+{
+	float X = Rotation.X;
+	float Y = Rotation.Y;
+	float Z = Rotation.Z;
+	float W = Rotation.W;
+	float X2 = X * 2.f;
+	float Y2 = Y * 2.f;
+	float Z2 = Z * 2.f;
+	float XX2 = X * X2;
+	float XY2 = X * Y2;
+	float XZ2 = X * Z2;
+	float YX2 = Y * X2;
+	float YY2 = Y * Y2;
+	float YZ2 = Y * Z2;
+	float ZX2 = Z * X2;
+	float ZY2 = Z * Y2;
+	float ZZ2 = Z * Z2;
+	float WX2 = W * X2;
+	float WY2 = W * Y2;
+	float WZ2 = W * Z2;
+
+	FVector AxisX, AxisY, AxisZ;
+	AxisX.X = (1.f - (YY2 + ZZ2));
+	AxisY.X = (XY2 + WZ2);
+	AxisZ.X = (XZ2 - WY2);
+	AxisX.Y = (XY2 - WZ2);
+	AxisY.Y = (1.f - (XX2 + ZZ2));
+	AxisZ.Y = (YZ2 + WX2);
+	AxisX.Z = (XZ2 + WY2);
+	AxisY.Z = (YZ2 - WX2);
+	AxisZ.Z = (1.f - (XX2 + YY2));
+
+	FVector Result = FVector::ZeroVector;
+
+	if (RotationOrder == ESkelAnimRotationOrder::XYZ)
+	{
+		Result.Y = FMath::Asin(-FMath::Clamp<float>(AxisZ.X, -1.f, 1.f));
+
+		if (FMath::Abs(AxisZ.X) < 1.f - SMALL_NUMBER)
+		{
+			Result.X = FMath::Atan2(AxisZ.Y, AxisZ.Z);
+			Result.Z = FMath::Atan2(AxisY.X, AxisX.X);
+		}
+		else
+		{
+			Result.X = 0.f;
+			Result.Z = FMath::Atan2(-AxisX.Y, AxisY.Y);
+		}
+	}
+	else if (RotationOrder == ESkelAnimRotationOrder::XZY)
+	{
+
+		Result.Z = FMath::Asin(FMath::Clamp<float>(AxisY.X, -1.f, 1.f));
+
+		if (FMath::Abs(AxisY.X) < 1.f - SMALL_NUMBER)
+		{
+			Result.X = FMath::Atan2(-AxisY.Z, AxisY.Y);
+			Result.Y = FMath::Atan2(-AxisZ.X, AxisX.X);
+		}
+		else
+		{
+			Result.X = 0.f;
+			Result.Y = FMath::Atan2(AxisX.Z, AxisZ.Z);
+		}
+	}
+	else if (RotationOrder == ESkelAnimRotationOrder::YXZ)
+	{
+		Result.X = FMath::Asin(FMath::Clamp<float>(AxisZ.Y, -1.f, 1.f));
+
+		if (FMath::Abs(AxisZ.Y) < 1.f - SMALL_NUMBER)
+		{
+			Result.Y = FMath::Atan2(-AxisZ.X, AxisZ.Z);
+			Result.Z = FMath::Atan2(-AxisX.Y, AxisY.Y);
+		}
+		else
+		{
+			Result.Y = 0.f;
+			Result.Z = FMath::Atan2(AxisY.X, AxisX.X);
+		}
+	}
+	else if (RotationOrder == ESkelAnimRotationOrder::YZX)
+	{
+		Result.Z = FMath::Asin(-FMath::Clamp<float>(AxisX.Y, -1.f, 1.f));
+
+		if (FMath::Abs(AxisX.Y) < 1.f - SMALL_NUMBER)
+		{
+			Result.X = FMath::Atan2(AxisZ.Y, AxisY.Y);
+			Result.Y = FMath::Atan2(AxisX.Z, AxisX.X);
+		}
+		else
+		{
+			Result.X = FMath::Atan2(-AxisY.Z, AxisZ.Z);
+			Result.Y = 0.f;
+		}
+	}
+	else if (RotationOrder == ESkelAnimRotationOrder::ZXY)
+	{
+		Result.X = FMath::Asin(-FMath::Clamp<float>(AxisY.Z, -1.f, 1.f));
+
+		if (FMath::Abs(AxisY.Z) < 1.f - SMALL_NUMBER)
+		{
+			Result.Y = FMath::Atan2(AxisX.Z, AxisZ.Z);
+			Result.Z = FMath::Atan2(AxisY.X, AxisY.Y);
+		}
+		else
+		{
+			Result.Y = FMath::Atan2(-AxisZ.X, AxisX.X);
+			Result.Z = 0.f;
+		}
+	}
+	else if (RotationOrder == ESkelAnimRotationOrder::ZYX)
+	{
+		Result.Y = FMath::Asin(FMath::Clamp<float>(AxisX.Z, -1.f, 1.f));
+
+		if (FMath::Abs(AxisX.Z) < 1.f - SMALL_NUMBER)
+		{
+			Result.X = FMath::Atan2(-AxisY.Z, AxisZ.Z);
+			Result.Z = FMath::Atan2(-AxisX.Y, AxisX.X);
+		}
+		else
+		{
+			Result.X = FMath::Atan2(AxisZ.Y, AxisY.Y);
+			Result.Z = 0.f;
+		}
+	}
+
+	return Result * 180.f / PI;
+}
+/**
+*  Function to find best rotation order given how we are matching the rotations. 
+*  If it is matched we need to make sure it happens first
+*  Issue is Yaw is most common match but by default FRotation:FQuat conversions it's last, this causes issues.
+*/
+static ESkelAnimRotationOrder FindBestRotationOrder(bool bMatchRoll, bool bMatchPitch, bool bMatchYaw)
+{
+	if (bMatchYaw)
+	{
+		return ESkelAnimRotationOrder::YXZ;
+	}
+	if (bMatchPitch)
+	{
+		return ESkelAnimRotationOrder::YZX;
+	}
+	return ESkelAnimRotationOrder::XYZ;
+}
+
+void UMovieSceneCommonAnimationTrack::MatchSectionByBoneTransform(bool bMatchWithPrevious, USkeletalMeshComponent* SkelMeshComp, UMovieSceneSkeletalAnimationSection* CurrentSection, FFrameTime CurrentFrame, FFrameRate FrameRate,
+	const FName& BoneName, FTransform& SecondSectionRootDiff, FVector& TranslationDiff, FQuat& RotationDiff) //add options for z and for rotation.
+{
+	SortSections();
+	UMovieSceneSection* PrevSection = nullptr;
+	UMovieSceneSection* NextSection = nullptr;
+	for (int32 Index = 0; Index < AnimationSections.Num(); ++Index)
+	{
+		UMovieSceneSection* Section = AnimationSections[Index];
+		if (Section == CurrentSection)
+		{
+			if (++Index < AnimationSections.Num())
+			{
+				NextSection = AnimationSections[Index];
+			}
+			break;
+		}
+		PrevSection = Section;
+	}
+
+	TranslationDiff = FVector(0.0f, 0.0f, 0.0f);
+	RotationDiff = FQuat::Identity;
+	SecondSectionRootDiff = FTransform::Identity;
+
+	if (bMatchWithPrevious && PrevSection)
+	{
+		UMovieSceneSkeletalAnimationSection* FirstSection = Cast<UMovieSceneSkeletalAnimationSection>(PrevSection);
+		UAnimSequence* FirstAnimSequence = Cast<UAnimSequence>(FirstSection->Params.Animation);
+		UAnimSequence* SecondAnimSequence = Cast<UAnimSequence>(CurrentSection->Params.Animation);
+
+		if (FirstAnimSequence && SecondAnimSequence)
+		{
+			double FirstSectionTime = FirstSection->MapTimeToAnimation(CurrentFrame, FrameRate);
+			//use same index for all
+			int32 Index = CurrentSection->SetBoneIndexForRootMotionCalculations(bBlendFirstChildOfRoot);
+			FCompactPoseBoneIndex ParentIndex(Index);
+			
+			FTransform  FirstTransform = GetTransformForBoneRelativeToIndex(FirstAnimSequence, SkelMeshComp, BoneName, ParentIndex, FirstSectionTime);
+			double SecondSectionTime = CurrentSection->MapTimeToAnimation(CurrentFrame, FrameRate);
+			FTransform  SecondTransform = GetTransformForBoneRelativeToIndex(SecondAnimSequence, SkelMeshComp, BoneName, ParentIndex,SecondSectionTime);
+			//Need to match the translations and rotations here 
+			//First need to get the correct rotation order based upon what's matching, otherwise if not all are matched 
+			//and one rotation is set last we will get errors.
+			ESkelAnimRotationOrder RotationOrder = FindBestRotationOrder(CurrentSection->bMatchRotationRoll, CurrentSection->bMatchRotationPitch, CurrentSection->bMatchRotationYaw);
+
+			FVector FirstTransformTranslation = FirstTransform.GetTranslation();
+			FVector SecondTransformTranslation = SecondTransform.GetTranslation();
+			FQuat FirstTransformQuat = FirstTransform.GetRotation();
+			FQuat SecondTransformQuat = SecondTransform.GetRotation();
+
+			FirstTransformQuat.EnforceShortestArcWith(SecondTransformQuat);
+
+			FRotator FirstTransformRotation(FRotator::MakeFromEuler(EulerFromQuat(FirstTransformQuat, RotationOrder)));
+			FRotator SecondTransformRotation(FRotator::MakeFromEuler(EulerFromQuat(SecondTransformQuat, RotationOrder)));
+			SecondTransformRotation.SetClosestToMe(FirstTransformRotation);
+
+			if (!CurrentSection->bMatchTranslation)
+			{
+				FirstTransformTranslation.X = SecondTransformTranslation.X;
+				FirstTransformTranslation.Y = SecondTransformTranslation.Y;
+				FirstTransformTranslation.Z = SecondTransformTranslation.Z;
+			}
+			if (!CurrentSection->bMatchIncludeZHeight)
+			{
+				FirstTransformTranslation.Z = SecondTransformTranslation.Z;
+			}
+			FirstTransform.SetTranslation(FirstTransformTranslation);
+
+			if (!CurrentSection->bMatchRotationYaw)
+			{
+				FirstTransformRotation.Yaw = SecondTransformRotation.Yaw;
+			}
+			if (!CurrentSection->bMatchRotationPitch)
+			{
+				FirstTransformRotation.Pitch = SecondTransformRotation.Pitch;
+			}
+			if (!CurrentSection->bMatchRotationRoll)
+			{
+				FirstTransformRotation.Roll = SecondTransformRotation.Roll;
+			}
+			FirstTransformQuat = QuatFromEuler(FirstTransformRotation.Euler(), RotationOrder);
+			FirstTransform.SetRotation(FirstTransformQuat);
+
+			// Below is the match but we need to use GetRelativeTransformReverse since Inverse doesn't work as expected.
+			//	* GetRelativeTransformReverse returns this(-1)* Other, and parameter is Other.
+			SecondSectionRootDiff = SecondTransform.GetRelativeTransformReverse(FirstTransform);
+			TranslationDiff = SecondSectionRootDiff.GetTranslation();
+			RotationDiff = SecondSectionRootDiff.GetRotation();
+
+		}
+	}
+	else if (bMatchWithPrevious == false && NextSection) //match with next
+	{
+		UMovieSceneSkeletalAnimationSection* SecondSection = Cast<UMovieSceneSkeletalAnimationSection>(NextSection);
+		UAnimSequence* FirstAnimSequence = Cast<UAnimSequence>(CurrentSection->Params.Animation);
+		UAnimSequence* SecondAnimSequence = Cast<UAnimSequence>(SecondSection->Params.Animation);
+
+		if (FirstAnimSequence && SecondAnimSequence)
+		{
+			//use same index for all
+			int32 Index = CurrentSection->SetBoneIndexForRootMotionCalculations(bBlendFirstChildOfRoot);
+			FCompactPoseBoneIndex ParentIndex(Index);
+			float FirstSectionTime = static_cast<float>(CurrentSection->MapTimeToAnimation(CurrentFrame, FrameRate));
+			FTransform  FirstTransform = GetTransformForBoneRelativeToIndex(FirstAnimSequence, SkelMeshComp, BoneName,ParentIndex, FirstSectionTime);
+			float SecondSectionTime = static_cast<float>(SecondSection->MapTimeToAnimation(CurrentFrame, FrameRate));
+			FTransform  SecondTransform = GetTransformForBoneRelativeToIndex(SecondAnimSequence, SkelMeshComp, BoneName,ParentIndex, SecondSectionTime);
+
+			//Need to match the translations and rotations here 
+			//First need to get the correct rotation order based upon what's matching, otherwise if not all are matched 
+			//and one rotation is set last we will get errors.
+			ESkelAnimRotationOrder RotationOrder = FindBestRotationOrder(CurrentSection->bMatchRotationRoll, CurrentSection->bMatchRotationPitch, CurrentSection->bMatchRotationYaw);
+
+			FVector FirstTransformTranslation = FirstTransform.GetTranslation();
+			FVector SecondTransformTranslation = SecondTransform.GetTranslation();
+			FQuat FirstTransformQuat = FirstTransform.GetRotation();
+			FQuat SecondTransformQuat = SecondTransform.GetRotation();
+
+			SecondTransformQuat.EnforceShortestArcWith(FirstTransformQuat);
+
+			FRotator FirstTransformRotation(FRotator::MakeFromEuler(EulerFromQuat(FirstTransformQuat, RotationOrder)));
+			FRotator SecondTransformRotation(FRotator::MakeFromEuler(EulerFromQuat(SecondTransformQuat, RotationOrder)));
+			FirstTransformRotation.SetClosestToMe(SecondTransformRotation);
+
+			if (!CurrentSection->bMatchTranslation)
+			{
+				SecondTransformTranslation.X = FirstTransformTranslation.X;
+				SecondTransformTranslation.Y = FirstTransformTranslation.Y;
+				SecondTransformTranslation.Z = FirstTransformTranslation.Z;
+			}
+			if (!CurrentSection->bMatchIncludeZHeight)
+			{
+				SecondTransformTranslation.Z = FirstTransformTranslation.Z;
+			}
+			SecondTransform.SetTranslation(SecondTransformTranslation);
+
+			if (!CurrentSection->bMatchRotationYaw)
+			{
+				SecondTransformRotation.Yaw = FirstTransformRotation.Yaw;
+			}
+			if (!CurrentSection->bMatchRotationPitch)
+			{
+				SecondTransformRotation.Pitch = FirstTransformRotation.Pitch;
+			}
+			if (!CurrentSection->bMatchRotationRoll)
+			{
+				SecondTransformRotation.Roll = FirstTransformRotation.Roll;
+			}
+			SecondTransformQuat = QuatFromEuler(SecondTransformRotation.Euler(), RotationOrder);
+			SecondTransform.SetRotation(SecondTransformQuat);
+
+			//GetRelativeTransformReverse returns this(-1)* Other, and parameter is Other.
+			SecondSectionRootDiff = FirstTransform.GetRelativeTransformReverse(SecondTransform);
+			TranslationDiff = SecondSectionRootDiff.GetTranslation();
+			RotationDiff = SecondSectionRootDiff.GetRotation();
+
+		}
+	}
+}
+
+#if WITH_EDITORONLY_DATA
+
+void UMovieSceneCommonAnimationTrack::ToggleShowRootMotionTrail()
+{
+	bShowRootMotionTrail = bShowRootMotionTrail ? false : true;
+}
+#endif
+
+
+//MZ To Do need way to get passed the skelmeshcomp when we add or move a section.
+void UMovieSceneCommonAnimationTrack::AutoMatchSectionRoot(UMovieSceneSkeletalAnimationSection* CurrentSection)
+{
+	return;
+#if 0
+	UMovieScene* MovieScene = GetTypedOuter<UMovieScene>();
+	if (AnimationSections.Num() > 0 && MovieScene && CurrentSection)
+	{
+		SortSections();
+
+		for (int32 Index = 0; Index < AnimationSections.Num(); ++Index)
+		{
+			UMovieSceneSection* Section = AnimationSections[Index];
+			if (Section && Section == CurrentSection)
+			{
+				CurrentSection->bMatchWithPrevious = (Index == 0) ? false : true;
+				FFrameTime FrameTime = (Index == 0) ? CurrentSection->GetRange().GetUpperBoundValue() : CurrentSection->GetRange().GetLowerBoundValue();
+				USkeletalMeshComponent* SkelMeshComp = nullptr;
+				CurrentSection->MatchSectionByBoneTransform(SkelMeshComp, FrameTime, MovieScene->GetTickResolution(), CurrentSection->MatchedBoneName);
+			}
+		}
+	}
+#endif
 }
\ No newline at end of file
diff --git a/Engine/Source/Runtime/MovieSceneTracks/Private/Tracks/MovieSceneSkeletalAnimationTrack.cpp b/Engine/Source/Runtime/MovieSceneTracks/Private/Tracks/MovieSceneSkeletalAnimationTrack.cpp
--- a/Engine/Source/Runtime/MovieSceneTracks/Private/Tracks/MovieSceneSkeletalAnimationTrack.cpp
+++ b/Engine/Source/Runtime/MovieSceneTracks/Private/Tracks/MovieSceneSkeletalAnimationTrack.cpp
@@ -34,12 +34,10 @@
 UMovieSceneSkeletalAnimationTrack::UMovieSceneSkeletalAnimationTrack(const FObjectInitializer& ObjectInitializer)
 	: Super(ObjectInitializer)
 	, bUseLegacySectionIndexBlend(false)
-	, bBlendFirstChildOfRoot(false)
 {
 #if WITH_EDITORONLY_DATA
 	TrackTint = FColor(124, 15, 124, 65);
 	bSupportsDefaultSections = false;
-	bShowRootMotionTrail = false;
 #endif
 
 	SupportedBlendTypes.Add(EMovieSceneBlendType::Absolute);
@@ -52,30 +50,6 @@ UMovieSceneSkeletalAnimationTrack::UMovieSceneSkeletalAnimationTrack(const FObje
 /* UMovieSceneSkeletalAnimationTrack interface
  *****************************************************************************/
 
-UMovieSceneSection* UMovieSceneCommonAnimationTrack::AddNewAnimationOnRow(FFrameNumber KeyTime, UAnimSequenceBase* AnimSequence, int32 RowIndex)
-{
-	UMovieSceneSkeletalAnimationSection* NewSection = Cast<UMovieSceneSkeletalAnimationSection>(CreateNewSection());
-	{
-		FFrameTime AnimationLength = AnimSequence->GetPlayLength() * GetTypedOuter<UMovieScene>()->GetTickResolution();
-		int32 IFrameNumber = AnimationLength.FrameNumber.Value + static_cast<int32>(AnimationLength.GetSubFrame() + 0.5f) + 1;
-		NewSection->InitialPlacementOnRow(AnimationSections, KeyTime, IFrameNumber, RowIndex);
-		NewSection->Params.Animation = AnimSequence;
-
-#if WITH_EDITORONLY_DATA
-		FQualifiedFrameTime SourceStartFrameTime;
-		if (UAnimationBlueprintLibrary::EvaluateRootBoneTimecodeAttributesAtTime(NewSection->Params.Animation, 0.0f, SourceStartFrameTime))
-		{
-			NewSection->TimecodeSource.Timecode = SourceStartFrameTime.ToTimecode();
-		}
-#endif
-	}
-
-	AddSection(*NewSection);
-
-	return NewSection;
-}
-
-
 TArray<UMovieSceneSection*> UMovieSceneSkeletalAnimationTrack::GetAnimSectionsAtTime(FFrameNumber Time)
 {
 	TArray<UMovieSceneSection*> Sections;
@@ -90,6 +64,11 @@ TArray<UMovieSceneSection*> UMovieSceneSkeletalAnimationTrack::GetAnimSectionsAt
 	return Sections;
 }
 
+bool UMovieSceneSkeletalAnimationTrack::ShouldUseRootMotions() const
+{
+	//if we are swapping root bone turn on root motion matching anyway
+	return SwapRootBone != ESwapRootBone::SwapRootBone_None;
+}
 
 /* UMovieSceneTrack interface
  *****************************************************************************/
@@ -170,57 +149,6 @@ ESwapRootBone UMovieSceneSkeletalAnimationTrack::GetSwapRootBone() const
 	return SwapRootBone;
 }
 
-#if WITH_EDITOR
-void UMovieSceneSkeletalAnimationTrack::PostEditImport()
-{
-	Super::PostEditImport();
-	RootMotionParams.bRootMotionsDirty = true;
-}
-
-void UMovieSceneSkeletalAnimationTrack::PostEditUndo()
-{
-	Super::PostEditUndo();
-	RootMotionParams.bRootMotionsDirty = true;
-}
-
-void UMovieSceneSkeletalAnimationTrack::PostEditChangeProperty(FPropertyChangedEvent& PropertyChangedEvent)
-{
-	Super::PostEditChangeProperty(PropertyChangedEvent);
-	RootMotionParams.bRootMotionsDirty = true;
-}
-
-#endif
-
-void UMovieSceneSkeletalAnimationTrack::AddSection(UMovieSceneSection& Section)
-{
-	Super::AddSection(Section);
-
-	UMovieSceneSkeletalAnimationSection* AnimSection = Cast< UMovieSceneSkeletalAnimationSection>(&Section);
-	if (AnimSection)
-	{
-		AnimSection->Params.SwapRootBone = SwapRootBone;
-		SetUpRootMotions(true);
-	}
-}
-
-void UMovieSceneSkeletalAnimationTrack::UpdateEasing()
-{
-	Super::UpdateEasing();
-	SetRootMotionsDirty();
-}
-
-void UMovieSceneSkeletalAnimationTrack::RemoveSection(UMovieSceneSection& Section)
-{
-	Super::RemoveSection(Section);
-	SetUpRootMotions(true);
-}
-
-void UMovieSceneSkeletalAnimationTrack::RemoveSectionAt(int32 SectionIndex)
-{
-	Super::RemoveSectionAt(SectionIndex);
-	SetUpRootMotions(true);
-}
-
 #if WITH_EDITORONLY_DATA
 
 FText UMovieSceneSkeletalAnimationTrack::GetDefaultDisplayName() const
@@ -258,958 +186,5 @@ bool UMovieSceneSkeletalAnimationTrack::PopulateEvaluationTree(TMovieSceneEvalua
 	return true;
 }
 
-#if WITH_EDITOR
-EMovieSceneSectionMovedResult UMovieSceneSkeletalAnimationTrack::OnSectionMoved(UMovieSceneSection& Section, const FMovieSceneSectionMovedParams& Params)
-{
-	if (Params.MoveType == EPropertyChangeType::ValueSet)
-	{
-		RootMotionParams.bRootMotionsDirty = true;
-	}
-	
-	return EMovieSceneSectionMovedResult::None;
-}
-#endif
-
-void UMovieSceneSkeletalAnimationTrack::SortSections()
-{
-	AnimationSections.Sort([](const UMovieSceneSection& A, const UMovieSceneSection& B) {return ((A).GetTrueRange().GetLowerBoundValue() < (B).GetTrueRange().GetLowerBoundValue());});
-}
-
-//expectation is the weights may be unnormalized.
-static void BlendTheseTransformsByWeight(FTransform& OutTransform, const TArray<FTransform>& Transforms,  TArray<float>& Weights)
-{
-	if (Weights.Num() > 0)
-	{
-		float TotalWeight = 0.0f;
-		for (int32 WeightIndex = 0; WeightIndex < Weights.Num(); ++WeightIndex)
-		{
-			TotalWeight += Weights[WeightIndex];
-		}
-		if (!FMath::IsNearlyEqual(TotalWeight, 1.0f))
-		{
-			for (int32 DivideIndex = 0; DivideIndex < Weights.Num(); ++DivideIndex)
-			{
-				Weights[DivideIndex] /= TotalWeight;
-			}
-		}
-	}
-	int32 NumBlends = Transforms.Num();
-	check(Transforms.Num() == Weights.Num());
-	if (NumBlends == 0)
-	{
-		OutTransform = FTransform::Identity;
-	}
-	else if (NumBlends == 1)
-	{
-		OutTransform = Transforms[0];
-	}
-	else
-	{
-		FVector		OutTranslation(0.0f, 0.0f, 0.0f);
-		FVector		OutScale(0.0f, 0.0f, 0.0f);
-
-		//rotation will get set to the first weighted and then made closest to that so linear interp works.
-		FQuat FirstRot = Transforms[0].GetRotation();
-		FQuat		OutRotation(FirstRot.X * Weights[0], FirstRot.Y * Weights[0], FirstRot.Z * Weights[0], FirstRot.W * Weights[0]);
-
-		for (int32 Index = 0; Index < NumBlends; ++Index)
-		{
-			OutTranslation += Transforms[Index].GetTranslation() * Weights[Index];
-			OutScale +=  Transforms[Index].GetScale3D() * Weights[Index];
-			if (Index != 0)
-			{
-				FQuat Quat = Transforms[Index].GetRotation();
-				Quat.EnforceShortestArcWith(FirstRot);
-				Quat *= Weights[Index];
-				OutRotation += Quat;
-			}
-		}
-
-		OutRotation.Normalize();
-		OutTransform = FTransform(OutRotation, OutTranslation, OutScale);
-	}
-}
-void UMovieSceneSkeletalAnimationTrack::SetRootMotionsDirty()
-{
-	RootMotionParams.bRootMotionsDirty = true;
-}
-
-struct FSkelBoneLength
-{
-	FSkelBoneLength(FCompactPoseBoneIndex InPoseIndex, float InBL) :PoseBoneIndex(InPoseIndex), BoneLength(InBL) {};
-	FCompactPoseBoneIndex PoseBoneIndex;
-	float BoneLength; //squared
-};
-
-static void CalculateDistanceMap(USkeletalMeshComponent* SkelMeshComp, UAnimSequenceBase* FirstAnimSeq, UAnimSequenceBase* SecondAnimSeq, float StartFirstAnimTime, float FrameRate,
-	TArray<TArray<float>>& OutDistanceDifferences)
-{
-
-	int32 FirstAnimNumFrames = (FirstAnimSeq->GetPlayLength() - StartFirstAnimTime) * FrameRate + 1;
-	int32 SecondAnimNumFrames = SecondAnimSeq->GetPlayLength() * FrameRate + 1;
-	OutDistanceDifferences.SetNum(FirstAnimNumFrames);
-	float FirstAnimIndex = 0.0f;
-	float FrameRateDiff = 1.0f / FrameRate;
-	FCompactPose FirstAnimPose, SecondAnimPose;
-	FCSPose<FCompactPose> FirstMeshPoses, SecondMeshPoses;
-	FirstAnimPose.ResetToRefPose(SkelMeshComp->GetAnimInstance()->GetRequiredBones());
-	SecondAnimPose.ResetToRefPose(SkelMeshComp->GetAnimInstance()->GetRequiredBones());
-
-	FBlendedCurve FirstOutCurve, SecondOutCurve;
-	UE::Anim::FStackAttributeContainer FirstTempAttributes, SecondTempAttributes;
-	FAnimationPoseData FirstAnimPoseData(FirstAnimPose, FirstOutCurve, FirstTempAttributes);
-	FAnimationPoseData SecondAnimPoseData(SecondAnimPose, SecondOutCurve, SecondTempAttributes);
-
-	//sort by bone lengths just do the first half
-	//this should avoid us overvalueing to many small values.
-	/*
-	TArray<FSkelBoneLength> BoneLengths;
-	BoneLengths.SetNum(FirstAnimPose.GetNumBones());
-	int32 Index = 0;
-	for (FCompactPoseBoneIndex PoseBoneIndex : FirstAnimPose.ForEachBoneIndex())
-	{
-		FTransform LocalTransform = FirstMeshPoses.GetLocalSpaceTransform(PoseBoneIndex);
-		float BoneLengthVal = LocalTransform.GetLocation().SizeSquared();
-		BoneLengths[Index++] = FSkelBoneLength(PoseBoneIndex, BoneLengthVal);
-	}
-	BoneLengths.Sort([](const FSkelBoneLength& Item1, const FSkelBoneLength& Item2) {
-		return Item1.BoneLength > Item2.BoneLength;
-		});
-		*/
-	FBlendedCurve OutCurve;
-	const FBoneContainer& RequiredBones = FirstAnimPoseData.GetPose().GetBoneContainer();
-	for (TArray<float>& FloatArray : OutDistanceDifferences)
-	{
-		FloatArray.SetNum(SecondAnimNumFrames);
-		float FirstAnimTime = FirstAnimIndex * FrameRateDiff + StartFirstAnimTime;
-		FirstAnimIndex += 1.0f;
-		FAnimExtractContext FirstExtractionContext(static_cast<double>(FirstAnimTime), false);
-		FirstAnimSeq->GetAnimationPose(FirstAnimPoseData, FirstExtractionContext);
-		FirstMeshPoses.InitPose(FirstAnimPoseData.GetPose());
-		float SecondAnimIndex = 0.0f;
-		for (float& DistVal : FloatArray)
-		{
-			DistVal = 0.0f;
-			float SecondAnimTime = SecondAnimIndex * FrameRateDiff;
-			SecondAnimIndex += 1.0f;
-			FAnimExtractContext SecondExtractionContext(static_cast<double>(SecondAnimTime), false);
-			SecondAnimSeq->GetAnimationPose(SecondAnimPoseData, SecondExtractionContext);
-			SecondMeshPoses.InitPose(SecondAnimPoseData.GetPose());
-
-			float DiffVal = 0.0f;
-			for (FCompactPoseBoneIndex PoseBoneIndex : FirstAnimPoseData.GetPose().ForEachBoneIndex())
-			{
-				FTransform FirstTransform = FirstMeshPoses.GetComponentSpaceTransform(PoseBoneIndex);
-				FTransform SecondTransform = SecondMeshPoses.GetComponentSpaceTransform(PoseBoneIndex);
-				if (PoseBoneIndex != 0)
-				{
-					DistVal += (FirstTransform.GetTranslation() - SecondTransform.GetTranslation()).SizeSquared();
-				}
-			}
-		}
-	}
-}
-//outer is startanimtime to firstanim->seqlength...
-//inner is 0 to secondanim->seqlength...
-//for this function just find the smallest in the second...
-//return the end anim time
-static float GetBestBlendPointTimeAtStart(UAnimSequenceBase* FirstAnimSeq, UAnimSequenceBase* SecondAnimSeq, float StartFirstAnimTime, float FrameRate,
-	TArray<TArray<float>>& DistanceDifferences)
-{
-
-	//int32 FirstAnimNumFrames = (FirstAnimSeq->SequenceLength - StartFirstAnimTime) * FrameRate + 1;
-	int32 SecondAnimNumFrames = SecondAnimSeq->GetPlayLength() * FrameRate + 1;
-	if (SecondAnimNumFrames <= 0)
-	{
-		return 0.0f;
-	}
-	TArray<float>& Distances = DistanceDifferences[0];
-	float MinVal = Distances[0];
-	int32 SmallIndex = 0;
-	for (int32 Index = 1; Index < SecondAnimNumFrames; ++Index)
-	{
-		float NewMin = Distances[Index];
-		if (NewMin < MinVal)
-		{
-			MinVal = NewMin;
-			SmallIndex = Index;
-		}
-	}
-	return SmallIndex * 1.0f / FrameRate;
-}
-
-TOptional<FTransform>  UMovieSceneSkeletalAnimationTrack::GetRootMotion(FFrameTime CurrentTime)
-{
-	TOptional<FTransform> Transform;
-
-	UMovieScene* MovieScene = GetTypedOuter<UMovieScene>();
-	if (!MovieScene)
-	{
-		return Transform;
-	}
-	FFrameRate TickResolution = MovieScene->GetTickResolution();
-	
-	if (RootMotionParams.bRootMotionsDirty) 
-	{
-		SetUpRootMotions(true);
-	}
-	if (RootMotionParams.bHaveRootMotion == false)
-	{
-		return Transform;
-	}
-	SortSections();
-	TArray<FTransform> CurrentTransforms;
-	TArray<float> CurrentWeights;
-	TArray<FTransform> CurrentAdditiveTransforms;
-	TArray<float> CurrentAdditiveWeights;
-	FTransform CurrentTransform = FTransform::Identity;
-
-	// Use evaluation field to iterate, as it will take into account 'Evaluate Nearest Section'.
-	for (const FMovieSceneTrackEvaluationFieldEntry& FieldEntry : GetEvaluationField().Entries)
-	{
-		if (FieldEntry.Range.Contains(CurrentTime.FrameNumber))
-		{
-			if (UMovieSceneSkeletalAnimationSection* AnimSection = Cast<UMovieSceneSkeletalAnimationSection>(FieldEntry.Section))
-			{
-				FFrameNumber EvaluationFrame = FieldEntry.ForcedTime == TNumericLimits<int32>::Lowest() ? CurrentTime.FrameNumber : FieldEntry.ForcedTime;
-				if (AnimSection->Params.Animation)
-				{
-					UAnimSequenceBase* ValidAnimSequence = AnimSection->Params.Animation;
-					FMemMark Mark(FMemStack::Get());
-					FCompactPose OutPose;
-					TArray<FBoneIndexType> RequiredBoneIndexArray;
-					RequiredBoneIndexArray.AddUninitialized(ValidAnimSequence->GetSkeleton()->GetReferenceSkeleton().GetNum());
-					for (int32 BoneIndex = 0; BoneIndex < RequiredBoneIndexArray.Num(); ++BoneIndex)
-					{
-						RequiredBoneIndexArray[BoneIndex] = BoneIndex;
-					}
-
-					FBoneContainer BoneContainer(RequiredBoneIndexArray, UE::Anim::FCurveFilterSettings(UE::Anim::ECurveFilterMode::None), *ValidAnimSequence->GetSkeleton());
-					OutPose.ResetToRefPose(BoneContainer);
-					FBlendedCurve OutCurve;
-					OutCurve.InitFrom(BoneContainer);
-					UE::Anim::FStackAttributeContainer TempAttributes;
-					FAnimationPoseData AnimationPoseData(OutPose, OutCurve, TempAttributes);
-					UMovieSceneSkeletalAnimationSection::FRootMotionTransformParam Param;
-					Param.FrameRate = TickResolution;
-					Param.CurrentTime = EvaluationFrame;
-
-					if (AnimSection->GetRootMotionTransform(AnimationPoseData, Param))
-					{
-						if (!Param.bOutIsAdditive)
-						{
-							CurrentTransform = Param.OutTransform * AnimSection->PreviousTransform;
-							CurrentTransforms.Add(CurrentTransform);
-							CurrentWeights.Add(Param.OutWeight);
-						}
-						else
-						{
-							CurrentAdditiveTransforms.Add(Param.OutTransform);
-							CurrentAdditiveWeights.Add(Param.OutWeight);
-						}
-					}
-				}
-			}
-		}
-	}
-
-	BlendTheseTransformsByWeight(CurrentTransform, CurrentTransforms, CurrentWeights);
-	//now handle additive onto the current
-	if (CurrentAdditiveWeights.Num() > 0)
-	{
-		FTransform AdditiveTransform;
-		BlendTheseTransformsByWeight(AdditiveTransform, CurrentAdditiveTransforms, CurrentAdditiveWeights);
-		const ScalarRegister VBlendWeight(1.0f);
-		FTransform::BlendFromIdentityAndAccumulate(CurrentTransform, AdditiveTransform, VBlendWeight);
-	}
-	Transform = CurrentTransform;
-	return Transform;
-
-}
-
-void UMovieSceneSkeletalAnimationTrack::SetUpRootMotions(bool bForce)
-{
-	UMovieScene* MovieScene = GetTypedOuter<UMovieScene>();
-	if (!MovieScene)
-	{
-		return;
-	}
-	
-	if (bForce || RootMotionParams.bRootMotionsDirty)
-	{
-		RootMotionParams.bRootMotionsDirty = false;
-		RootMotionParams.bHaveRootMotion = false;
-		const FFrameRate MinDisplayRate(60, 1);
-
-		FFrameRate DisplayRate =  MovieScene->GetDisplayRate().AsDecimal() < MinDisplayRate.AsDecimal() ? MinDisplayRate : MovieScene->GetDisplayRate();
-		FFrameRate TickResolution = MovieScene->GetTickResolution();
-		FFrameTime FrameTick = FFrameTime(FMath::Max(1, TickResolution.AsFrameNumber(1.0).Value / DisplayRate.AsFrameNumber(1.0).Value));
-		if (AnimationSections.Num() == 0 || FrameTick.FrameNumber.Value == 0)
-		{
-#if WITH_EDITORONLY_DATA
-			RootMotionParams.RootTransforms.SetNum(0);
-#endif
-			return;
-		}
-		
-		SortSections();
-		//Set the TempOffset.
-		FTransform InitialTransform = FTransform::Identity;
-		UMovieSceneSkeletalAnimationSection* PrevAnimSection = nullptr;
-		//valid anim sequence to use to calculate bones.
-		UAnimSequenceBase* ValidAnimSequence = nullptr;
-		//if no transforms have offsets then don't do root motion caching.
-		bool bAnySectionsHaveOffset = false;
-		for (UMovieSceneSection* Section : AnimationSections)
-		{
-			UMovieSceneSkeletalAnimationSection* AnimSection = Cast<UMovieSceneSkeletalAnimationSection>(Section);
-			if (AnimSection)
-			{
-				if (AnimSection->StartLocationOffset.IsNearlyZero() == false || AnimSection->StartRotationOffset.IsNearlyZero() == false ||
-					AnimSection->MatchedLocationOffset.IsNearlyZero() == false || AnimSection->MatchedRotationOffset.IsNearlyZero() == false)
-				{
-					bAnySectionsHaveOffset = true;
-				}
-				if (ValidAnimSequence == nullptr)
-				{
-					ValidAnimSequence = AnimSection->Params.Animation;
-				}
-				if (PrevAnimSection)
-				{
-					if (UAnimSequenceBase* PrevAnimSequence = PrevAnimSection->Params.Animation)
-					{	
-						if (bAnySectionsHaveOffset)
-						{ 
-							RootMotionParams.RootMotionStartOffset = AnimSection->GetRootMotionStartOffset();
-							FMemMark Mark(FMemStack::Get());
-							FCompactPose OutPose;
-							TArray<FBoneIndexType> RequiredBoneIndexArray;
-							RequiredBoneIndexArray.AddUninitialized(PrevAnimSequence->GetSkeleton()->GetReferenceSkeleton().GetNum());
-							for (int32 BoneIndex = 0; BoneIndex < RequiredBoneIndexArray.Num(); ++BoneIndex)
-							{
-								RequiredBoneIndexArray[BoneIndex] = BoneIndex;
-							}
-
-							FBoneContainer BoneContainer(RequiredBoneIndexArray, UE::Anim::FCurveFilterSettings(UE::Anim::ECurveFilterMode::None), *PrevAnimSequence->GetSkeleton());
-							OutPose.ResetToRefPose(BoneContainer);
-							FBlendedCurve OutCurve;
-							OutCurve.InitFrom(BoneContainer);
-
-							UE::Anim::FStackAttributeContainer TempAttributes;
-							FAnimationPoseData AnimationPoseData(OutPose, OutCurve, TempAttributes);
-
-							UMovieSceneSkeletalAnimationSection::FRootMotionTransformParam Param;
-							Param.FrameRate = MovieScene->GetTickResolution();
-							Param.CurrentTime = AnimSection->GetRange().GetLowerBoundValue();
-
-							if (PrevAnimSection->GetRootMotionTransform(AnimationPoseData, Param))
-							{
-								AnimSection->PreviousTransform = Param.OutPoseTransform.GetRelativeTransformReverse(Param.OutTransform);
-								AnimSection->PreviousTransform = AnimSection->PreviousTransform * InitialTransform;
-
-							}
-							else
-							{
-								AnimSection->PreviousTransform = FTransform::Identity;
-							}
-						}
-						else
-						{
-							AnimSection->PreviousTransform = FTransform::Identity;
-						}
-					}
-					else
-					{
-						AnimSection->PreviousTransform = FTransform::Identity;
-					}
-					InitialTransform = AnimSection->PreviousTransform;
-
-				}
-				else
-				{
-					AnimSection->PreviousTransform = FTransform::Identity;
-				}
-				PrevAnimSection = AnimSection;
-				AnimSection->SetBoneIndexForRootMotionCalculations(bBlendFirstChildOfRoot);
-			}
-		}
-
-		//if we are swapping root bone turn on root motion matching anyway
-		if (bAnySectionsHaveOffset == false && SwapRootBone == ESwapRootBone::SwapRootBone_None)
-		{
-#if WITH_EDITORONLY_DATA
-			RootMotionParams.RootTransforms.SetNum(0);
-#endif			
-			return;
-		}
-
-		RootMotionParams.bHaveRootMotion = true;
-		RootMotionParams.StartFrame = AnimationSections[0]->GetInclusiveStartFrame();
-		RootMotionParams.EndFrame = AnimationSections[AnimationSections.Num() - 1]->GetExclusiveEndFrame() - 1;
-		RootMotionParams.FrameTick = FrameTick;
-		
-		
-#if WITH_EDITORONLY_DATA
-		if (RootMotionParams.bCacheRootTransforms == false)
-		{
-			return;
-		}
-		//set up pose from valid anim sequences.
-		FMemMark Mark(FMemStack::Get());
-		FCompactPose OutPose;
-
-		if (ValidAnimSequence)
-		{
-			TArray<FBoneIndexType> RequiredBoneIndexArray;
-			const UE::Anim::FCurveFilterSettings CurveFilterSettings;
-			RequiredBoneIndexArray.AddUninitialized(ValidAnimSequence->GetSkeleton()->GetReferenceSkeleton().GetNum());
-			for (int32 BoneIndex = 0; BoneIndex < RequiredBoneIndexArray.Num(); ++BoneIndex)
-			{
-				RequiredBoneIndexArray[BoneIndex] = BoneIndex;
-			}
-
-			FBoneContainer BoneContainer(RequiredBoneIndexArray, CurveFilterSettings, *ValidAnimSequence->GetSkeleton());
-			OutPose.ResetToRefPose(BoneContainer);
-			FBlendedCurve OutCurve;
-			OutCurve.InitFrom(BoneContainer);
-			TArray< UMovieSceneSkeletalAnimationSection*> SectionsAtCurrentTime;
-			int32 NumTotal = (RootMotionParams.EndFrame.FrameNumber.Value - RootMotionParams.StartFrame.FrameNumber.Value) / (RootMotionParams.FrameTick.FrameNumber.Value) + 1;
-			RootMotionParams.RootTransforms.SetNum(NumTotal);
-
-			TArray<FTransform> CurrentTransforms;
-			TArray<float> CurrentWeights;
-			TArray<FTransform> CurrentAdditiveTransforms;
-			TArray<float> CurrentAdditiveWeights;
-			FFrameTime PreviousFrame = RootMotionParams.StartFrame;
-			int32 Index = 0;
-			for (FFrameTime FrameNumber = RootMotionParams.StartFrame; FrameNumber <= RootMotionParams.EndFrame; FrameNumber += RootMotionParams.FrameTick)
-			{
-				CurrentTransforms.SetNum(0);
-				CurrentWeights.SetNum(0);
-				FTransform CurrentTransform(FTransform::Identity), ParentTransform(FTransform::Identity);
-
-				UMovieSceneSkeletalAnimationSection* PrevSection = nullptr;
-				for (UMovieSceneSection* Section : AnimationSections)
-				{
-					if (Section && Section->GetRange().Contains(FrameNumber.FrameNumber))
-					{
-						UMovieSceneSkeletalAnimationSection* AnimSection = CastChecked<UMovieSceneSkeletalAnimationSection>(Section);
-
-						UE::Anim::FStackAttributeContainer TempAttributes;
-						FAnimationPoseData AnimationPoseData(OutPose, OutCurve, TempAttributes);
-						UMovieSceneSkeletalAnimationSection::FRootMotionTransformParam Param;
-						Param.FrameRate = TickResolution;
-						Param.CurrentTime = FrameNumber.FrameNumber;
-
-						if (AnimSection->GetRootMotionTransform(AnimationPoseData,Param))
-						{
-							if (!Param.bOutIsAdditive)
-							{
-								CurrentTransform = Param.OutTransform * AnimSection->PreviousTransform;
-								CurrentTransforms.Add(CurrentTransform);
-								CurrentWeights.Add(Param.OutWeight);
-							}
-							else
-							{
-								CurrentAdditiveTransforms.Add(Param.OutTransform);
-								CurrentAdditiveWeights.Add(Param.OutWeight);
-							}
-						}
-						PrevSection = AnimSection;
-					}
-				}
-
-				BlendTheseTransformsByWeight(CurrentTransform, CurrentTransforms, CurrentWeights);
-				//now handle additive onto the current
-				if (CurrentAdditiveWeights.Num() > 0)
-				{
-					FTransform AdditiveTransform;
-					BlendTheseTransformsByWeight(AdditiveTransform, CurrentAdditiveTransforms, CurrentAdditiveWeights);
-					const ScalarRegister VBlendWeight(1.0f);
-					FTransform::BlendFromIdentityAndAccumulate(CurrentTransform, AdditiveTransform, VBlendWeight);
-				}
-				RootMotionParams.RootTransforms[Index] = CurrentTransform;
-				++Index;
-				PreviousFrame = FrameNumber;
-
-			}
-		}
-		else //no valid anim sequence just clear out
-		{
-			RootMotionParams.RootTransforms.SetNum(0);
-		}
-#endif
-	}
-
-}
-
-static FTransform GetTransformForBoneRelativeToIndex(UAnimSequence* AnimSequence, USkeletalMeshComponent* MeshComponent, const FName& InBoneName,
-	const FCompactPoseBoneIndex& ParentCPIndex, double Seconds)
-{
-	FTransform  WorldTransform = FTransform::Identity;
-	//AnimSequence->GetBoneTransform doesn't seem to be as accurate as GetAnimationPose
-	FMemMark Mark(FMemStack::Get());
-	FCompactPose OutPose;
-	TArray<FBoneIndexType> RequiredBoneIndexArray;
-	const UE::Anim::FCurveFilterSettings CurveFilterSettings;
-	RequiredBoneIndexArray.AddUninitialized(AnimSequence->GetSkeleton()->GetReferenceSkeleton().GetNum());
-	for (int32 BoneIndex = 0; BoneIndex < RequiredBoneIndexArray.Num(); ++BoneIndex)
-	{
-		RequiredBoneIndexArray[BoneIndex] = BoneIndex;
-	}
-
-	FBoneContainer BoneContainer(RequiredBoneIndexArray, CurveFilterSettings, *AnimSequence->GetSkeleton());
-	OutPose.ResetToRefPose(BoneContainer);
-	FBlendedCurve OutCurve;
-	OutCurve.InitFrom(BoneContainer);
-	UE::Anim::FStackAttributeContainer TempAttributes;
-	FAnimationPoseData AnimationPoseData(OutPose, OutCurve, TempAttributes);
-	FAnimExtractContext ExtractionContext(Seconds, false);
-#if WITH_EDITOR
-	ExtractionContext.bIgnoreRootLock = true;
-#endif
-	AnimSequence->GetAnimationPose(AnimationPoseData, ExtractionContext);
-	int32 MeshIndex = AnimationPoseData.GetPose().GetBoneContainer().GetPoseBoneIndexForBoneName(InBoneName);
-	if (MeshIndex != INDEX_NONE)
-	{
-		FCompactPoseBoneIndex CPIndex = AnimationPoseData.GetPose().GetBoneContainer().MakeCompactPoseIndex(FMeshPoseBoneIndex(MeshIndex));
-		if (CPIndex != INDEX_NONE )
-		{
-			FTransform BoneTransform = FTransform::Identity;
-			WorldTransform *= BoneTransform;
-			do
-			{
-				BoneTransform = AnimationPoseData.GetPose()[CPIndex];
-				WorldTransform *= BoneTransform;
-				if (CPIndex == ParentCPIndex)  //if we are the parent then we stop
-				{
-					CPIndex = FCompactPoseBoneIndex(INDEX_NONE);
-				}
-				else
-				{
-					CPIndex = AnimationPoseData.GetPose().GetBoneContainer().GetParentBoneIndex(CPIndex);
-				}
-			} while (CPIndex.IsValid());
-		}
-	}
-	return WorldTransform;
-}
-
-enum class ESkelAnimRotationOrder 
-{
-	XYZ,
-	XZY,
-	YXZ,
-	YZX,
-	ZXY,
-	ZYX
-};
-
-static FQuat QuatFromEuler(const FVector& XYZAnglesInDegrees, ESkelAnimRotationOrder RotationOrder)
-{
-	float X = FMath::DegreesToRadians(XYZAnglesInDegrees.X);
-	float Y = FMath::DegreesToRadians(XYZAnglesInDegrees.Y);
-	float Z = FMath::DegreesToRadians(XYZAnglesInDegrees.Z);
-
-	float CosX = FMath::Cos(X * 0.5f);
-	float CosY = FMath::Cos(Y * 0.5f);
-	float CosZ = FMath::Cos(Z * 0.5f);
-
-	float SinX = FMath::Sin(X * 0.5f);
-	float SinY = FMath::Sin(Y * 0.5f);
-	float SinZ = FMath::Sin(Z * 0.5f);
-
-	if (RotationOrder == ESkelAnimRotationOrder::XYZ)
-	{
-		return FQuat(SinX * CosY * CosZ - CosX * SinY * SinZ,
-			CosX * SinY * CosZ + SinX * CosY * SinZ,
-			CosX * CosY * SinZ - SinX * SinY * CosZ,
-			CosX * CosY * CosZ + SinX * SinY * SinZ);
-
-	}
-	else if (RotationOrder == ESkelAnimRotationOrder::XZY)
-	{
-		return FQuat(SinX * CosY * CosZ + CosX * SinY * SinZ,
-			CosX * SinY * CosZ + SinX * CosY * SinZ,
-			CosX * CosY * SinZ - SinX * SinY * CosZ,
-			CosX * CosY * CosZ - SinX * SinY * SinZ);
-
-	}
-	else if (RotationOrder == ESkelAnimRotationOrder::YXZ)
-	{
-		return FQuat(SinX * CosY * CosZ - CosX * SinY * SinZ,
-			CosX * SinY * CosZ + SinX * CosY * SinZ,
-			CosX * CosY * SinZ + SinX * SinY * CosZ,
-			CosX * CosY * CosZ - SinX * SinY * SinZ);
-
-	}
-	else if (RotationOrder == ESkelAnimRotationOrder::YZX)
-	{
-		return FQuat(SinX * CosY * CosZ - CosX * SinY * SinZ,
-			CosX * SinY * CosZ - SinX * CosY * SinZ,
-			CosX * CosY * SinZ + SinX * SinY * CosZ,
-			CosX * CosY * CosZ + SinX * SinY * SinZ);
-	}
-	else if (RotationOrder == ESkelAnimRotationOrder::ZXY)
-	{
-		return FQuat(SinX * CosY * CosZ + CosX * SinY * SinZ,
-			CosX * SinY * CosZ - SinX * CosY * SinZ,
-			CosX * CosY * SinZ - SinX * SinY * CosZ,
-			CosX * CosY * CosZ + SinX * SinY * SinZ);
-
-	}
-	else if (RotationOrder == ESkelAnimRotationOrder::ZYX)
-	{
-		return FQuat(SinX * CosY * CosZ + CosX * SinY * SinZ,
-			CosX * SinY * CosZ - SinX * CosY * SinZ,
-			CosX * CosY * SinZ + SinX * SinY * CosZ,
-			CosX * CosY * CosZ - SinX * SinY * SinZ);
-
-	}
-
-	// should not happen
-	return FQuat::Identity;
-}
-
-static FVector EulerFromQuat(const FQuat& Rotation, ESkelAnimRotationOrder RotationOrder)
-{
-	float X = Rotation.X;
-	float Y = Rotation.Y;
-	float Z = Rotation.Z;
-	float W = Rotation.W;
-	float X2 = X * 2.f;
-	float Y2 = Y * 2.f;
-	float Z2 = Z * 2.f;
-	float XX2 = X * X2;
-	float XY2 = X * Y2;
-	float XZ2 = X * Z2;
-	float YX2 = Y * X2;
-	float YY2 = Y * Y2;
-	float YZ2 = Y * Z2;
-	float ZX2 = Z * X2;
-	float ZY2 = Z * Y2;
-	float ZZ2 = Z * Z2;
-	float WX2 = W * X2;
-	float WY2 = W * Y2;
-	float WZ2 = W * Z2;
-
-	FVector AxisX, AxisY, AxisZ;
-	AxisX.X = (1.f - (YY2 + ZZ2));
-	AxisY.X = (XY2 + WZ2);
-	AxisZ.X = (XZ2 - WY2);
-	AxisX.Y = (XY2 - WZ2);
-	AxisY.Y = (1.f - (XX2 + ZZ2));
-	AxisZ.Y = (YZ2 + WX2);
-	AxisX.Z = (XZ2 + WY2);
-	AxisY.Z = (YZ2 - WX2);
-	AxisZ.Z = (1.f - (XX2 + YY2));
-
-	FVector Result = FVector::ZeroVector;
-
-	if (RotationOrder == ESkelAnimRotationOrder::XYZ)
-	{
-		Result.Y = FMath::Asin(-FMath::Clamp<float>(AxisZ.X, -1.f, 1.f));
-
-		if (FMath::Abs(AxisZ.X) < 1.f - SMALL_NUMBER)
-		{
-			Result.X = FMath::Atan2(AxisZ.Y, AxisZ.Z);
-			Result.Z = FMath::Atan2(AxisY.X, AxisX.X);
-		}
-		else
-		{
-			Result.X = 0.f;
-			Result.Z = FMath::Atan2(-AxisX.Y, AxisY.Y);
-		}
-	}
-	else if (RotationOrder == ESkelAnimRotationOrder::XZY)
-	{
-
-		Result.Z = FMath::Asin(FMath::Clamp<float>(AxisY.X, -1.f, 1.f));
-
-		if (FMath::Abs(AxisY.X) < 1.f - SMALL_NUMBER)
-		{
-			Result.X = FMath::Atan2(-AxisY.Z, AxisY.Y);
-			Result.Y = FMath::Atan2(-AxisZ.X, AxisX.X);
-		}
-		else
-		{
-			Result.X = 0.f;
-			Result.Y = FMath::Atan2(AxisX.Z, AxisZ.Z);
-		}
-	}
-	else if (RotationOrder == ESkelAnimRotationOrder::YXZ)
-	{
-		Result.X = FMath::Asin(FMath::Clamp<float>(AxisZ.Y, -1.f, 1.f));
-
-		if (FMath::Abs(AxisZ.Y) < 1.f - SMALL_NUMBER)
-		{
-			Result.Y = FMath::Atan2(-AxisZ.X, AxisZ.Z);
-			Result.Z = FMath::Atan2(-AxisX.Y, AxisY.Y);
-		}
-		else
-		{
-			Result.Y = 0.f;
-			Result.Z = FMath::Atan2(AxisY.X, AxisX.X);
-		}
-	}
-	else if (RotationOrder == ESkelAnimRotationOrder::YZX)
-	{
-		Result.Z = FMath::Asin(-FMath::Clamp<float>(AxisX.Y, -1.f, 1.f));
-
-		if (FMath::Abs(AxisX.Y) < 1.f - SMALL_NUMBER)
-		{
-			Result.X = FMath::Atan2(AxisZ.Y, AxisY.Y);
-			Result.Y = FMath::Atan2(AxisX.Z, AxisX.X);
-		}
-		else
-		{
-			Result.X = FMath::Atan2(-AxisY.Z, AxisZ.Z);
-			Result.Y = 0.f;
-		}
-	}
-	else if (RotationOrder == ESkelAnimRotationOrder::ZXY)
-	{
-		Result.X = FMath::Asin(-FMath::Clamp<float>(AxisY.Z, -1.f, 1.f));
-
-		if (FMath::Abs(AxisY.Z) < 1.f - SMALL_NUMBER)
-		{
-			Result.Y = FMath::Atan2(AxisX.Z, AxisZ.Z);
-			Result.Z = FMath::Atan2(AxisY.X, AxisY.Y);
-		}
-		else
-		{
-			Result.Y = FMath::Atan2(-AxisZ.X, AxisX.X);
-			Result.Z = 0.f;
-		}
-	}
-	else if (RotationOrder == ESkelAnimRotationOrder::ZYX)
-	{
-		Result.Y = FMath::Asin(FMath::Clamp<float>(AxisX.Z, -1.f, 1.f));
-
-		if (FMath::Abs(AxisX.Z) < 1.f - SMALL_NUMBER)
-		{
-			Result.X = FMath::Atan2(-AxisY.Z, AxisZ.Z);
-			Result.Z = FMath::Atan2(-AxisX.Y, AxisX.X);
-		}
-		else
-		{
-			Result.X = FMath::Atan2(AxisZ.Y, AxisY.Y);
-			Result.Z = 0.f;
-		}
-	}
-
-	return Result * 180.f / PI;
-}
-/**
-*  Function to find best rotation order given how we are matching the rotations. 
-*  If it is matched we need to make sure it happens first
-*  Issue is Yaw is most common match but by default FRotation:FQuat conversions it's last, this causes issues.
-*/
-static ESkelAnimRotationOrder FindBestRotationOrder(bool bMatchRoll, bool bMatchPitch, bool bMatchYaw)
-{
-	if (bMatchYaw)
-	{
-		return ESkelAnimRotationOrder::YXZ;
-	}
-	if (bMatchPitch)
-	{
-		return ESkelAnimRotationOrder::YZX;
-	}
-	return ESkelAnimRotationOrder::XYZ;
-}
-
-void UMovieSceneSkeletalAnimationTrack::MatchSectionByBoneTransform(bool bMatchWithPrevious, USkeletalMeshComponent* SkelMeshComp, UMovieSceneSkeletalAnimationSection* CurrentSection, FFrameTime CurrentFrame, FFrameRate FrameRate,
-	const FName& BoneName, FTransform& SecondSectionRootDiff, FVector& TranslationDiff, FQuat& RotationDiff) //add options for z and for rotation.
-{
-	SortSections();
-	UMovieSceneSection* PrevSection = nullptr;
-	UMovieSceneSection* NextSection = nullptr;
-	for (int32 Index = 0; Index < AnimationSections.Num(); ++Index)
-	{
-		UMovieSceneSection* Section = AnimationSections[Index];
-		if (Section == CurrentSection)
-		{
-			if (++Index < AnimationSections.Num())
-			{
-				NextSection = AnimationSections[Index];
-			}
-			break;
-		}
-		PrevSection = Section;
-	}
-
-	TranslationDiff = FVector(0.0f, 0.0f, 0.0f);
-	RotationDiff = FQuat::Identity;
-	SecondSectionRootDiff = FTransform::Identity;
-
-	if (bMatchWithPrevious && PrevSection)
-	{
-		UMovieSceneSkeletalAnimationSection* FirstSection = Cast<UMovieSceneSkeletalAnimationSection>(PrevSection);
-		UAnimSequence* FirstAnimSequence = Cast<UAnimSequence>(FirstSection->Params.Animation);
-		UAnimSequence* SecondAnimSequence = Cast<UAnimSequence>(CurrentSection->Params.Animation);
-
-		if (FirstAnimSequence && SecondAnimSequence)
-		{
-			double FirstSectionTime = FirstSection->MapTimeToAnimation(CurrentFrame, FrameRate);
-			//use same index for all
-			int32 Index = CurrentSection->SetBoneIndexForRootMotionCalculations(bBlendFirstChildOfRoot);
-			FCompactPoseBoneIndex ParentIndex(Index);
-			
-			FTransform  FirstTransform = GetTransformForBoneRelativeToIndex(FirstAnimSequence, SkelMeshComp, BoneName, ParentIndex, FirstSectionTime);
-			double SecondSectionTime = CurrentSection->MapTimeToAnimation(CurrentFrame, FrameRate);
-			FTransform  SecondTransform = GetTransformForBoneRelativeToIndex(SecondAnimSequence, SkelMeshComp, BoneName, ParentIndex,SecondSectionTime);
-			//Need to match the translations and rotations here 
-			//First need to get the correct rotation order based upon what's matching, otherwise if not all are matched 
-			//and one rotation is set last we will get errors.
-			ESkelAnimRotationOrder RotationOrder = FindBestRotationOrder(CurrentSection->bMatchRotationRoll, CurrentSection->bMatchRotationPitch, CurrentSection->bMatchRotationYaw);
-
-			FVector FirstTransformTranslation = FirstTransform.GetTranslation();
-			FVector SecondTransformTranslation = SecondTransform.GetTranslation();
-			FQuat FirstTransformQuat = FirstTransform.GetRotation();
-			FQuat SecondTransformQuat = SecondTransform.GetRotation();
-
-			FirstTransformQuat.EnforceShortestArcWith(SecondTransformQuat);
-
-			FRotator FirstTransformRotation(FRotator::MakeFromEuler(EulerFromQuat(FirstTransformQuat, RotationOrder)));
-			FRotator SecondTransformRotation(FRotator::MakeFromEuler(EulerFromQuat(SecondTransformQuat, RotationOrder)));
-			SecondTransformRotation.SetClosestToMe(FirstTransformRotation);
-
-			if (!CurrentSection->bMatchTranslation)
-			{
-				FirstTransformTranslation.X = SecondTransformTranslation.X;
-				FirstTransformTranslation.Y = SecondTransformTranslation.Y;
-				FirstTransformTranslation.Z = SecondTransformTranslation.Z;
-			}
-			if (!CurrentSection->bMatchIncludeZHeight)
-			{
-				FirstTransformTranslation.Z = SecondTransformTranslation.Z;
-			}
-			FirstTransform.SetTranslation(FirstTransformTranslation);
-
-			if (!CurrentSection->bMatchRotationYaw)
-			{
-				FirstTransformRotation.Yaw = SecondTransformRotation.Yaw;
-			}
-			if (!CurrentSection->bMatchRotationPitch)
-			{
-				FirstTransformRotation.Pitch = SecondTransformRotation.Pitch;
-			}
-			if (!CurrentSection->bMatchRotationRoll)
-			{
-				FirstTransformRotation.Roll = SecondTransformRotation.Roll;
-			}
-			FirstTransformQuat = QuatFromEuler(FirstTransformRotation.Euler(), RotationOrder);
-			FirstTransform.SetRotation(FirstTransformQuat);
-
-			// Below is the match but we need to use GetRelativeTransformReverse since Inverse doesn't work as expected.
-			//	* GetRelativeTransformReverse returns this(-1)* Other, and parameter is Other.
-			SecondSectionRootDiff = SecondTransform.GetRelativeTransformReverse(FirstTransform);
-			TranslationDiff = SecondSectionRootDiff.GetTranslation();
-			RotationDiff = SecondSectionRootDiff.GetRotation();
-
-		}
-	}
-	else if (bMatchWithPrevious == false && NextSection) //match with next
-	{
-		UMovieSceneSkeletalAnimationSection* SecondSection = Cast<UMovieSceneSkeletalAnimationSection>(NextSection);
-		UAnimSequence* FirstAnimSequence = Cast<UAnimSequence>(CurrentSection->Params.Animation);
-		UAnimSequence* SecondAnimSequence = Cast<UAnimSequence>(SecondSection->Params.Animation);
-
-		if (FirstAnimSequence && SecondAnimSequence)
-		{
-			//use same index for all
-			int32 Index = CurrentSection->SetBoneIndexForRootMotionCalculations(bBlendFirstChildOfRoot);
-			FCompactPoseBoneIndex ParentIndex(Index);
-			float FirstSectionTime = static_cast<float>(CurrentSection->MapTimeToAnimation(CurrentFrame, FrameRate));
-			FTransform  FirstTransform = GetTransformForBoneRelativeToIndex(FirstAnimSequence, SkelMeshComp, BoneName,ParentIndex, FirstSectionTime);
-			float SecondSectionTime = static_cast<float>(SecondSection->MapTimeToAnimation(CurrentFrame, FrameRate));
-			FTransform  SecondTransform = GetTransformForBoneRelativeToIndex(SecondAnimSequence, SkelMeshComp, BoneName,ParentIndex, SecondSectionTime);
-
-			//Need to match the translations and rotations here 
-			//First need to get the correct rotation order based upon what's matching, otherwise if not all are matched 
-			//and one rotation is set last we will get errors.
-			ESkelAnimRotationOrder RotationOrder = FindBestRotationOrder(CurrentSection->bMatchRotationRoll, CurrentSection->bMatchRotationPitch, CurrentSection->bMatchRotationYaw);
-
-			FVector FirstTransformTranslation = FirstTransform.GetTranslation();
-			FVector SecondTransformTranslation = SecondTransform.GetTranslation();
-			FQuat FirstTransformQuat = FirstTransform.GetRotation();
-			FQuat SecondTransformQuat = SecondTransform.GetRotation();
-
-			SecondTransformQuat.EnforceShortestArcWith(FirstTransformQuat);
-
-			FRotator FirstTransformRotation(FRotator::MakeFromEuler(EulerFromQuat(FirstTransformQuat, RotationOrder)));
-			FRotator SecondTransformRotation(FRotator::MakeFromEuler(EulerFromQuat(SecondTransformQuat, RotationOrder)));
-			FirstTransformRotation.SetClosestToMe(SecondTransformRotation);
-
-			if (!CurrentSection->bMatchTranslation)
-			{
-				SecondTransformTranslation.X = FirstTransformTranslation.X;
-				SecondTransformTranslation.Y = FirstTransformTranslation.Y;
-				SecondTransformTranslation.Z = FirstTransformTranslation.Z;
-			}
-			if (!CurrentSection->bMatchIncludeZHeight)
-			{
-				SecondTransformTranslation.Z = FirstTransformTranslation.Z;
-			}
-			SecondTransform.SetTranslation(SecondTransformTranslation);
-
-			if (!CurrentSection->bMatchRotationYaw)
-			{
-				SecondTransformRotation.Yaw = FirstTransformRotation.Yaw;
-			}
-			if (!CurrentSection->bMatchRotationPitch)
-			{
-				SecondTransformRotation.Pitch = FirstTransformRotation.Pitch;
-			}
-			if (!CurrentSection->bMatchRotationRoll)
-			{
-				SecondTransformRotation.Roll = FirstTransformRotation.Roll;
-			}
-			SecondTransformQuat = QuatFromEuler(SecondTransformRotation.Euler(), RotationOrder);
-			SecondTransform.SetRotation(SecondTransformQuat);
-
-			//GetRelativeTransformReverse returns this(-1)* Other, and parameter is Other.
-			SecondSectionRootDiff = FirstTransform.GetRelativeTransformReverse(SecondTransform);
-			TranslationDiff = SecondSectionRootDiff.GetTranslation();
-			RotationDiff = SecondSectionRootDiff.GetRotation();
-
-		}
-	}
-}
-
-#if WITH_EDITORONLY_DATA
-
-void UMovieSceneSkeletalAnimationTrack::ToggleShowRootMotionTrail()
-{
-	bShowRootMotionTrail = bShowRootMotionTrail ? false : true;
-}
-#endif
-
-
-//MZ To Do need way to get passed the skelmeshcomp when we add or move a section.
-void UMovieSceneSkeletalAnimationTrack::AutoMatchSectionRoot(UMovieSceneSkeletalAnimationSection* CurrentSection)
-{
-	return;
-#if 0
-	UMovieScene* MovieScene = GetTypedOuter<UMovieScene>();
-	if (AnimationSections.Num() > 0 && MovieScene && CurrentSection)
-	{
-		SortSections();
-
-		for (int32 Index = 0; Index < AnimationSections.Num(); ++Index)
-		{
-			UMovieSceneSection* Section = AnimationSections[Index];
-			if (Section && Section == CurrentSection)
-			{
-				CurrentSection->bMatchWithPrevious = (Index == 0) ? false : true;
-				FFrameTime FrameTime = (Index == 0) ? CurrentSection->GetRange().GetUpperBoundValue() : CurrentSection->GetRange().GetLowerBoundValue();
-				USkeletalMeshComponent* SkelMeshComp = nullptr;
-				CurrentSection->MatchSectionByBoneTransform(SkelMeshComp, FrameTime, MovieScene->GetTickResolution(), CurrentSection->MatchedBoneName);
-			}
-		}
-	}
-#endif
-}
-
 #undef LOCTEXT_NAMESPACE
 
diff --git a/Engine/Source/Runtime/MovieSceneTracks/Public/Systems/MovieSceneTransformOriginSystem.h b/Engine/Source/Runtime/MovieSceneTracks/Public/Systems/MovieSceneTransformOriginSystem.h
--- a/Engine/Source/Runtime/MovieSceneTracks/Public/Systems/MovieSceneTransformOriginSystem.h
+++ b/Engine/Source/Runtime/MovieSceneTracks/Public/Systems/MovieSceneTransformOriginSystem.h
@@ -64,6 +64,10 @@ public:
 	UMovieSceneTransformOriginSystem(const FObjectInitializer& ObjInit);
 
 public:
+
+
+	MOVIESCENETRACKS_API bool GetTransformOrigin(UE::MovieScene::FInstanceHandle InstanceHandle, FTransform& OutTransform) const;
+
 	const TSparseArray<FTransform>& GetTransformOriginsByInstanceID() const { return TransformOriginsByInstanceID; }
 
 
diff --git a/Engine/Source/Runtime/MovieSceneTracks/Public/Tracks/MovieSceneCommonAnimationTrack.h b/Engine/Source/Runtime/MovieSceneTracks/Public/Tracks/MovieSceneCommonAnimationTrack.h
--- a/Engine/Source/Runtime/MovieSceneTracks/Public/Tracks/MovieSceneCommonAnimationTrack.h
+++ b/Engine/Source/Runtime/MovieSceneTracks/Public/Tracks/MovieSceneCommonAnimationTrack.h
@@ -9,6 +9,31 @@
 
 class UAnimSequenceBase;
 class UMovieSceneSection;
+class UMovieSceneSkeletalAnimationSection;
+
+/**Struct to hold the cached root motion positions based upon how we calculated them.
+* Also provides way to get the root motion at a particular time.
+*/
+
+USTRUCT()
+struct FMovieSceneSkeletalAnimRootMotionTrackParams
+{
+	GENERATED_BODY()
+
+	FFrameTime FrameTick; 
+	FFrameTime StartFrame;
+	FFrameTime EndFrame;
+	bool bRootMotionsDirty;
+	bool bHaveRootMotion;
+	FTransform RootMotionStartOffset; // root motion may not be in Mesh Space if we are putting values on a bone that is a child of a root with an offset
+	/** Get the Root Motion transform at the specified time.*/
+	FMovieSceneSkeletalAnimRootMotionTrackParams() : bRootMotionsDirty(true),bHaveRootMotion(false) {}
+
+#if WITH_EDITORONLY_DATA
+	bool bCacheRootTransforms = false;
+	TArray<FTransform> RootTransforms;
+#endif
+};
 
 UCLASS()
 class MOVIESCENETRACKS_API UMovieSceneCommonAnimationTrack
@@ -18,6 +43,8 @@ class MOVIESCENETRACKS_API UMovieSceneCommonAnimationTrack
 
 public:
 
+	UMovieSceneCommonAnimationTrack(const FObjectInitializer& ObjInit);
+
 	/** Adds a new animation to this track */
 	virtual UMovieSceneSection* AddNewAnimationOnRow(FFrameNumber KeyTime, UAnimSequenceBase* AnimSequence, int32 RowIndex);
 
@@ -34,8 +61,58 @@ public:
 	virtual void RemoveSection(UMovieSceneSection& Section) override;
 	virtual void RemoveSectionAt(int32 SectionIndex) override;
 	virtual bool IsEmpty() const override;
+	virtual void UpdateEasing() override;
+
+#if WITH_EDITOR
+	virtual void PostEditImport() override;
+	virtual void PostEditUndo() override;
+	virtual void PostEditChangeProperty(FPropertyChangedEvent& PropertyChangedEvent) override;
+#endif
+
+	virtual bool ShouldUseRootMotions() const
+	{
+		return false;
+	}
+
+public:
+
+	void SetRootMotionsDirty();
+	void SetUpRootMotions(bool bForce);
+	TOptional<FTransform>  GetRootMotion(FFrameTime CurrentTime);
+	void MatchSectionByBoneTransform(bool bMatchWithPrevious, USkeletalMeshComponent* SkelMeshComp, UMovieSceneSkeletalAnimationSection* CurrentSection, FFrameTime CurrentFrame, FFrameRate FrameRate,
+		const FName& BoneName, FTransform& SecondSectionRootDiff, FVector& TranslationDiff, FQuat& RotationDiff); //add options for z and for rotation.
+#if WITH_EDITORONLY_DATA
+	void ToggleShowRootMotionTrail();
+#endif
+
+#if WITH_EDITOR
+	virtual EMovieSceneSectionMovedResult OnSectionMoved(UMovieSceneSection& Section, const FMovieSceneSectionMovedParams& Params) override;
+#endif
+
+private:
+	//Not called yet, will be used to automatch a section when it's added to another
+	void AutoMatchSectionRoot(UMovieSceneSkeletalAnimationSection* AnimSection);
+
+	void SortSections();
+
+public:
 
 	/** List of all animation sections */
 	UPROPERTY()
 	TArray<TObjectPtr<UMovieSceneSection>> AnimationSections;
+
+	UPROPERTY()
+	FMovieSceneSkeletalAnimRootMotionTrackParams RootMotionParams;
+
+	/** Whether to blend and adjust the first child node with animation instead of the root, this should be true for blending when the root is static, false if the animations have proper root motion*/
+	UPROPERTY(EditAnywhere, Category = "Root Motions")
+	bool bBlendFirstChildOfRoot;
+
+#if WITH_EDITORONLY_DATA
+
+	/** Whether to show the position of the root for this sections */
+	UPROPERTY(EditAnywhere, Category = "Root Motions")
+	bool bShowRootMotionTrail;
+
+#endif
 };
\ No newline at end of file
diff --git a/Engine/Source/Runtime/MovieSceneTracks/Public/Tracks/MovieSceneSkeletalAnimationTrack.h b/Engine/Source/Runtime/MovieSceneTracks/Public/Tracks/MovieSceneSkeletalAnimationTrack.h
--- a/Engine/Source/Runtime/MovieSceneTracks/Public/Tracks/MovieSceneSkeletalAnimationTrack.h
+++ b/Engine/Source/Runtime/MovieSceneTracks/Public/Tracks/MovieSceneSkeletalAnimationTrack.h
@@ -11,30 +11,6 @@
 
 enum class ESwapRootBone : uint8;
 
-/**Struct to hold the cached root motion positions based upon how we calculated them.
-* Also provides way to get the root motion at a particular time.
-*/
-
-USTRUCT()
-struct FMovieSceneSkeletalAnimRootMotionTrackParams
-{
-	GENERATED_BODY()
-
-	FFrameTime FrameTick; 
-	FFrameTime StartFrame;
-	FFrameTime EndFrame;
-	bool bRootMotionsDirty;
-	bool bHaveRootMotion;
-	FTransform RootMotionStartOffset; // root motion may not be in Mesh Space if we are putting values on a bone that is a child of a root with an offset
-	/** Get the Root Motion transform at the specified time.*/
-	FMovieSceneSkeletalAnimRootMotionTrackParams() : bRootMotionsDirty(true),bHaveRootMotion(false) {}
-
-#if WITH_EDITORONLY_DATA
-	bool bCacheRootTransforms = false;
-	TArray<FTransform> RootTransforms;
-#endif
-};
-
 /**
  * Handles animation of skeletal mesh actors
  */
@@ -54,57 +30,18 @@ public:
 	// UMovieSceneTrack interface
 
 	virtual void PostLoad() override;
-#if WITH_EDITOR
-	virtual void PostEditImport() override;
-	virtual void PostEditUndo() override;
-	virtual void PostEditChangeProperty(FPropertyChangedEvent& PropertyChangedEvent) override;
-
-#endif
-
-	virtual void AddSection(UMovieSceneSection& Section) override;
-	virtual void RemoveSection(UMovieSceneSection& Section) override;
-	virtual void RemoveSectionAt(int32 SectionIndex) override;
 	virtual bool PopulateEvaluationTree(TMovieSceneEvaluationTree<FMovieSceneTrackEvaluationData>& OutData) const override;
-	virtual void UpdateEasing() override;
+	virtual bool ShouldUseRootMotions() const override;
 
 #if WITH_EDITORONLY_DATA
 	virtual FText GetDefaultDisplayName() const override;
 #endif
 
-#if WITH_EDITOR
-		virtual EMovieSceneSectionMovedResult OnSectionMoved(UMovieSceneSection& Section, const FMovieSceneSectionMovedParams& Params) override;
-#endif
-
-private:
-	void SortSections();
-
-public:
-
-	MOVIESCENETRACKS_API void SetRootMotionsDirty();
-	MOVIESCENETRACKS_API void SetUpRootMotions(bool bForce);
-	MOVIESCENETRACKS_API TOptional<FTransform>  GetRootMotion(FFrameTime CurrentTime);
-	MOVIESCENETRACKS_API void MatchSectionByBoneTransform(bool bMatchWithPrevious, USkeletalMeshComponent* SkelMeshComp, UMovieSceneSkeletalAnimationSection* CurrentSection, FFrameTime CurrentFrame, FFrameRate FrameRate,
-		const FName& BoneName, FTransform& SecondSectionRootDiff, FVector& TranslationDiff, FQuat& RotationDiff); //add options for z and for rotation.
-#if WITH_EDITORONLY_DATA
-
-	MOVIESCENETRACKS_API void ToggleShowRootMotionTrail();
-#endif
-
-private:
-	//Not called yet, will be used to automatch a section when it's added to another
-	void AutoMatchSectionRoot(UMovieSceneSkeletalAnimationSection* AnimSection);
 public:
 
 	UPROPERTY()
 	bool bUseLegacySectionIndexBlend;
 
-	UPROPERTY()
-	FMovieSceneSkeletalAnimRootMotionTrackParams RootMotionParams;
-
-	/** Whether to blend and adjust the first child node with animation instead of the root, this should be true for blending when the root is static, false if the animations have proper root motion*/
-	UPROPERTY(EditAnywhere, Category = "Root Motions")
-	bool bBlendFirstChildOfRoot;
-
 	/** If on the root bone transform will be swapped to the specified root*/
 	UPROPERTY(BlueprintReadWrite, EditAnywhere, Setter,Getter, Category = "Root Motions")
 	ESwapRootBone SwapRootBone;
@@ -114,14 +51,6 @@ public:
 	UFUNCTION()
 	MOVIESCENETRACKS_API ESwapRootBone GetSwapRootBone() const;
 
-#if WITH_EDITORONLY_DATA
-public:
-
-	/** Whether to show the position of the root for this sections */
-	UPROPERTY(EditAnywhere, Category = "Root Motions")
-	bool bShowRootMotionTrail;
-#endif
-
 };
 
 
