Sequencer: No longer propagate root motion offsets when encountering non-matched animation sections

diff --git a/Engine/Source/Runtime/MovieSceneTracks/Private/Tracks/MovieSceneCommonAnimationTrack.cpp b/Engine/Source/Runtime/MovieSceneTracks/Private/Tracks/MovieSceneCommonAnimationTrack.cpp
--- a/Engine/Source/Runtime/MovieSceneTracks/Private/Tracks/MovieSceneCommonAnimationTrack.cpp
+++ b/Engine/Source/Runtime/MovieSceneTracks/Private/Tracks/MovieSceneCommonAnimationTrack.cpp
@@ -455,6 +455,8 @@ void UMovieSceneCommonAnimationTrack::SetUpRootMotions(bool bForce)
 			UMovieSceneSkeletalAnimationSection* AnimSection = Cast<UMovieSceneSkeletalAnimationSection>(Section);
 			if (AnimSection)
 			{
+				const bool bCurrentSectionIsMatched = !AnimSection->MatchedBoneName.IsNone();
+
 				AnimSections.Add(AnimSection);
 				if (AnimSection->StartLocationOffset.IsNearlyZero() == false || AnimSection->StartRotationOffset.IsNearlyZero() == false ||
 					AnimSection->MatchedLocationOffset.IsNearlyZero() == false || AnimSection->MatchedRotationOffset.IsNearlyZero() == false)
@@ -465,7 +467,8 @@ void UMovieSceneCommonAnimationTrack::SetUpRootMotions(bool bForce)
 				{
 					ValidAnimSequence = AnimSection->Params.Animation;
 				}
-				if (PrevAnimSection)
+				// Only propagate offsets as long as root motion offsets are matched. Non-matched offsets should operate directly on the transform values.
+				if (PrevAnimSection && bCurrentSectionIsMatched)
 				{
 					if (UAnimSequenceBase* PrevAnimSequence = PrevAnimSection->Params.Animation)
 					{	
