Sequencer: Fixed crash caused by incorrect end time used for root motion transform generation when bone matching

diff --git a/Engine/Source/Runtime/MovieSceneTracks/Private/Tracks/MovieSceneCommonAnimationTrack.cpp b/Engine/Source/Runtime/MovieSceneTracks/Private/Tracks/MovieSceneCommonAnimationTrack.cpp
--- a/Engine/Source/Runtime/MovieSceneTracks/Private/Tracks/MovieSceneCommonAnimationTrack.cpp
+++ b/Engine/Source/Runtime/MovieSceneTracks/Private/Tracks/MovieSceneCommonAnimationTrack.cpp
@@ -545,8 +545,8 @@ void UMovieSceneCommonAnimationTrack::SetUpRootMotions(bool bForce)
 		}
 		RootMotionParams.bHaveRootMotion = true;
 		RootMotionParams.StartFrame = AnimSections[0]->GetInclusiveStartFrame();
-		RootMotionParams.StartFrame = AnimSections.Last()->GetExclusiveEndFrame() - 1;
-		RootMotionParams.FrameTick = FrameTick;		
+		RootMotionParams.EndFrame = AnimSections.Last()->GetExclusiveEndFrame() - 1;
+		RootMotionParams.FrameTick = FrameTick;
 		
 #if WITH_EDITORONLY_DATA
 		if (RootMotionParams.bCacheRootTransforms == false)
