Sequencer: Fixed decorations not being added correctly for common animation tracks

diff --git a/Engine/Source/Runtime/MovieSceneTracks/Private/Tracks/MovieSceneCommonAnimationTrack.cpp b/Engine/Source/Runtime/MovieSceneTracks/Private/Tracks/MovieSceneCommonAnimationTrack.cpp
--- a/Engine/Source/Runtime/MovieSceneTracks/Private/Tracks/MovieSceneCommonAnimationTrack.cpp
+++ b/Engine/Source/Runtime/MovieSceneTracks/Private/Tracks/MovieSceneCommonAnimationTrack.cpp
@@ -109,7 +109,7 @@ UMovieSceneSection* UMovieSceneCommonAnimationTrack::AddNewAnimationOnRow(FFrame
 #endif
 	}
 
-	AddSection(*NewSection);
+	Super::AddSection(*NewSection);
 
 	return NewSection;
 }
