Fix potential unsafe memory access in Transform Origin System

From UDN. A customer ran into an issue with -stompmalloc when passing the parent transform by reference when inserting a new value. Seemd the array may be reallocated during insertion. Passing by value now instead.

diff --git a/Engine/Source/Runtime/MovieSceneTracks/Private/Systems/MovieSceneTransformOriginSystem.cpp b/Engine/Source/Runtime/MovieSceneTracks/Private/Systems/MovieSceneTransformOriginSystem.cpp
--- a/Engine/Source/Runtime/MovieSceneTracks/Private/Systems/MovieSceneTransformOriginSystem.cpp
+++ b/Engine/Source/Runtime/MovieSceneTracks/Private/Systems/MovieSceneTransformOriginSystem.cpp
@@ -119,7 +119,8 @@ struct FGatherTransformOriginsFromSubscenes
 			}
 			else
 			{
-				TransformOriginsByInstanceID->Insert(Mapping.Child.InstanceID, (*TransformOriginsByInstanceID)[Mapping.Parent.InstanceID]);
+				const FTransform ParentTransform = (*TransformOriginsByInstanceID)[Mapping.Parent.InstanceID];
+				TransformOriginsByInstanceID->Insert(Mapping.Child.InstanceID, ParentTransform);
 			}
 		}
 	}
