Fixing metasound duration in sequencer tracks.
Fixing issue of base-submix enabled (and other similar properties) being ignored in sequencer

- Added a base function on USoundBase to query if a sound is procedurally generated
- Fixed up sequencer tracks on playback so it doesn't constantly retrigger metasounds incorrectly and obeys the track looping.
- Makes the track looping if the metasound is looping (this should be what we do for regular sounds too)
- Does NOT play the audio component as a preview component in the sequencer editor

[RN] Fixing issue of MetaSounds on sequencer having a gigantic duration. Now, it gives a 1-second region you can drag for the region.

diff --git a/Engine/Source/Editor/MovieSceneTools/Private/TrackEditors/AudioTrackEditor.cpp b/Engine/Source/Editor/MovieSceneTools/Private/TrackEditors/AudioTrackEditor.cpp
--- a/Engine/Source/Editor/MovieSceneTools/Private/TrackEditors/AudioTrackEditor.cpp
+++ b/Engine/Source/Editor/MovieSceneTools/Private/TrackEditors/AudioTrackEditor.cpp
@@ -472,7 +472,7 @@ void FAudioThumbnail::GenerateWaveformPreview(TArray<uint8>& OutData, TRange<flo
 	check(SoundWave->NumChannels == 1 || SoundWave->NumChannels == 2);
 
 	// If this SoundWave is generated procedurally, it's not possible to render a thumbnail.
-	if (SoundWave->bProcedural)
+	if (SoundWave->IsProcedurallyGenerated())
 	{
 		return;
 	}
diff --git a/Engine/Source/Runtime/Engine/Classes/Sound/SoundBase.h b/Engine/Source/Runtime/Engine/Classes/Sound/SoundBase.h
--- a/Engine/Source/Runtime/Engine/Classes/Sound/SoundBase.h
+++ b/Engine/Source/Runtime/Engine/Classes/Sound/SoundBase.h
@@ -300,6 +300,9 @@ public:
 	/** Query if it's one shot. One shot is defined as a sound which is intended to have a fixed duration. */
 	ENGINE_API virtual bool IsOneShot() const;
 
+	/** Query if it's a procedurally generated sound. */
+	ENGINE_API virtual bool IsProcedurallyGenerated() const { return false; }
+
 	/** Parses the Sound to generate the WaveInstances to play. */
 	virtual void Parse( class FAudioDevice* AudioDevice, const UPTRINT NodeWaveInstanceHash, FActiveSound& ActiveSound, const FSoundParseParameters& ParseParams, TArray<FWaveInstance*>& WaveInstances ) { }
 
diff --git a/Engine/Source/Runtime/Engine/Classes/Sound/SoundWave.h b/Engine/Source/Runtime/Engine/Classes/Sound/SoundWave.h
--- a/Engine/Source/Runtime/Engine/Classes/Sound/SoundWave.h
+++ b/Engine/Source/Runtime/Engine/Classes/Sound/SoundWave.h
@@ -1159,6 +1159,7 @@ public:
 	ENGINE_API virtual bool GetSoundWavesWithCookedAnalysisData(TArray<USoundWave*>& OutSoundWaves) override;
 	ENGINE_API virtual bool HasCookedFFTData() const override;
 	ENGINE_API virtual bool HasCookedAmplitudeEnvelopeData() const override;
+	ENGINE_API virtual bool IsProcedurallyGenerated() const override { return bProcedural; }
 	//~ End USoundBase Interface.
 
 	ENGINE_API FSoundWaveProxyPtr CreateSoundWaveProxy();
diff --git a/Engine/Source/Runtime/MovieScene/Private/MovieSceneCommonHelpers.cpp b/Engine/Source/Runtime/MovieScene/Private/MovieSceneCommonHelpers.cpp
--- a/Engine/Source/Runtime/MovieScene/Private/MovieSceneCommonHelpers.cpp
+++ b/Engine/Source/Runtime/MovieScene/Private/MovieSceneCommonHelpers.cpp
@@ -543,10 +543,13 @@ UCameraComponent* MovieSceneHelpers::CameraComponentFromRuntimeObject(UObject* R
 
 float MovieSceneHelpers::GetSoundDuration(USoundBase* Sound)
 {
+	if (Sound->IsProcedurallyGenerated() || !Sound->IsOneShot())
+	{
+		return INDEFINITELY_LOOPING_DURATION;
+	}
 	return Sound ? FMath::Max(0.0f, Sound->GetDuration()) : 0.0f;
 }
 
-
 float MovieSceneHelpers::CalculateWeightForBlending(UMovieSceneSection* SectionToKey, FFrameNumber Time)
 {
 	float Weight = 1.0f;
diff --git a/Engine/Source/Runtime/MovieSceneTracks/Private/Sections/MovieSceneAudioSection.cpp b/Engine/Source/Runtime/MovieSceneTracks/Private/Sections/MovieSceneAudioSection.cpp
--- a/Engine/Source/Runtime/MovieSceneTracks/Private/Sections/MovieSceneAudioSection.cpp
+++ b/Engine/Source/Runtime/MovieSceneTracks/Private/Sections/MovieSceneAudioSection.cpp
@@ -356,8 +356,6 @@ TOptional<TRange<FFrameNumber> > UMovieSceneAudioSection::GetAutoSizeRange() con
 		return TRange<FFrameNumber>();
 	}
 
-	const float SoundDuration = MovieSceneHelpers::GetSoundDuration(Sound);
-
 	const FFrameRate FrameRate = GetTypedOuter<UMovieScene>()->GetTickResolution();
 
 	// determine initial duration
@@ -365,9 +363,16 @@ TOptional<TRange<FFrameNumber> > UMovieSceneAudioSection::GetAutoSizeRange() con
 	// @todo ^^ Why? Infinte sections would mean there's no starting time?
 	FFrameTime DurationToUse = 1.f * FrameRate; // if all else fails, use 1 second duration
 
-	if (SoundDuration != INDEFINITELY_LOOPING_DURATION && SoundDuration > 0)
+	// Only one-shot non-procedural sounds have a known duration
+	if (Sound->IsOneShot() && !Sound->IsProcedurallyGenerated())
 	{
-		DurationToUse = FMath::Max(SoundDuration * FrameRate - StartFrameOffset, FFrameTime(1));
+		const float SoundDuration = MovieSceneHelpers::GetSoundDuration(Sound);
+		// This should not hit if the sound has returned true for IsOneShot()
+		check(SoundDuration != INDEFINITELY_LOOPING_DURATION);
+		if (SoundDuration > 0)
+		{
+			DurationToUse = FMath::Max(SoundDuration * FrameRate - StartFrameOffset, FFrameTime(1));
+		}
 	}
 
 	const int32 IFrameNumber = DurationToUse.FrameNumber.Value + static_cast<int32>(DurationToUse.GetSubFrame() + 0.5f);
@@ -412,6 +417,15 @@ UMovieSceneSection* UMovieSceneAudioSection::SplitSection(FQualifiedFrameTime Sp
 void UMovieSceneAudioSection::SetSound(USoundBase* InSound)
 {
 	Sound = InSound;
+
+	// For a procedurally generated sound, only have the section loop if it's specifically a looping sound
+	// this is because there's no way to determine a duration for the sound so the user has no way of
+	// now when to end the region to avoid retriggering the one-shot procedural asset.
+	if (Sound->IsProcedurallyGenerated())
+	{
+		bLooping = !Sound->IsOneShot();
+	}
+
 	InvalidateChannelProxy();
 }
 
diff --git a/Engine/Source/Runtime/MovieSceneTracks/Private/Systems/MovieSceneAudioSystem.cpp b/Engine/Source/Runtime/MovieSceneTracks/Private/Systems/MovieSceneAudioSystem.cpp
--- a/Engine/Source/Runtime/MovieSceneTracks/Private/Systems/MovieSceneAudioSystem.cpp
+++ b/Engine/Source/Runtime/MovieSceneTracks/Private/Systems/MovieSceneAudioSystem.cpp
@@ -614,20 +614,34 @@ private:
 				}
 			}
 
-			const float Duration = MovieSceneHelpers::GetSoundDuration(Sound);
-
-			if (!AudioSection.GetLooping() && AudioTime > Duration && Duration != 0.f)
+			// Procedurally generated sounds don't have a defined duration so when the audio component is done, it's done
+			if (EvaluationData.bAudioComponentHasBeenPlayed && !AudioSection.GetLooping())
 			{
-				if (AudioComponent.IsPlaying())
+				if (Sound->IsProcedurallyGenerated())
 				{
-					UE_LOG(LogMovieScene, Verbose, TEXT("Audio Component reached end of playback. Component: %s Sound: %s"), *AudioComponent.GetName(), *GetNameSafe(AudioComponent.Sound));
-					AudioComponent.Stop();
+					// If we're not a looping section and the AC is done, return
+					// otherwise a looping section will restart the sound if it's not playing
+					if (!AudioComponent.IsPlaying())
+					{
+					    UE_LOG(LogMovieScene, Verbose, TEXT("Procedural sound Audio Component reached end of playback. Component: %s Sound: %s"), *AudioComponent.GetName(), *GetNameSafe(AudioComponent.Sound));
+						return;
+					}
 				}
-				
-				return;
-			}
+				else
+				{
+					const float Duration = MovieSceneHelpers::GetSoundDuration(Sound);
+					if (AudioTime > Duration && Duration != 0.f)
+					{
+						if (AudioComponent.IsPlaying())
+						{
+							UE_LOG(LogMovieScene, Verbose, TEXT("Audio Component reached end of playback. Component: %s Sound: %s"), *AudioComponent.GetName(), *GetNameSafe(AudioComponent.Sound));
+							AudioComponent.Stop();
+						}
 
-			AudioTime = Duration > 0.f ? FMath::Fmod(AudioTime, Duration) : AudioTime;
+						return;
+					}
+				}
+			}
 		}
 
 		// If the audio component is not playing we (may) need a state change. If the audio component is playing
@@ -723,8 +737,8 @@ private:
 #if WITH_EDITOR
 			if (GIsEditor && World != nullptr && !World->IsPlayInEditor())
 			{
+				// This is needed otherwise the sound doesn't have a position and will not play properly
 				AudioComponent.bIsUISound = true;
-				AudioComponent.bIsPreviewSound = true;
 			}
 			else
 #endif // WITH_EDITOR
@@ -742,7 +756,7 @@ private:
 					if (USoundWave* SoundWave = Cast<USoundWave>(Sound))
 					{
 						// Can't scrub a procedural asset right now
-						if (ScrubbedSound && !SoundWave->bProcedural)
+						if (ScrubbedSound && !SoundWave->IsProcedurallyGenerated())
 						{
 							// If we're not playing the audio component, then set up the audio component to use the scrubbed sound wave and play it
 							if (!AudioComponent.IsPlaying())
@@ -776,8 +790,11 @@ private:
 						ScrubbedSound->SetIsScrubbing(false);
 					}
 #endif // WITH_EDITOR
-
+					
 					AudioComponent.Play(AudioTime);
+					
+					// Set that we've played an audio component. This is used by procedural sounds who have undefined duration.
+					EvaluationData.bAudioComponentHasBeenPlayed = true;
 
 					// Keep track of when we asked this audio clip to play (in game time) so that we can figure 
 					// out if there's a significant desync in the future.
@@ -1152,6 +1169,7 @@ UMovieSceneAudioSystem::FAudioComponentEvaluationData* UMovieSceneAudioSystem::A
 		FInstanceObjectKey DataKey{ InstanceHandle, SectionKey };
 		ExistingData = &ActorAudioComponentMap.Add(DataKey);
 		ExistingData->AudioComponent = NewComponent;
+		ExistingData->bAudioComponentHasBeenPlayed = false;
 
 #if WITH_EDITOR
 		static int32 ScrubSoundCounter = 0;
@@ -1203,6 +1221,7 @@ UMovieSceneAudioSystem::FAudioComponentEvaluationData* UMovieSceneAudioSystem::A
 		FInstanceObjectKey DataKey{ InstanceHandle, SectionKey };
 		ExistingData = &RootAudioComponentMap.Add(DataKey);
 		ExistingData->AudioComponent = NewComponent;
+		ExistingData->bAudioComponentHasBeenPlayed = false;
 
 #if WITH_EDITOR
 		static int32 ScrubSoundCounter = 0;
diff --git a/Engine/Source/Runtime/MovieSceneTracks/Private/Systems/MovieSceneAudioSystem.h b/Engine/Source/Runtime/MovieSceneTracks/Private/Systems/MovieSceneAudioSystem.h
--- a/Engine/Source/Runtime/MovieSceneTracks/Private/Systems/MovieSceneAudioSystem.h
+++ b/Engine/Source/Runtime/MovieSceneTracks/Private/Systems/MovieSceneAudioSystem.h
@@ -46,10 +46,10 @@ namespace UE::MovieScene
 #endif
 
 		/** Volume multiplier to use this frame */
-		double VolumeMultiplier;
+		double VolumeMultiplier = 1.0;
 
 		/** Pitch multiplier to use this frame */
-		double PitchMultiplier;
+		double PitchMultiplier = 1.0;
 
 		/**
 		 * Set whenever we ask the Audio component to start playing a sound.
@@ -63,7 +63,10 @@ namespace UE::MovieScene
 		TOptional<float> LastContextTime;
 
 		/** Flag to keep track of audio components evaluated on a given frame */
-		bool bEvaluatedThisFrame;
+		bool bEvaluatedThisFrame = false;
+
+		/** Flag to keep track of if the audio component was played in a previous frame. */
+		bool bAudioComponentHasBeenPlayed = false;
 	};
 }
 
