Sequencer: fix logging of blending camera cut

It was always reporting a blend time of zero.

diff --git a/Engine/Source/Runtime/MovieSceneTracks/Private/TrackInstances/MovieSceneCameraCutGameHandler.cpp b/Engine/Source/Runtime/MovieSceneTracks/Private/TrackInstances/MovieSceneCameraCutGameHandler.cpp
--- a/Engine/Source/Runtime/MovieSceneTracks/Private/TrackInstances/MovieSceneCameraCutGameHandler.cpp
+++ b/Engine/Source/Runtime/MovieSceneTracks/Private/TrackInstances/MovieSceneCameraCutGameHandler.cpp
@@ -349,18 +349,20 @@ void FCameraCutGameHandler::SetCameraCut(
 			);  
 	if (bDoBlend)
 	{
+		// The playrate of sequences is defined as delta time * rate, so we need to match approach.
+		const float PlayRateFactor = Wrapper.CameraCutCapability ? Wrapper.CameraCutCapability->GetCameraBlendPlayRate() : 1.f;
+		const float BlendTime = CameraCutParams.BlendTime * (1 / FMath::Max(UE_SMALL_NUMBER, abs(PlayRateFactor)));
+
 		UE_LOG(LogMovieScene, Log, TEXT("Blending into new camera cut: '%s' -> '%s' (blend time: %f)"),
 			(ViewTarget ? *ViewTarget->GetName() : TEXT("None")),
 			(CameraActor ? *CameraActor->GetName() : TEXT("None")),
-			TransitionParams.BlendTime);
-
-		const float PlayRateFactor = Wrapper.CameraCutCapability ? Wrapper.CameraCutCapability->GetCameraBlendPlayRate() : 1.f;
+			BlendTime);
 
 		// Convert known easing functions to their corresponding view target blend parameters.
 		TTuple<EViewTargetBlendFunction, float> BlendFunctionAndExp = BuiltInEasingTypeToBlendFunction(CameraCutParams.BlendType.GetValue());
 		
 		// The playrate of sequences is defined as delta time * rate, so we need to match approach.
-		TransitionParams.BlendTime = CameraCutParams.BlendTime * (1 / FMath::Max(UE_SMALL_NUMBER, abs(PlayRateFactor)));
+		TransitionParams.BlendTime = BlendTime;
 		TransitionParams.bLockOutgoing = CameraCutParams.bLockPreviousCamera;
 		TransitionParams.BlendFunction = BlendFunctionAndExp.Get<0>();
 		TransitionParams.BlendExp = BlendFunctionAndExp.Get<1>();
