Sequencer: Fixed copying dynamic bindings between Sequences causing a crash due to funciton call from the wrong director BP instance

It is invalid to call a UFunction from one class in the context of a different class. This change prevents a crash that will occur when pasting the dynamic binding, but really we should try and copy the endpoint to the new director BP

diff --git a/Engine/Source/Runtime/MovieScene/Private/MovieSceneDynamicBindingInvoker.cpp b/Engine/Source/Runtime/MovieScene/Private/MovieSceneDynamicBindingInvoker.cpp
--- a/Engine/Source/Runtime/MovieScene/Private/MovieSceneDynamicBindingInvoker.cpp
+++ b/Engine/Source/Runtime/MovieScene/Private/MovieSceneDynamicBindingInvoker.cpp
@@ -52,6 +52,15 @@ FMovieSceneDynamicBindingResolveResult FMovieSceneDynamicBindingInvoker::Resolve
 		return FMovieSceneDynamicBindingResolveResult();
 	}
 
+	if (!DynamicBindingFunc->IsIn(DirectorInstance->GetClass()))
+	{
+		UE_LOG(LogMovieScene, Warning,
+			TEXT("%s: Failed to resolve dynamic binding '%s' because the function is not contained within the director instance."),
+			*Sequence->GetName(), *DynamicBindingFunc->GetName());
+
+		return FMovieSceneDynamicBindingResolveResult();
+	}
+
 #if WITH_EDITOR
 	const static FName NAME_CallInEditor(TEXT("CallInEditor"));
 
diff --git a/Engine/Source/Runtime/MovieSceneTracks/Private/Bindings/MovieSceneReplaceableDirectorBlueprintBinding.cpp b/Engine/Source/Runtime/MovieSceneTracks/Private/Bindings/MovieSceneReplaceableDirectorBlueprintBinding.cpp
--- a/Engine/Source/Runtime/MovieSceneTracks/Private/Bindings/MovieSceneReplaceableDirectorBlueprintBinding.cpp
+++ b/Engine/Source/Runtime/MovieSceneTracks/Private/Bindings/MovieSceneReplaceableDirectorBlueprintBinding.cpp
@@ -61,6 +61,29 @@ UMovieSceneCustomBinding* UMovieSceneReplaceableDirectorBlueprintBinding::Create
 }
 #endif
 
+void UMovieSceneReplaceableDirectorBlueprintBinding::PostDuplicate(EDuplicateMode::Type DuplicateMode)
+{
+	Super::PostDuplicate(DuplicateMode);
+
+	// If we were duplicated into a different package we can't reference the old function any more
+	//    For now we null it out, but it would be good to copy the endpoint as well (we currently can't do that
+	//    because there is no way to generically access the correct director BP class from UMovieSceneSequence
+	if (DynamicBinding.Function)
+	{
+		UPackage* FunctionPackage = DynamicBinding.Function->GetPackage();
+		UPackage* Package = GetPackage();
+
+		if (Package != FunctionPackage)
+		{
+			DynamicBinding.Function = nullptr;
+			DynamicBinding.ResolveParamsProperty = nullptr;
+#if WITH_EDITORONLY_DATA
+			DynamicBinding.WeakEndpoint = nullptr;
+#endif
+		}
+	}
+}
+
 FMovieSceneBindingResolveResult UMovieSceneReplaceableDirectorBlueprintBinding::ResolveRuntimeBindingInternal(const FMovieSceneBindingResolveParams& ResolveParams, int32 BindingIndex, TSharedRef<const UE::MovieScene::FSharedPlaybackState> SharedPlaybackState) const
 {
 	FMovieSceneBindingResolveResult ResolveResult;
diff --git a/Engine/Source/Runtime/MovieSceneTracks/Public/Bindings/MovieSceneReplaceableDirectorBlueprintBinding.h b/Engine/Source/Runtime/MovieSceneTracks/Public/Bindings/MovieSceneReplaceableDirectorBlueprintBinding.h
--- a/Engine/Source/Runtime/MovieSceneTracks/Public/Bindings/MovieSceneReplaceableDirectorBlueprintBinding.h
+++ b/Engine/Source/Runtime/MovieSceneTracks/Public/Bindings/MovieSceneReplaceableDirectorBlueprintBinding.h
@@ -41,6 +41,8 @@ public:
 protected:
 	/* MovieSceneReplaceableBindingBase overrides*/
 
+	virtual void PostDuplicate(EDuplicateMode::Type DuplicateMode) override;
+
 #if WITH_EDITOR
 	MOVIESCENETRACKS_API virtual UMovieSceneCustomBinding* CreateCustomBindingFromBinding(const FMovieSceneBindingReference& BindingReference, UObject* SourceObject, UMovieScene& OwnerMovieScene) override;
 #endif
