Fix for unrolled SoundWaves not looping in sequencer

diff --git a/Engine/Source/Runtime/MovieSceneTracks/Private/Systems/MovieSceneAudioSystem.cpp b/Engine/Source/Runtime/MovieSceneTracks/Private/Systems/MovieSceneAudioSystem.cpp
--- a/Engine/Source/Runtime/MovieSceneTracks/Private/Systems/MovieSceneAudioSystem.cpp
+++ b/Engine/Source/Runtime/MovieSceneTracks/Private/Systems/MovieSceneAudioSystem.cpp
@@ -614,25 +614,29 @@ private:
 				}
 			}
 
-			// Procedurally generated sounds don't have a defined duration so when the audio component is done, it's done
-			if (EvaluationData.bAudioComponentHasBeenPlayed && !AudioSection.GetLooping())
+			if (EvaluationData.bAudioComponentHasBeenPlayed)
 			{
+				// Procedurally generated sounds don't have a defined duration so when the audio component is done, it's done
 				if (Sound->IsProcedurallyGenerated())
 				{
-					// If we're not a looping section and the AC is done, return
-					// otherwise a looping section will restart the sound if it's not playing
-					if (!AudioComponent.IsPlaying())
+					if (!AudioSection.GetLooping())
 					{
-					    UE_LOG(LogMovieScene, Verbose, TEXT("Procedural sound Audio Component reached end of playback. Component: %s Sound: %s"), *AudioComponent.GetName(), *GetNameSafe(AudioComponent.Sound));
-						return;
+						// If we're not a looping section and the AC is done, return
+						// otherwise a looping section will restart the sound if it's not playing
+						if (!AudioComponent.IsPlaying())
+						{
+							UE_LOG(LogMovieScene, Verbose, TEXT("Procedural sound Audio Component reached end of playback. Component: %s Sound: %s"), *AudioComponent.GetName(), *GetNameSafe(AudioComponent.Sound));
+							return;
+						}
 					}
 				}
 				else
 				{
 					const float Duration = MovieSceneHelpers::GetSoundDuration(Sound);
-					if (AudioTime > Duration && Duration != 0.f)
+					if (!AudioSection.GetLooping())
 					{
-						if (AudioComponent.IsPlaying())
+						// If this is Non-ProcedurallyGenerated audio, and it's not looping then check to see if it needs to be stopped
+						if (AudioTime > Duration && AudioComponent.IsPlaying())
 						{
 							UE_LOG(LogMovieScene, Verbose, TEXT("Audio Component reached end of playback. Component: %s Sound: %s"), *AudioComponent.GetName(), *GetNameSafe(AudioComponent.Sound));
 							AudioComponent.Stop();
@@ -640,6 +644,11 @@ private:
 
 						return;
 					}
+					else
+					{
+						// Wrap AudioTime according to duration for Non-ProcedurallyGenerated audio that is looping
+						AudioTime = Duration > 0.f ? FMath::Fmod(AudioTime, Duration) : AudioTime;
+					}
 				}
 			}
 		}
