Sequencer: properly setup the enum blender system override for blend value traits

diff --git a/Engine/Source/Runtime/MovieSceneTracks/Private/Systems/MovieScenePiecewiseEnumBlenderSystem.cpp b/Engine/Source/Runtime/MovieSceneTracks/Private/Systems/MovieScenePiecewiseEnumBlenderSystem.cpp
--- a/Engine/Source/Runtime/MovieSceneTracks/Private/Systems/MovieScenePiecewiseEnumBlenderSystem.cpp
+++ b/Engine/Source/Runtime/MovieSceneTracks/Private/Systems/MovieScenePiecewiseEnumBlenderSystem.cpp
@@ -8,41 +8,6 @@
 
 #include UE_INLINE_GENERATED_CPP_BY_NAME(MovieScenePiecewiseEnumBlenderSystem)
 
-namespace UE
-{
-namespace MovieScene
-{
-
-/** 
- * Custom blend result traits for enums:
- *
- * - We don't blend enums, unlike bytes, because we don't even know if the in-between values are valid
- *   enumerations.
- * - Number of contributors doesn't matter. The last one wins.
- */
-template<>
-struct TSimpleBlendResultTraits<uint8>
-{
-	static void ZeroAccumulationBuffer(TArrayView<TSimpleBlendResult<uint8>> Buffer)
-	{
-		FMemory::Memzero(Buffer.GetData(), sizeof(TSimpleBlendResult<uint8>) * Buffer.Num());
-	}
-
-	static void AccumulateResult(TSimpleBlendResult<uint8>& InOutValue, uint8 Contributor)
-	{
-		InOutValue.Value = Contributor;
-		++InOutValue.NumContributors;
-	}
-
-	static uint8 BlendResult(const TSimpleBlendResult<uint8>& InResult)
-	{
-		return InResult.Value;
-	}
-};
-
-} // namespace MovieScene
-} // namespace UE
-
 UMovieScenePiecewiseEnumBlenderSystem::UMovieScenePiecewiseEnumBlenderSystem(const FObjectInitializer& ObjInit)
 	: Super(ObjInit)
 {
diff --git a/Engine/Source/Runtime/MovieSceneTracks/Public/Systems/MovieSceneBlenderSystemHelper.h b/Engine/Source/Runtime/MovieSceneTracks/Public/Systems/MovieSceneBlenderSystemHelper.h
--- a/Engine/Source/Runtime/MovieSceneTracks/Public/Systems/MovieSceneBlenderSystemHelper.h
+++ b/Engine/Source/Runtime/MovieSceneTracks/Public/Systems/MovieSceneBlenderSystemHelper.h
@@ -55,7 +55,7 @@ struct TSimpleBlendResultTraits
 	}
 };
 
-template<typename PropertyType>
+template<typename PropertyType, typename ResultTraits = TSimpleBlendResultTraits<PropertyType>>
 struct TSimpleBlenderGatherResults
 {
 	FMovieSceneBlenderSystemID SystemID;
@@ -68,17 +68,15 @@ struct TSimpleBlenderGatherResults
 
 	void ForEachEntity(FMovieSceneBlendChannelID BlendChannelInput, PropertyType Value) const
 	{
-		using FResultTraits = TSimpleBlendResultTraits<PropertyType>;
-
 		ensureMsgf(BlendChannelInput.SystemID == SystemID, TEXT("Overriding the standard blender system of standard types isn't supported."));
 		if (ensure(BlendChannelResults.IsValidIndex(BlendChannelInput.ChannelID)))
 		{
-			FResultTraits::AccumulateResult(BlendChannelResults[BlendChannelInput.ChannelID], Value);
+			ResultTraits::AccumulateResult(BlendChannelResults[BlendChannelInput.ChannelID], Value);
 		}
 	}
 };
 
-template<typename PropertyType>
+template<typename PropertyType, typename ResultTraits = TSimpleBlendResultTraits<PropertyType>>
 struct TSimpleBlenderCombineResults
 {
 	FMovieSceneBlenderSystemID SystemID;
@@ -91,12 +89,10 @@ struct TSimpleBlenderCombineResults
 
 	void ForEachEntity(FMovieSceneBlendChannelID BlendChannelOutput, PropertyType& OutValue) const
 	{
-		using FResultTraits = TSimpleBlendResultTraits<PropertyType>;
-
 		ensureMsgf(BlendChannelOutput.SystemID == SystemID, TEXT("Overriding the standard blender system of standard types isn't supported."));
 		if (ensure(BlendChannelResults.IsValidIndex(BlendChannelOutput.ChannelID)))
 		{
-			OutValue = FResultTraits::BlendResult(BlendChannelResults[BlendChannelOutput.ChannelID]);
+			OutValue = ResultTraits::BlendResult(BlendChannelResults[BlendChannelOutput.ChannelID]);
 		}
 	}
 };
@@ -107,10 +103,14 @@ struct TSimpleBlenderCombineResults
  * This class implements a simple blender system that accumulates all contributors to a blend channel, and returns 
  * the blended value, as defined by the value type's traits class (above).
  */
-template<typename PropertyType>
+template<typename PropertyType, typename ResultTraits = TSimpleBlendResultTraits<PropertyType>>
 class TSimpleBlenderSystemImpl
 {
 public:
+
+	using FGatherResults = TSimpleBlenderGatherResults<PropertyType, ResultTraits>;
+	using FCombineResults = TSimpleBlenderCombineResults<PropertyType, ResultTraits>;
+
 	TSimpleBlenderSystemImpl()
 		: BlenderSystem(nullptr)
 	{}
@@ -168,7 +168,7 @@ public:
 			.Read(ResultComponentID)
 			.FilterAll({ BlenderSystem->GetBlenderTypeTag() })
 			.FilterNone({ BuiltInComponents->Tags.Ignored })
-			.template Schedule_PerEntity<TSimpleBlenderGatherResults<PropertyType>>(
+			.template Schedule_PerEntity<FGatherResults>(
 					&Linker->EntityManager, TaskScheduler, BlenderSystem->GetBlenderSystemID(), BlendChannelResults);
 
 		TaskScheduler->AddPrerequisite(ResetWeightsTask, AccumulateTask);
@@ -179,7 +179,7 @@ public:
 			.Write(ResultComponentID)
 			.FilterAll({ BlenderSystem->GetBlenderTypeTag() })
 			.FilterNone({ BuiltInComponents->Tags.Ignored })
-			.template Fork_PerEntity<TSimpleBlenderCombineResults<PropertyType>>(
+			.template Fork_PerEntity<FCombineResults>(
 					&Linker->EntityManager, TaskScheduler, BlenderSystem->GetBlenderSystemID(), BlendChannelResults);
 
 		TaskScheduler->AddPrerequisite(AccumulateTask, CombineTask);
@@ -229,7 +229,7 @@ public:
 			.Read(ResultComponentID)
 			.FilterAll({ BlenderSystem->GetBlenderTypeTag() })
 			.FilterNone({ BuiltInComponents->Tags.Ignored })
-			.template Dispatch_PerEntity<TSimpleBlenderGatherResults<PropertyType>>(
+			.template Dispatch_PerEntity<FGatherResults>(
 					&Linker->EntityManager, InPrerequisites, nullptr, BlenderSystem->GetBlenderSystemID(), BlendChannelResults);
 		if (Task)
 		{
@@ -242,7 +242,7 @@ public:
 			.Write(ResultComponentID)
 			.FilterAll({ BlenderSystem->GetBlenderTypeTag() })
 			.FilterNone({ BuiltInComponents->Tags.Ignored })
-			.template Dispatch_PerEntity<TSimpleBlenderCombineResults<PropertyType>>(
+			.template Dispatch_PerEntity<FCombineResults>(
 					&Linker->EntityManager, Prereqs, &Subsequents, BlenderSystem->GetBlenderSystemID(), BlendChannelResults);
 	}
 
diff --git a/Engine/Source/Runtime/MovieSceneTracks/Public/Systems/MovieScenePiecewiseEnumBlenderSystem.h b/Engine/Source/Runtime/MovieSceneTracks/Public/Systems/MovieScenePiecewiseEnumBlenderSystem.h
--- a/Engine/Source/Runtime/MovieSceneTracks/Public/Systems/MovieScenePiecewiseEnumBlenderSystem.h
+++ b/Engine/Source/Runtime/MovieSceneTracks/Public/Systems/MovieScenePiecewiseEnumBlenderSystem.h
@@ -14,6 +14,39 @@
 
 class UObject;
 
+namespace UE
+{
+namespace MovieScene
+{
+
+/** 
+ * Custom blend result traits for enums:
+ *
+ * - We don't blend enums, unlike bytes, because we don't even know if the in-between values are valid
+ *   enumerations.
+ * - Number of contributors doesn't matter. The last one wins.
+ */
+struct FSimpleBlenderEnumResultTraits
+{
+	static void ZeroAccumulationBuffer(TArrayView<TSimpleBlendResult<uint8>> Buffer)
+	{
+		FMemory::Memzero(Buffer.GetData(), sizeof(TSimpleBlendResult<uint8>) * Buffer.Num());
+	}
+
+	static void AccumulateResult(TSimpleBlendResult<uint8>& InOutValue, uint8 Contributor)
+	{
+		InOutValue.Value = Contributor;
+		++InOutValue.NumContributors;
+	}
+
+	static uint8 BlendResult(const TSimpleBlendResult<uint8>& InResult)
+	{
+		return InResult.Value;
+	}
+};
+
+} // namespace MovieScene
+} // namespace UE
 
 UCLASS(MinimalAPI)
 class UMovieScenePiecewiseEnumBlenderSystem : public UMovieSceneBlenderSystem
@@ -28,6 +61,6 @@ public:
 
 private:
 
-	UE::MovieScene::TSimpleBlenderSystemImpl<uint8> Impl;
+	UE::MovieScene::TSimpleBlenderSystemImpl<uint8, UE::MovieScene::FSimpleBlenderEnumResultTraits> Impl;
 };
 
