Sequencer: Fixes a crash caused when the user reviews a take taken by [REDACTED_PERSON], modifies a property on a recorded actor, and closes the take.

The crash occurs because the track's template object is incorrectly Rename()-ed: the Rename function creates a UObjectRedirector. Later, in MovieSceneHelpers::MakeSpawnableTemplateFromInstance, that redirector is found by NewObject - which causes a crash due to name collision. Creating an object redirector was certainly not the intention of the original author. My impression is that this workflow is not very well tested & uncommon.

diff --git a/Engine/Source/Runtime/MovieSceneTracks/Private/Bindings/MovieSceneSpawnableActorBinding.cpp b/Engine/Source/Runtime/MovieSceneTracks/Private/Bindings/MovieSceneSpawnableActorBinding.cpp
--- a/Engine/Source/Runtime/MovieSceneTracks/Private/Bindings/MovieSceneSpawnableActorBinding.cpp
+++ b/Engine/Source/Runtime/MovieSceneTracks/Private/Bindings/MovieSceneSpawnableActorBinding.cpp
@@ -378,7 +378,9 @@ void UMovieSceneSpawnableActorBinding::CopyObjectTemplate(UObject* InSourceObjec
 
 		if (ActorTemplate)
 		{
-			ActorTemplate->Rename(*MakeUniqueObjectName(MovieSceneSequence.GetMovieScene(), GetActorClass(), "ExpiredSpawnable").ToString());
+			const FString NewName = MakeUniqueObjectName(MovieSceneSequence.GetMovieScene(), GetActorClass(), "ExpiredSpawnable").ToString();
+			// Without REN_DontCreateRedirectors, we'd create a redirector causing a name collision when we call MakeSpawnableTemplateFromInstance.
+			ActorTemplate->Rename(*NewName, nullptr, REN_DoNotDirty | REN_DontCreateRedirectors | REN_NonTransactional);
 			ActorTemplate->MarkAsGarbage();
 			ActorTemplate = nullptr;
 		}
