Sequencer: Check for track row disabled when populating the cinematic shot track evaluation

diff --git a/Engine/Source/Runtime/MovieSceneTracks/Private/Tracks/MovieSceneCinematicShotTrack.cpp b/Engine/Source/Runtime/MovieSceneTracks/Private/Tracks/MovieSceneCinematicShotTrack.cpp
--- a/Engine/Source/Runtime/MovieSceneTracks/Private/Tracks/MovieSceneCinematicShotTrack.cpp
+++ b/Engine/Source/Runtime/MovieSceneTracks/Private/Tracks/MovieSceneCinematicShotTrack.cpp
@@ -112,7 +112,7 @@ bool UMovieSceneCinematicShotTrack::PopulateEvaluationTree(TMovieSceneEvaluation
 	{
 		UMovieSceneSection* Section = Sections[SectionIndex];
 
-		if (Section && Section->IsActive())
+		if (Section && Section->IsActive() && !IsRowEvalDisabled(Section->GetRowIndex()))
 		{
 			const TRange<FFrameNumber> SectionRange = Section->GetRange();
 			if (!SectionRange.IsEmpty())
