Sequencer: Remove DestroyConstructedComponents

From 5289466, this forced DestroyConstructedComponents because RerunConstructionScripts was getting called on spawning. That is no longer the case because SetActorLabel is no longer called. This fixes a bug where possessable components would be destroyed and renamed as TRASH, but then the object template would inherit those TRASH components.

diff --git a/Engine/Plugins/Experimental/DaySequence/Source/DaySequence/Private/DaySequenceActorSpawner.cpp b/Engine/Plugins/Experimental/DaySequence/Source/DaySequence/Private/DaySequenceActorSpawner.cpp
--- a/Engine/Plugins/Experimental/DaySequence/Source/DaySequence/Private/DaySequenceActorSpawner.cpp
+++ b/Engine/Plugins/Experimental/DaySequence/Source/DaySequence/Private/DaySequenceActorSpawner.cpp
@@ -156,9 +156,6 @@ UObject* FDaySequenceActorSpawner::SpawnObject(FMovieSceneSpawnable& Spawnable,
 		SpawnInfo.OverrideLevel = WorldContext->PersistentLevel;
 	}
 
-	//Chaos - Avoiding crash in UWorld::SendAllEndOfFrameUpdates due to duplicating template components/re-runing the construction script on a fully formed hierarchy
-	ObjectTemplate->DestroyConstructedComponents();
-
 	FTransform SpawnTransform;
 
 	if (USceneComponent* RootComponent = ObjectTemplate->GetRootComponent())
diff --git a/Engine/Source/Runtime/LevelSequence/Private/LevelSequenceActorSpawner.cpp b/Engine/Source/Runtime/LevelSequence/Private/LevelSequenceActorSpawner.cpp
--- a/Engine/Source/Runtime/LevelSequence/Private/LevelSequenceActorSpawner.cpp
+++ b/Engine/Source/Runtime/LevelSequence/Private/LevelSequenceActorSpawner.cpp
@@ -171,9 +171,6 @@ UObject* FLevelSequenceActorSpawner::SpawnObject(FMovieSceneSpawnable& Spawnable
 		SpawnInfo.OverrideLevel = WorldContext->PersistentLevel;
 	}
 
-	//Chaos - Avoiding crash in UWorld::SendAllEndOfFrameUpdates due to duplicating template components/re-runing the construction script on a fully formed hierarchy
-	ObjectTemplate->DestroyConstructedComponents();
-
 	FTransform SpawnTransform;
 
 	if (USceneComponent* RootComponent = ObjectTemplate->GetRootComponent())
diff --git a/Engine/Source/Runtime/MovieSceneTracks/Private/Bindings/MovieSceneSpawnableActorBinding.cpp b/Engine/Source/Runtime/MovieSceneTracks/Private/Bindings/MovieSceneSpawnableActorBinding.cpp
--- a/Engine/Source/Runtime/MovieSceneTracks/Private/Bindings/MovieSceneSpawnableActorBinding.cpp
+++ b/Engine/Source/Runtime/MovieSceneTracks/Private/Bindings/MovieSceneSpawnableActorBinding.cpp
@@ -92,9 +92,6 @@ UObject* UMovieSceneSpawnableActorBindingBase::SpawnObjectInternal(UWorld* World
 
 	if (ActorTemplate)
 	{
-		//Chaos - Avoiding crash in UWorld::SendAllEndOfFrameUpdates due to duplicating template components/re-runing the construction script on a fully formed hierarchy
-		ActorTemplate->DestroyConstructedComponents();
-
 		// Disable all particle components so that they don't auto fire as soon as the actor is spawned. The particles should be triggered through the particle track.
 		for (UActorComponent* Component : ActorTemplate->GetComponents())
 		{
