Sequencer: don't fall back to the double blender system when a property doesn't have blending information. Instead, fall back to the fast path where we don't blend anything.

diff --git a/Engine/Source/Runtime/MovieSceneTracks/Private/Systems/MovieScenePropertyInstantiator.cpp b/Engine/Source/Runtime/MovieSceneTracks/Private/Systems/MovieScenePropertyInstantiator.cpp
--- a/Engine/Source/Runtime/MovieSceneTracks/Private/Systems/MovieScenePropertyInstantiator.cpp
+++ b/Engine/Source/Runtime/MovieSceneTracks/Private/Systems/MovieScenePropertyInstantiator.cpp
@@ -817,8 +817,6 @@ UMovieScenePropertyInstantiatorSystem::FSetupBlenderSystemResult UMovieSceneProp
 {
 	using namespace UE::MovieScene;
 
-	FSetupBlenderSystemResult Result;
-
 	UClass* NewBlenderClass = nullptr;
 	int32 BlenderClassPriority = TNumericLimits<int32>::Lowest();
 
@@ -864,9 +862,7 @@ UMovieScenePropertyInstantiatorSystem::FSetupBlenderSystemResult UMovieSceneProp
 
 	if (!NewBlenderClass)
 	{
-		UE_LOG(LogMovieScene, Warning, TEXT("No default blender class specified on property, and no custom blender specified on entities. Falling back to double blender."));
-
-		NewBlenderClass = UMovieScenePiecewiseDoubleBlenderSystem::StaticClass();
+		return FSetupBlenderSystemResult();
 	}
 
 	FComponentTypeID BlenderTypeTag = GetDefault<UMovieSceneBlenderSystem>(NewBlenderClass)->GetBlenderTypeTag();
@@ -903,6 +899,16 @@ void UMovieScenePropertyInstantiatorSystem::InitializeBlendPath(const FPropertyP
 
 	FSetupBlenderSystemResult SetupResult = SetupBlenderSystem(Params);
 
+	// -----------------------------------------------------------------------------------------------------
+	// Situation 0: there is no blender system for this type of property. Normally, this means we can't
+	// end up trying to blend it, but it may happen if that property is inside a sub-sequence with hierarchical
+	// easing.
+	if (!SetupResult.CurrentInfo.BlenderSystemClass)
+	{
+		InitializeFastPath(Params);
+		return;
+	}
+
 	// -----------------------------------------------------------------------------------------------------
 	// Situation 1: New or modified contributors (inputs) but we're already set up for blending using the same system
 	if (Params.PropertyInfo->FinalBlendOutputID && SetupResult.CurrentInfo.BlenderSystemClass == SetupResult.PreviousInfo.BlenderSystemClass)
