Sequencer: don't blend enum values, and make the byte track use the enum system when it animates an enum

diff --git a/Engine/Source/Runtime/MovieSceneTracks/Private/MovieSceneTracksComponentTypes.cpp b/Engine/Source/Runtime/MovieSceneTracks/Private/MovieSceneTracksComponentTypes.cpp
--- a/Engine/Source/Runtime/MovieSceneTracks/Private/MovieSceneTracksComponentTypes.cpp
+++ b/Engine/Source/Runtime/MovieSceneTracks/Private/MovieSceneTracksComponentTypes.cpp
@@ -12,7 +12,6 @@
 #include "EntitySystem/MovieSceneBlenderSystem.h"
 #include "Systems/MovieScenePiecewiseBoolBlenderSystem.h"
 #include "Systems/MovieScenePiecewiseByteBlenderSystem.h"
-#include "Systems/MovieScenePiecewiseEnumBlenderSystem.h"
 #include "Systems/MovieScenePiecewiseIntegerBlenderSystem.h"
 #include "Systems/MovieScenePiecewiseDoubleBlenderSystem.h"
 #include "Systems/MovieScenePiecewiseDoubleBlenderSystem.h"
@@ -695,9 +694,9 @@ FMovieSceneTracksComponentTypes::FMovieSceneTracksComponentTypes()
 
 	// --------------------------------------------------------------------------------------------
 	// Set up enum properties
+	// No blender system specified, enums aren't supposed blend
 	BuiltInComponents->PropertyRegistry.DefineProperty(Enum, TEXT("Apply Enum Properties"))
 	.AddSoleChannel(BuiltInComponents->ByteResult)
-	.SetBlenderSystem<UMovieScenePiecewiseEnumBlenderSystem>()
 	.SetCustomAccessors(&Accessors.Enum)
 	.Commit();
 
diff --git a/Engine/Source/Runtime/MovieSceneTracks/Private/Sections/MovieSceneByteSection.cpp b/Engine/Source/Runtime/MovieSceneTracks/Private/Sections/MovieSceneByteSection.cpp
--- a/Engine/Source/Runtime/MovieSceneTracks/Private/Sections/MovieSceneByteSection.cpp
+++ b/Engine/Source/Runtime/MovieSceneTracks/Private/Sections/MovieSceneByteSection.cpp
@@ -5,6 +5,7 @@
 #include "Evaluation/MovieScenePropertyTemplates.h"
 #include "MovieSceneTracksComponentTypes.h"
 #include "Systems/MovieScenePiecewiseByteBlenderSystem.h"
+#include "Tracks/MovieSceneByteTrack.h"
 #include "Tracks/MovieScenePropertyTrack.h"
 #include "UObject/SequencerObjectVersion.h"
 
@@ -47,12 +48,24 @@ void UMovieSceneByteSection::ImportEntityImpl(UMovieSceneEntitySystemLinker* Ent
 		return;
 	}
 
+	const UMovieSceneByteTrack* ParentByteTrack = GetTypedOuter<UMovieSceneByteTrack>();
+	const UEnum* EnumClass = ParentByteTrack ? ParentByteTrack->GetEnum() : nullptr;
+
 	const FBuiltInComponentTypes* Components = FBuiltInComponentTypes::Get();
 	const FMovieSceneTracksComponentTypes* TracksComponents = FMovieSceneTracksComponentTypes::Get();
 
-	FPropertyTrackEntityImportHelper(TracksComponents->Byte)
-		.Add(Components->ByteChannel, &ByteCurve)
-		.Add(Components->BlenderType, UMovieScenePiecewiseByteBlenderSystem::StaticClass())
-		.Commit(this, Params, OutImportedEntity);
+	if (EnumClass)
+	{
+		FPropertyTrackEntityImportHelper(TracksComponents->Enum)
+			.Add(Components->ByteChannel, &ByteCurve)
+			.Commit(this, Params, OutImportedEntity);
+	}
+	else
+	{
+		FPropertyTrackEntityImportHelper(TracksComponents->Byte)
+			.Add(Components->ByteChannel, &ByteCurve)
+			.Add(Components->BlenderType, UMovieScenePiecewiseByteBlenderSystem::StaticClass())
+			.Commit(this, Params, OutImportedEntity);
+	}
 }
 
diff --git a/Engine/Source/Runtime/MovieSceneTracks/Private/Sections/MovieSceneEnumSection.cpp b/Engine/Source/Runtime/MovieSceneTracks/Private/Sections/MovieSceneEnumSection.cpp
--- a/Engine/Source/Runtime/MovieSceneTracks/Private/Sections/MovieSceneEnumSection.cpp
+++ b/Engine/Source/Runtime/MovieSceneTracks/Private/Sections/MovieSceneEnumSection.cpp
@@ -3,7 +3,6 @@
 #include "Sections/MovieSceneEnumSection.h"
 #include "Channels/MovieSceneChannelProxy.h"
 #include "MovieSceneTracksComponentTypes.h"
-#include "Systems/MovieScenePiecewiseEnumBlenderSystem.h"
 #include "Tracks/MovieScenePropertyTrack.h"
 #include "UObject/SequencerObjectVersion.h"
 #include "UObject/SequencerObjectVersion.h"
@@ -52,7 +51,6 @@ void UMovieSceneEnumSection::ImportEntityImpl(UMovieSceneEntitySystemLinker* Ent
 
 	FPropertyTrackEntityImportHelper(TracksComponents->Enum)
 		.Add(Components->ByteChannel, &EnumCurve)
-		.Add(Components->BlenderType, UMovieScenePiecewiseEnumBlenderSystem::StaticClass())
 		.Commit(this, Params, OutImportedEntity);
 }
 
