Ensure that evaluation context buffers have been flipped if an eval task is active prior to forced update

diff --git a/Engine/Source/Runtime/MovieSceneTracks/Private/Systems/MovieSceneSkeletalAnimationSystem.cpp b/Engine/Source/Runtime/MovieSceneTracks/Private/Systems/MovieSceneSkeletalAnimationSystem.cpp
--- a/Engine/Source/Runtime/MovieSceneTracks/Private/Systems/MovieSceneSkeletalAnimationSystem.cpp
+++ b/Engine/Source/Runtime/MovieSceneTracks/Private/Systems/MovieSceneSkeletalAnimationSystem.cpp
@@ -671,6 +671,7 @@ private:
 			(SkeletalMeshComponent->PoseTickedThisFrame() || (SequencerInstance && SequencerInstance->GetSourceAnimInstance() != ExistingAnimInstance))
 			)
 		{
+			SkeletalMeshComponent->HandleExistingParallelEvaluationTask(true, true);
 			SkeletalMeshComponent->TickAnimation(0.f, false);
 
 			SkeletalMeshComponent->RefreshBoneTransforms();
