Sequencer: Only perform post-duplicate fixup on replaceables that exist within an actual Sequence.

This allows bindings to still be copied around for the detail customizations

diff --git a/Engine/Source/Runtime/MovieSceneTracks/Private/Bindings/MovieSceneReplaceableDirectorBlueprintBinding.cpp b/Engine/Source/Runtime/MovieSceneTracks/Private/Bindings/MovieSceneReplaceableDirectorBlueprintBinding.cpp
--- a/Engine/Source/Runtime/MovieSceneTracks/Private/Bindings/MovieSceneReplaceableDirectorBlueprintBinding.cpp
+++ b/Engine/Source/Runtime/MovieSceneTracks/Private/Bindings/MovieSceneReplaceableDirectorBlueprintBinding.cpp
@@ -68,7 +68,7 @@ void UMovieSceneReplaceableDirectorBlueprintBinding::PostDuplicate(EDuplicateMod
 	// If we were duplicated into a different package we can't reference the old function any more
 	//    For now we null it out, but it would be good to copy the endpoint as well (we currently can't do that
 	//    because there is no way to generically access the correct director BP class from UMovieSceneSequence
-	if (DynamicBinding.Function)
+	if (DynamicBinding.Function && GetTypedOuter<UMovieScene>() != nullptr)
 	{
 		UPackage* FunctionPackage = DynamicBinding.Function->GetPackage();
 		UPackage* Package = GetPackage();
