Fix for audio track not playing in TakeRecorder sequence after initial play through.

#rnx

diff --git a/Engine/Source/Runtime/MovieSceneTracks/Private/Systems/MovieSceneAudioSystem.cpp b/Engine/Source/Runtime/MovieSceneTracks/Private/Systems/MovieSceneAudioSystem.cpp
--- a/Engine/Source/Runtime/MovieSceneTracks/Private/Systems/MovieSceneAudioSystem.cpp
+++ b/Engine/Source/Runtime/MovieSceneTracks/Private/Systems/MovieSceneAudioSystem.cpp
@@ -614,41 +614,38 @@ private:
 				}
 			}
 
-			if (EvaluationData.bAudioComponentHasBeenPlayed)
+			// Procedurally generated sounds don't have a defined duration so when the audio component is done, it's done
+			if (Sound->IsProcedurallyGenerated())
 			{
-				// Procedurally generated sounds don't have a defined duration so when the audio component is done, it's done
-				if (Sound->IsProcedurallyGenerated())
+				if (EvaluationData.bAudioComponentHasBeenPlayed && !AudioSection.GetLooping())
 				{
-					if (!AudioSection.GetLooping())
+					// If we're not a looping section and the AC is done, return
+					// otherwise a looping section will restart the sound if it's not playing
+					if (!AudioComponent.IsPlaying())
 					{
-						// If we're not a looping section and the AC is done, return
-						// otherwise a looping section will restart the sound if it's not playing
-						if (!AudioComponent.IsPlaying())
-						{
-							UE_LOG(LogMovieScene, Verbose, TEXT("Procedural sound Audio Component reached end of playback. Component: %s Sound: %s"), *AudioComponent.GetName(), *GetNameSafe(AudioComponent.Sound));
-							return;
-						}
+						UE_LOG(LogMovieScene, Verbose, TEXT("Procedural sound Audio Component reached end of playback. Component: %s Sound: %s"), *AudioComponent.GetName(), *GetNameSafe(AudioComponent.Sound));
+						return;
 					}
 				}
-				else
+			}
+			else
+			{
+				const float Duration = MovieSceneHelpers::GetSoundDuration(Sound);
+				if (!AudioSection.GetLooping() && AudioTime > Duration && Duration != 0.f)
 				{
-					const float Duration = MovieSceneHelpers::GetSoundDuration(Sound);
-					if (!AudioSection.GetLooping())
+					// If this is Non-ProcedurallyGenerated audio, and it's not looping then check to see if it needs to be stopped
+					if (AudioComponent.IsPlaying())
 					{
-						// If this is Non-ProcedurallyGenerated audio, and it's not looping then check to see if it needs to be stopped
-						if (AudioTime > Duration && AudioComponent.IsPlaying())
-						{
-							UE_LOG(LogMovieScene, Verbose, TEXT("Audio Component reached end of playback. Component: %s Sound: %s"), *AudioComponent.GetName(), *GetNameSafe(AudioComponent.Sound));
-							AudioComponent.Stop();
-						}
-
-						return;
-					}
-					else
-					{
-						// Wrap AudioTime according to duration for Non-ProcedurallyGenerated audio that is looping
-						AudioTime = Duration > 0.f ? FMath::Fmod(AudioTime, Duration) : AudioTime;
+						UE_LOG(LogMovieScene, Verbose, TEXT("Audio Component reached end of playback. Component: %s Sound: %s"), *AudioComponent.GetName(), *GetNameSafe(AudioComponent.Sound));
+						AudioComponent.Stop();
 					}
+
+					return;
+				}
+				else
+				{
+					// Wrap AudioTime according to duration for Non-ProcedurallyGenerated audio that is looping
+					AudioTime = Duration > 0.f ? FMath::Fmod(AudioTime, Duration) : AudioTime;
 				}
 			}
 		}
