Fix for sounds that aren't sound waves not scrubbing

#rnx

diff --git a/Engine/Source/Runtime/MovieSceneTracks/Private/Systems/MovieSceneAudioSystem.cpp b/Engine/Source/Runtime/MovieSceneTracks/Private/Systems/MovieSceneAudioSystem.cpp
--- a/Engine/Source/Runtime/MovieSceneTracks/Private/Systems/MovieSceneAudioSystem.cpp
+++ b/Engine/Source/Runtime/MovieSceneTracks/Private/Systems/MovieSceneAudioSystem.cpp
@@ -731,13 +731,13 @@ private:
 					UE_LOG(LogMovieScene, Verbose, TEXT("Audio Component stopped due to needing a state change bIsPlaying: %d bNeedsTimeSync: %d. Component: %s Sound: %s"), AudioComponent.IsPlaying(), bSoundNeedsTimeSync, *AudioComponent.GetName(), *GetNameSafe(AudioComponent.Sound));
 					AudioComponent.Stop();
 				}
+			}
 
-				// Only change the sound clip if it has actually changed. This calls Stop internally if needed.
-				if (AudioComponent.Sound != Sound)
-				{
-					UE_LOG(LogMovieScene, Verbose, TEXT("Audio Component calling SetSound due to new sound. Component: %s OldSound: %s NewSound: %s"), *AudioComponent.GetName(), *GetNameSafe(AudioComponent.Sound), *GetNameSafe(Sound));
-					AudioComponent.SetSound(Sound);
-				}
+			// Only change the sound clip if it has actually changed. This calls Stop internally if needed.
+			if (AudioComponent.Sound != Sound)
+			{
+				UE_LOG(LogMovieScene, Verbose, TEXT("Audio Component calling SetSound due to new sound. Component: %s OldSound: %s NewSound: %s"), *AudioComponent.GetName(), *GetNameSafe(AudioComponent.Sound), *GetNameSafe(Sound));
+				AudioComponent.SetSound(Sound);
 			}
 
 #if WITH_EDITOR
@@ -757,34 +757,36 @@ private:
 				UE_LOG(LogMovieScene, Verbose, TEXT("Audio Component Play at Local Time: %6.2f CurrentTime: %6.2f(s) SectionStart: %6.2f(s), SoundDur: %6.2f OffsetIntoClip: %6.2f Component: %s Sound: %s"), AudioTime, (EvalTime / Context.GetFrameRate()), SectionStartTimeSeconds, AudioComponent.Sound ? AudioComponent.Sound->GetDuration() : 0.0f, (float)Context.GetFrameRate().AsSeconds(AudioStartOffset), *AudioComponent.GetName(), *GetNameSafe(AudioComponent.Sound));
 				
 #if WITH_EDITOR
-				if (Context.GetStatus() == EMovieScenePlayerStatus::Scrubbing && bEnableGranularScrubbingCVar)
+				// We only want to perform granular scrubbing in the narrow case of a non-procedural sound wave
+				// Otherwise, we fallback to simply restarting the sound at the given audio time.
+				USoundWave* SoundWave = Cast<USoundWave>(Sound);
+				bool bPerformGranularScrubbing = bEnableGranularScrubbingCVar
+					&& (nullptr != SoundWave)
+					&& !SoundWave->IsProcedurallyGenerated()
+					&& (nullptr != ScrubbedSound)
+					&& Context.GetStatus() == EMovieScenePlayerStatus::Scrubbing;
+
+				if (bPerformGranularScrubbing)
 				{
-					if (USoundWave* SoundWave = Cast<USoundWave>(Sound))
+					// If we're not playing the audio component, then set up the audio component to use the scrubbed sound wave and play it
+					if (!AudioComponent.IsPlaying())
 					{
-						// Can't scrub a procedural asset right now
-						if (ScrubbedSound && !SoundWave->IsProcedurallyGenerated())
-						{
-							// If we're not playing the audio component, then set up the audio component to use the scrubbed sound wave and play it
-							if (!AudioComponent.IsPlaying())
-							{
-								ScrubbedSound->SetSoundWave(SoundWave);
-								ScrubbedSound->SetPlayheadTime(AudioTime);
-								float MaxScrubWidthSeconds = 0.001f * (float)ScrubWidthMillisecondsCVar;
-								ScrubbedSound->SetGrainDurationRange({MaxScrubWidthSeconds, 0.05f});
-
-								ScrubbedSound->SetIsScrubbing(true);
-								ScrubbedSound->SetIsScrubbingWhileStationary(bEnableGranularScrubbingWhileStationaryCVar);
-
-								AudioComponent.SetSound(ScrubbedSound);
-								AudioComponent.Play();
-							}
-							else
-							{
-								// If we're already playing, then just update the playhead time on the scrubbed sound
-								// This will propagate the playhead time to the rendering ISoundGenerator
-								ScrubbedSound->SetPlayheadTime(AudioTime);
-							}
-						}
+						ScrubbedSound->SetSoundWave(SoundWave);
+						ScrubbedSound->SetPlayheadTime(AudioTime);
+						float MaxScrubWidthSeconds = 0.001f * (float)ScrubWidthMillisecondsCVar;
+						ScrubbedSound->SetGrainDurationRange({ MaxScrubWidthSeconds, 0.05f });
+
+						ScrubbedSound->SetIsScrubbing(true);
+						ScrubbedSound->SetIsScrubbingWhileStationary(bEnableGranularScrubbingWhileStationaryCVar);
+
+						AudioComponent.SetSound(ScrubbedSound);
+						AudioComponent.Play();
+					}
+					else
+					{
+						// If we're already playing, then just update the playhead time on the scrubbed sound
+						// This will propagate the playhead time to the rendering ISoundGenerator
+						ScrubbedSound->SetPlayheadTime(AudioTime);
 					}
 				}
 				else
