Fix for procedural sound waves not restarting in sequencer if scrubbing/restarting the sequencer track

#rnx

diff --git a/Engine/Source/Runtime/MovieSceneTracks/Private/Systems/MovieSceneAudioSystem.cpp b/Engine/Source/Runtime/MovieSceneTracks/Private/Systems/MovieSceneAudioSystem.cpp
--- a/Engine/Source/Runtime/MovieSceneTracks/Private/Systems/MovieSceneAudioSystem.cpp
+++ b/Engine/Source/Runtime/MovieSceneTracks/Private/Systems/MovieSceneAudioSystem.cpp
@@ -589,7 +589,7 @@ private:
 
 		if (AudioTime >= 0.f && Sound)
 		{
-			if (bUseTimeDilationToAdjustment)
+			if (bUseTimeDilationToAdjustment || Sound->IsProcedurallyGenerated())
 			{
 				// Keep track of initial Audio and Context times when sound started playing.
 				if (!(EvaluationData.LastAudioTime.IsSet() && EvaluationData.LastContextTime.IsSet()))
@@ -609,6 +609,10 @@ private:
 					// Get Time Delta between previous time in sequencer context and current time, not taking into account dilation.
 					// Add to previous frame's audio time
 					AudioTime += (CurrContextTime - PrevContextTime) * (InverseTimeDilation);
+
+					// Evaluate if we've gone backward in the audio time, if so, flag that we've not played the sound yet so it retriggers
+					EvaluationData.bAudioComponentHasBeenPlayed = (EvaluationData.LastAudioTime.GetValue() > AudioTime);
+
 					EvaluationData.LastAudioTime = AudioTime;
 					EvaluationData.LastContextTime = CurrContextTime;
 				}
