Fix for granular scrubbing

#rnx

diff --git a/Engine/Source/Runtime/MovieSceneTracks/Private/Systems/MovieSceneAudioSystem.cpp b/Engine/Source/Runtime/MovieSceneTracks/Private/Systems/MovieSceneAudioSystem.cpp
--- a/Engine/Source/Runtime/MovieSceneTracks/Private/Systems/MovieSceneAudioSystem.cpp
+++ b/Engine/Source/Runtime/MovieSceneTracks/Private/Systems/MovieSceneAudioSystem.cpp
@@ -611,8 +611,10 @@ private:
 					AudioTime += (CurrContextTime - PrevContextTime) * (InverseTimeDilation);
 
 					// Evaluate if we've gone backward in the audio time, if so, flag that we've not played the sound yet so it retriggers
-					EvaluationData.bAudioComponentHasBeenPlayed = (EvaluationData.LastAudioTime.GetValue() > AudioTime);
-
+					if (EvaluationData.LastAudioTime.GetValue() > AudioTime)
+					{
+						EvaluationData.bAudioComponentHasBeenPlayed = false;
+					}
 					EvaluationData.LastAudioTime = AudioTime;
 					EvaluationData.LastContextTime = CurrContextTime;
 				}
@@ -737,13 +739,6 @@ private:
 				}
 			}
 
-			// Only change the sound clip if it has actually changed. This calls Stop internally if needed.
-			if (AudioComponent.Sound != Sound)
-			{
-				UE_LOG(LogMovieScene, Verbose, TEXT("Audio Component calling SetSound due to new sound. Component: %s OldSound: %s NewSound: %s"), *AudioComponent.GetName(), *GetNameSafe(AudioComponent.Sound), *GetNameSafe(Sound));
-				AudioComponent.SetSound(Sound);
-			}
-
 #if WITH_EDITOR
 			if (GIsEditor && World != nullptr && !World->IsPlayInEditor())
 			{
@@ -802,6 +797,13 @@ private:
 						ScrubbedSound->SetIsScrubbing(false);
 					}
 #endif // WITH_EDITOR
+
+					// Only change the sound clip if it has actually changed. This calls Stop internally if needed.
+					if (AudioComponent.Sound != Sound)
+					{
+						UE_LOG(LogMovieScene, Verbose, TEXT("Audio Component calling SetSound due to new sound. Component: %s OldSound: %s NewSound: %s"), *AudioComponent.GetName(), *GetNameSafe(AudioComponent.Sound), *GetNameSafe(Sound));
+						AudioComponent.SetSound(Sound);
+					}
 					
 					AudioComponent.Play(AudioTime);
 					
