Fix for procedural sound waves not restarting in sequencer if scrubbing/restarting the sequencer track

#rnx

diff --git a/Engine/Source/Runtime/MovieSceneTracks/Private/Systems/MovieSceneAudioSystem.cpp b/Engine/Source/Runtime/MovieSceneTracks/Private/Systems/MovieSceneAudioSystem.cpp
--- a/Engine/Source/Runtime/MovieSceneTracks/Private/Systems/MovieSceneAudioSystem.cpp
+++ b/Engine/Source/Runtime/MovieSceneTracks/Private/Systems/MovieSceneAudioSystem.cpp
@@ -589,7 +589,7 @@ private:
 
 		if (AudioTime >= 0.f && Sound)
 		{
-			if (bUseTimeDilationToAdjustment || Sound->IsProcedurallyGenerated())
+			if (bUseTimeDilationToAdjustment)
 			{
 				// Keep track of initial Audio and Context times when sound started playing.
 				if (!(EvaluationData.LastAudioTime.IsSet() && EvaluationData.LastContextTime.IsSet()))
@@ -610,11 +610,6 @@ private:
 					// Add to previous frame's audio time
 					AudioTime += (CurrContextTime - PrevContextTime) * (InverseTimeDilation);
 
-					// Evaluate if we've gone backward in the audio time, if so, flag that we've not played the sound yet so it retriggers
-					if (EvaluationData.LastAudioTime.GetValue() > AudioTime)
-					{
-						EvaluationData.bAudioComponentHasBeenPlayed = false;
-					}
 					EvaluationData.LastAudioTime = AudioTime;
 					EvaluationData.LastContextTime = CurrContextTime;
 				}
@@ -623,16 +618,23 @@ private:
 			// Procedurally generated sounds don't have a defined duration so when the audio component is done, it's done
 			if (Sound->IsProcedurallyGenerated())
 			{
-				if (EvaluationData.bAudioComponentHasBeenPlayed && !AudioSection.GetLooping())
+				if (Context.GetStatus() == EMovieScenePlayerStatus::Playing)
 				{
-					// If we're not a looping section and the AC is done, return
-					// otherwise a looping section will restart the sound if it's not playing
-					if (!AudioComponent.IsPlaying())
+					if (EvaluationData.bAudioComponentHasBeenPlayed && !AudioSection.GetLooping())
 					{
-						UE_LOG(LogMovieScene, Verbose, TEXT("Procedural sound Audio Component reached end of playback. Component: %s Sound: %s"), *AudioComponent.GetName(), *GetNameSafe(AudioComponent.Sound));
-						return;
+						// If we're not a looping section and the AC is done, return
+						// otherwise a looping section will restart the sound if it's not playing
+						if (!AudioComponent.IsPlaying())
+						{
+							UE_LOG(LogMovieScene, Verbose, TEXT("Procedural sound Audio Component reached end of playback. Component: %s Sound: %s"), *AudioComponent.GetName(), *GetNameSafe(AudioComponent.Sound));
+							return;
+						}
 					}
 				}
+				else
+				{
+					EvaluationData.bAudioComponentHasBeenPlayed = false;
+				}
 			}
 			else
 			{
@@ -736,6 +738,12 @@ private:
 				{
 					UE_LOG(LogMovieScene, Verbose, TEXT("Audio Component stopped due to needing a state change bIsPlaying: %d bNeedsTimeSync: %d. Component: %s Sound: %s"), AudioComponent.IsPlaying(), bSoundNeedsTimeSync, *AudioComponent.GetName(), *GetNameSafe(AudioComponent.Sound));
 					AudioComponent.Stop();
+
+					UE_LOG(LogMovieScene, Verbose, TEXT("AudioComponent.Stop()"));
+					if (Context.GetStatus() != EMovieScenePlayerStatus::Playing)
+					{
+						EvaluationData.bAudioComponentHasBeenPlayed = false;
+					}
 				}
 			}
 
@@ -807,8 +815,11 @@ private:
 					
 					AudioComponent.Play(AudioTime);
 					
-					// Set that we've played an audio component. This is used by procedural sounds who have undefined duration.
-					EvaluationData.bAudioComponentHasBeenPlayed = true;
+					if (Context.GetStatus() == EMovieScenePlayerStatus::Playing)
+					{
+						// Set that we've played an audio component. This is used by procedural sounds who have undefined duration.
+						EvaluationData.bAudioComponentHasBeenPlayed = true;
+					}
 
 					// Keep track of when we asked this audio clip to play (in game time) so that we can figure 
 					// out if there's a significant desync in the future.
@@ -1267,6 +1278,8 @@ void UMovieSceneAudioSystem::StopSound(FInstanceHandle InstanceHandle, FObjectKe
 			{
 				AudioComponent->Stop();
 			}
+
+			Data->bAudioComponentHasBeenPlayed = false;
 		}
 	}
 }
