Assemble the Movie Scene Audio Section float channels into entities in a way that does not stomp on the existing pitch and volume channels when using Metasounds in level sequences.

diff --git a/Engine/Source/Runtime/MovieSceneTracks/Private/Sections/MovieSceneAudioSection.cpp b/Engine/Source/Runtime/MovieSceneTracks/Private/Sections/MovieSceneAudioSection.cpp
--- a/Engine/Source/Runtime/MovieSceneTracks/Private/Sections/MovieSceneAudioSection.cpp
+++ b/Engine/Source/Runtime/MovieSceneTracks/Private/Sections/MovieSceneAudioSection.cpp
@@ -47,6 +47,9 @@ namespace
 namespace UE::MovieScene
 {
 
+// Built-in float channels used by audio tracks
+enum class EBuiltinFloatChannel : uint8 { VolumeChannel, PitchChannel, Count };
+
 /* 
  * Entity IDs are an encoded type and index, with the upper bit being the type (scalar iputs vs audio trigger),
  * and the lower 31 bits as the entity index
@@ -548,8 +551,8 @@ void UMovieSceneAudioSection::ImportEntityImpl(UMovieSceneEntitySystemLinker* En
 			.AddConditional(BuiltInComponents->GenericObjectBinding, ObjectBindingID, ObjectBindingID.IsValid())
 			.AddTagConditional(BuiltInComponents->Tags.Root, !ObjectBindingID.IsValid())
 			.Add(TrackComponents->Audio, FMovieSceneAudioComponentData{ this })
-			.Add(BuiltInComponents->FloatChannel[0], &SoundVolume)
-			.Add(BuiltInComponents->FloatChannel[1], &PitchMultiplier)
+			.Add(BuiltInComponents->FloatChannel[static_cast<int32>(EBuiltinFloatChannel::VolumeChannel)], &SoundVolume)
+			.Add(BuiltInComponents->FloatChannel[static_cast<int32>(EBuiltinFloatChannel::PitchChannel)], &PitchMultiplier)
 		);
 	}
 	else if (EntityType == EAudioSectionEntityType::InputsEntity)
@@ -560,14 +563,22 @@ void UMovieSceneAudioSection::ImportEntityImpl(UMovieSceneEntitySystemLinker* En
 
 		// There are up to 9 float channels per entity, and we know that all entities are fully packed
 		// until the last one. So add as many as we can for the given entity we're building.
-		int32 FloatInputStartIndex = EntityIndex * 9;
+
+		// In addition, there are two float channels allocated for SoundVolume and PitchMultiplier for each section, and we need to not
+		// overwrite those channels with float inputs, therefore each entity can only have 9 - 2 = 7 unique float channels per entity and
+		// they must be written in the slots 2 - 8 (ex. FloatInputs[9] = {"", "", "In0", "In1", "In2", "In3", "In4", "In5", "In6"}
+
+		constexpr int32 NumFloatChannelsPerEntity = 9;
+		constexpr int32 NumUniqueFloatChannelsPerEntity = NumFloatChannelsPerEntity - static_cast<int32>(EBuiltinFloatChannel::Count);
+
+		int32 FloatInputStartIndex = EntityIndex * NumUniqueFloatChannelsPerEntity;
 		if (FloatInputStartIndex < Inputs_Float.Num())
 		{
-			int32 FloatInputNum = FMath::Min(9, Inputs_Float.Num() - FloatInputStartIndex);
+			int32 FloatInputNum = FMath::Min(NumUniqueFloatChannelsPerEntity, Inputs_Float.Num() - FloatInputStartIndex);
 			Inputs_Float.GetKeys(InputNames);
 			for (int32 Offset = 0; Offset < FloatInputNum; ++Offset)
 			{
-				InputData.FloatInputs[Offset] = InputNames[FloatInputStartIndex + Offset];
+				InputData.FloatInputs[Offset + static_cast<int32>(EBuiltinFloatChannel::Count)] = InputNames[FloatInputStartIndex + Offset];
 			}
 		}
 
@@ -601,7 +612,7 @@ void UMovieSceneAudioSection::ImportEntityImpl(UMovieSceneEntitySystemLinker* En
 			.Add(TrackComponents->Audio, FMovieSceneAudioComponentData{ this })
 			.Add(TrackComponents->AudioInputs, InputData)
 		);
-		for (int32 Index = 0; Index < 9; ++Index)
+		for (int32 Index = static_cast<int32>(EBuiltinFloatChannel::Count); Index < NumFloatChannelsPerEntity; ++Index)
 		{
 			FName InputName = InputData.FloatInputs[Index];
 			if (!InputName.IsNone())
