Sequencer: Use localized text for property track tooltip

diff --git a/Engine/Source/Runtime/MovieSceneTracks/Private/Tracks/MovieScenePropertyTrack.cpp b/Engine/Source/Runtime/MovieSceneTracks/Private/Tracks/MovieScenePropertyTrack.cpp
--- a/Engine/Source/Runtime/MovieSceneTracks/Private/Tracks/MovieScenePropertyTrack.cpp
+++ b/Engine/Source/Runtime/MovieSceneTracks/Private/Tracks/MovieScenePropertyTrack.cpp
@@ -10,10 +10,13 @@
 #include "MovieSceneTracksComponentTypes.h"
 #include "PropertyPathHelpers.h"
 #include "Systems/MovieScenePiecewiseBoolBlenderSystem.h"
+#include "ObjectEditorUtils.h"
 #include "UObject/Field.h"
 
 #include UE_INLINE_GENERATED_CPP_BY_NAME(MovieScenePropertyTrack)
 
+#define LOCTEXT_NAMESPACE "MovieScenePropertyTrack"
+
 UMovieScenePropertyTrack::UMovieScenePropertyTrack(const FObjectInitializer& ObjectInitializer)
 	: Super(ObjectInitializer)
 {
@@ -82,19 +85,10 @@ FText UMovieScenePropertyTrack::GetDisplayNameToolTipText(const FMovieSceneLabel
 		FTrackInstancePropertyBindings InstancePropertyBinding(GetPropertyName(), GetPropertyPath().ToString());
 		if (FProperty* BoundProperty = InstancePropertyBinding.GetProperty(*BoundObject))
 		{
-			FString PropertyName = BoundProperty->GetMetaData(TEXT("DisplayName"));
-			if (PropertyName.IsEmpty())
-			{
-				PropertyName = BoundProperty->GetName();
-			}
-
-			FString CategoryName = BoundProperty->GetMetaData(TEXT("Category")).Replace(TEXT("|"), TEXT(" \u00BB "));
-			if (!CategoryName.IsEmpty())
-			{
-				CategoryName.Append(TEXT(" \u00BB "));
-			}
-
-			return FText::FromString(FString::Printf(TEXT("%s%s\n(Path: %s)"), *CategoryName, *PropertyName, *InstancePropertyBinding.GetPropertyPath()));
+			return FText::Format(LOCTEXT("DisplayNameTooltipFormat", "{0} \u00BB {1}\n(Path: {2})"),
+				FObjectEditorUtils::GetCategoryText(BoundProperty),
+				BoundProperty->GetDisplayNameText(),
+				FText::FromString(InstancePropertyBinding.GetPropertyPath()));
 		}
 	}
 	
@@ -560,3 +554,4 @@ FName FMovieScenePropertyTrackEntityImportHelper::SanitizeBoolPropertyName(FName
 	return FName(*PropertyVarName);
 }
 
+#undef LOCTEXT_NAMESPACE
