Sequencer: Mark FMovieScene3DLocationKeyStruct as FMovieScene3DScaleKeyStruct as Spatial so that they can receive the LUF treatment

Also, FMovieScene3DScaleKeyStruct should use FVector and not FVector3f, looks like this got missed in the LWC upgrade

#rnx

diff --git a/Engine/Source/Runtime/MovieSceneTracks/Public/Sections/MovieScene3DTransformSection.h b/Engine/Source/Runtime/MovieSceneTracks/Public/Sections/MovieScene3DTransformSection.h
--- a/Engine/Source/Runtime/MovieSceneTracks/Public/Sections/MovieScene3DTransformSection.h
+++ b/Engine/Source/Runtime/MovieSceneTracks/Public/Sections/MovieScene3DTransformSection.h
@@ -42,7 +42,7 @@ struct FMovieScene3DLocationKeyStruct
 	GENERATED_BODY()
 
 	/** The key's translation value. */
-	UPROPERTY(EditAnywhere, Category=Key)
+	UPROPERTY(EditAnywhere, Category=Key, meta=(Spatial))
 	FVector Location = FVector::ZeroVector;
 
 	/** The key's time. */
@@ -89,8 +89,8 @@ struct FMovieScene3DScaleKeyStruct
 	GENERATED_BODY()
 
 	/** The key's scale value. */
-	UPROPERTY(EditAnywhere, Category=Key)
-	FVector3f Scale = FVector3f::OneVector;
+	UPROPERTY(EditAnywhere, Category=Key, meta=(Spatial))
+	FVector Scale = FVector::OneVector;
 
 	/** The key's time. */
 	UPROPERTY(EditAnywhere, Category=Key)
@@ -113,7 +113,7 @@ struct FMovieScene3DTransformKeyStruct
 	GENERATED_BODY()
 
 	/** The key's translation value. */
-	UPROPERTY(EditAnywhere, Category=Key)
+	UPROPERTY(EditAnywhere, Category=Key, meta=(Spatial))
 	FVector Location = FVector::ZeroVector;
 
 	/** The key's rotation value. */
@@ -121,7 +121,7 @@ struct FMovieScene3DTransformKeyStruct
 	FRotator Rotation = FRotator::ZeroRotator;
 
 	/** The key's scale value. */
-	UPROPERTY(EditAnywhere, Category=Key)
+	UPROPERTY(EditAnywhere, Category=Key, meta=(Spatial))
 	FVector3f Scale = FVector3f::OneVector;
 
 	/** The key's time. */
