Remove Scale from Transform Origin overrides for level sequences

diff --git a/Engine/Source/Runtime/MovieSceneTracks/Private/Systems/MovieSceneTransformOriginSystem.cpp b/Engine/Source/Runtime/MovieSceneTracks/Private/Systems/MovieSceneTransformOriginSystem.cpp
--- a/Engine/Source/Runtime/MovieSceneTracks/Private/Systems/MovieSceneTransformOriginSystem.cpp
+++ b/Engine/Source/Runtime/MovieSceneTracks/Private/Systems/MovieSceneTransformOriginSystem.cpp
@@ -56,6 +56,8 @@ struct FGatherTransformOrigin
 				{
 					// Retrieve the current origin
 					FTransform TransformOrigin = RawInterface ? RawInterface->GetTransformOrigin() : IMovieSceneTransformOrigin::Execute_BP_GetTransformOrigin(InstanceData);
+					// Ignore scale. Sequencer does not apply scale, but the scale is factored into rotation, so assume always a scale of 1.
+					TransformOrigin.SetScale3D(FVector::OneVector);
 
 					TransformOriginsByInstanceID->Insert(Index, TransformOrigin);
 				}
