Sequencer: Mark FMovieScene3DLocationKeyStruct as FMovieScene3DScaleKeyStruct as Spatial so that they can receive the LUF treatment

Also, FMovieScene3DScaleKeyStruct should use FVector and not FVector3f, looks like this got missed in the LWC upgrade

#rnx

diff --git a/Engine/Source/Runtime/MovieSceneTracks/Public/Sections/MovieScene3DTransformSection.h b/Engine/Source/Runtime/MovieSceneTracks/Public/Sections/MovieScene3DTransformSection.h
--- a/Engine/Source/Runtime/MovieSceneTracks/Public/Sections/MovieScene3DTransformSection.h
+++ b/Engine/Source/Runtime/MovieSceneTracks/Public/Sections/MovieScene3DTransformSection.h
@@ -122,7 +122,7 @@ struct FMovieScene3DTransformKeyStruct
 
 	/** The key's scale value. */
 	UPROPERTY(EditAnywhere, Category=Key, meta=(Spatial))
-	FVector3f Scale = FVector3f::OneVector;
+	FVector Scale = FVector::OneVector;
 
 	/** The key's time. */
 	UPROPERTY(EditAnywhere, Category=Key)
