Improved Sorting for transform origin hierarchy. This resolves an issue where shots deep within a sequence hierarchy may not have the correct transform origin applied, due to their granparent, and higher ancestor sequences weren't taken into account when sorting. Sorting is now done by depth.

diff --git a/Engine/Source/Runtime/MovieSceneTracks/Private/Systems/MovieSceneTransformOriginSystem.cpp b/Engine/Source/Runtime/MovieSceneTracks/Private/Systems/MovieSceneTransformOriginSystem.cpp
--- a/Engine/Source/Runtime/MovieSceneTracks/Private/Systems/MovieSceneTransformOriginSystem.cpp
+++ b/Engine/Source/Runtime/MovieSceneTracks/Private/Systems/MovieSceneTransformOriginSystem.cpp
@@ -340,7 +340,9 @@ void UMovieSceneTransformOriginSystem::OnSchedulePersistentTasks(UE::MovieScene:
 	.Iterate_PerEntity(&Linker->EntityManager,  [this, InstanceRegistry](FRootInstanceHandle RootInstance, FInstanceHandle Instance, FMovieSceneSequenceID SequenceID)
 	{
 		FInstanceHandle SubInstanceHandle = InstanceRegistry->GetInstance(RootInstance).FindSubInstance(SequenceID);
-		InstanceHandleToParentHandle.AddUnique(FInstanceToParentPair(SubInstanceHandle, Instance));
+		FSubSequencePath SubSequencePath = InstanceRegistry->GetInstance(SubInstanceHandle).GetSubSequencePath();
+		int32 Depth = SubSequencePath.NumGenerationsFromRoot(SequenceID);
+		InstanceHandleToParentHandle.AddUnique(FInstanceToParentPair(SubInstanceHandle, Instance, Depth));
 		SequenceIDToInstanceHandle.Add(SequenceID, SubInstanceHandle);
 		
 	});
diff --git a/Engine/Source/Runtime/MovieSceneTracks/Public/Systems/MovieSceneTransformOriginSystem.h b/Engine/Source/Runtime/MovieSceneTracks/Public/Systems/MovieSceneTransformOriginSystem.h
--- a/Engine/Source/Runtime/MovieSceneTracks/Public/Systems/MovieSceneTransformOriginSystem.h
+++ b/Engine/Source/Runtime/MovieSceneTracks/Public/Systems/MovieSceneTransformOriginSystem.h
@@ -16,21 +16,18 @@ struct FInstanceToParentPair
 	// In this struct representation the Child represents this sequence.
 	UE::MovieScene::FInstanceHandle Child;
 	UE::MovieScene::FInstanceHandle Parent;
+	int32 Depth = 0;
 
-	FInstanceToParentPair(const UE::MovieScene::FInstanceHandle Child, const UE::MovieScene::FInstanceHandle Parent)
+	FInstanceToParentPair(const UE::MovieScene::FInstanceHandle Child, const UE::MovieScene::FInstanceHandle Parent, const int32 Depth)
 	{
 		this->Child = Child;
 		this->Parent = Parent;
+		this->Depth = Depth;
 	}
 
 	bool operator<(const FInstanceToParentPair& Other) const
 	{
-		// If this is the parent of Other, then it should be first in the mapping.
-		if(Child == Other.Parent)
-		{
-			return true;
-		}
-		return false;
+		return Depth < Other.Depth;
 	}
 
 	bool operator==(const FInstanceToParentPair& Other) const
