Sequencer: Fix Material Parameter Collection track display name so that it's not set explicitly on creation. Instead, the display name is just the name of the MPC if the user hasn't changed it. This allows the user to change the MPC and the name will automatically change.

diff --git a/Engine/Source/Editor/MovieSceneTools/Private/TrackEditors/MaterialParameterCollectionTrackEditor.cpp b/Engine/Source/Editor/MovieSceneTools/Private/TrackEditors/MaterialParameterCollectionTrackEditor.cpp
--- a/Engine/Source/Editor/MovieSceneTools/Private/TrackEditors/MaterialParameterCollectionTrackEditor.cpp
+++ b/Engine/Source/Editor/MovieSceneTools/Private/TrackEditors/MaterialParameterCollectionTrackEditor.cpp
@@ -166,7 +166,6 @@ void FMaterialParameterCollectionTrackEditor::AddTrackToSequence(const FAssetDat
 
 	Track->AddSection(*NewSection);
 	Track->MPC = MPC;
-	Track->SetDisplayName(FText::FromString(MPC->GetName()));
 
 	if (GetSequencer().IsValid())
 	{
@@ -222,7 +221,6 @@ TSharedPtr<SWidget> FMaterialParameterCollectionTrackEditor::BuildOutlinerEditCo
 
 		FScopedTransaction Transaction(LOCTEXT("SetAssetTransaction", "Assign Material Parameter Collection"));
 		MPCTrack->Modify();
-		MPCTrack->SetDisplayName(MPC ? FText::FromString(MPC->GetName()) : FText::GetEmpty());
 		MPCTrack->MPC = MPC;
 	};
 
diff --git a/Engine/Source/Runtime/MovieSceneTracks/Private/Tracks/MovieSceneMaterialParameterCollectionTrack.cpp b/Engine/Source/Runtime/MovieSceneTracks/Private/Tracks/MovieSceneMaterialParameterCollectionTrack.cpp
--- a/Engine/Source/Runtime/MovieSceneTracks/Private/Tracks/MovieSceneMaterialParameterCollectionTrack.cpp
+++ b/Engine/Source/Runtime/MovieSceneTracks/Private/Tracks/MovieSceneMaterialParameterCollectionTrack.cpp
@@ -5,6 +5,7 @@
 
 #include "Evaluation/MovieSceneEvaluationField.h"
 #include "EntitySystem/BuiltInComponentTypes.h"
+#include "MovieSceneSequence.h"
 #include "MovieSceneTracksComponentTypes.h"
 #include "Sections/MovieSceneComponentMaterialParameterSection.h"
 
@@ -100,9 +101,30 @@ bool UMovieSceneMaterialParameterCollectionTrack::SupportsType(TSubclassOf<UMovi
 #if WITH_EDITORONLY_DATA
 FText UMovieSceneMaterialParameterCollectionTrack::GetDefaultDisplayName() const
 {
+	if (MPC)
+	{
+		return FText::FromString(MPC->GetName());
+	}
+
 	return LOCTEXT("DefaultTrackName", "Material Parameter Collection");
 }
 #endif
 
+#if WITH_EDITOR
+
+void UMovieSceneMaterialParameterCollectionTrack::PostLoad()
+{
+	if (MPC && GetDisplayName().EqualTo(FText::FromString(MPC->GetName())))
+	{
+		UMovieSceneSequence* Sequence = GetTypedOuter<UMovieSceneSequence>();
+		UE_LOG(LogMovieScene, Display, TEXT("Material Parameter Collection Track: Removing display name %s from %s because it is identical to the material name and track names are now derived from the material name"), *MPC->GetName(), *GetNameSafe(Sequence));
+		SetDisplayName(FText::GetEmpty());
+	}
+
+	Super::PostLoad();
+}
+
+#endif
+
 #undef LOCTEXT_NAMESPACE
 
diff --git a/Engine/Source/Runtime/MovieSceneTracks/Public/Tracks/MovieSceneMaterialParameterCollectionTrack.h b/Engine/Source/Runtime/MovieSceneTracks/Public/Tracks/MovieSceneMaterialParameterCollectionTrack.h
--- a/Engine/Source/Runtime/MovieSceneTracks/Public/Tracks/MovieSceneMaterialParameterCollectionTrack.h
+++ b/Engine/Source/Runtime/MovieSceneTracks/Public/Tracks/MovieSceneMaterialParameterCollectionTrack.h
@@ -42,4 +42,11 @@ public:
 #if WITH_EDITORONLY_DATA
 	MOVIESCENETRACKS_API virtual FText GetDefaultDisplayName() const override;
 #endif
+
+protected:
+	//~ UObject interface
+#if WITH_EDITOR
+	MOVIESCENETRACKS_API virtual void PostLoad() override;
+#endif
+
 };
