Fixing a crash when opening a level sequence with an unsaved subsequence.
Now instance IDs are only added to the SequenceIDToInstanceHandle map if the instance ID is valid.
The instance ID is no longer used to get the subsequence path if it is invalid, which led to a crash. Before the fix that introduced the crash, however, invalid sequence IDs could still be added to the map. This fix prevents both that and the newly introduced crash.

diff --git a/Engine/Source/Runtime/MovieSceneTracks/Private/Systems/MovieSceneTransformOriginSystem.cpp b/Engine/Source/Runtime/MovieSceneTracks/Private/Systems/MovieSceneTransformOriginSystem.cpp
--- a/Engine/Source/Runtime/MovieSceneTracks/Private/Systems/MovieSceneTransformOriginSystem.cpp
+++ b/Engine/Source/Runtime/MovieSceneTracks/Private/Systems/MovieSceneTransformOriginSystem.cpp
@@ -339,11 +339,14 @@ void UMovieSceneTransformOriginSystem::OnSchedulePersistentTasks(UE::MovieScene:
 	.FilterNone({BuiltInComponents->Tags.ImportedEntity}) // filter out parent entities, otherwise we'd double-up in the InstanceHandleToParentHandle mapping.
 	.Iterate_PerEntity(&Linker->EntityManager,  [this, InstanceRegistry](FRootInstanceHandle RootInstance, FInstanceHandle Instance, FMovieSceneSequenceID SequenceID)
 	{
-		FInstanceHandle SubInstanceHandle = InstanceRegistry->GetInstance(RootInstance).FindSubInstance(SequenceID);
-		FSubSequencePath SubSequencePath = InstanceRegistry->GetInstance(SubInstanceHandle).GetSubSequencePath();
-		int32 Depth = SubSequencePath.NumGenerationsFromRoot(SequenceID);
-		InstanceHandleToParentHandle.AddUnique(FInstanceToParentPair(SubInstanceHandle, Instance, Depth));
-		SequenceIDToInstanceHandle.Add(SequenceID, SubInstanceHandle);
+		const FInstanceHandle SubInstanceHandle = InstanceRegistry->GetInstance(RootInstance).FindSubInstance(SequenceID);
+		if (InstanceRegistry->IsHandleValid(SubInstanceHandle))
+		{
+			const FSubSequencePath SubSequencePath = InstanceRegistry->GetInstance(SubInstanceHandle).GetSubSequencePath();
+			const int32 Depth = SubSequencePath.NumGenerationsFromRoot(SequenceID);
+			InstanceHandleToParentHandle.AddUnique(FInstanceToParentPair(SubInstanceHandle, Instance, Depth));
+			SequenceIDToInstanceHandle.Add(SequenceID, SubInstanceHandle);
+		}
 		
 	});
 
