Sequencer: Speculative fix for null PlaybackContext in Time Dilation track

diff --git a/Engine/Source/Runtime/MovieSceneTracks/Private/Systems/MovieSceneSlomoSystem.cpp b/Engine/Source/Runtime/MovieSceneTracks/Private/Systems/MovieSceneSlomoSystem.cpp
--- a/Engine/Source/Runtime/MovieSceneTracks/Private/Systems/MovieSceneSlomoSystem.cpp
+++ b/Engine/Source/Runtime/MovieSceneTracks/Private/Systems/MovieSceneSlomoSystem.cpp
@@ -43,7 +43,13 @@ struct FPreAnimatedSlomoState
 
 	static FPreAnimatedSlomoState SaveState(UObject* PlaybackContext)
 	{
-		if (AWorldSettings* WorldSettings = PlaybackContext->GetWorld()->GetWorldSettings())
+		UWorld* World = PlaybackContext ? PlaybackContext->GetWorld() : nullptr;
+		if (!World)
+		{
+			return FPreAnimatedSlomoState();
+		}
+
+		if (AWorldSettings* WorldSettings = World->GetWorldSettings())
 		{
 			return FPreAnimatedSlomoState{ WorldSettings->CinematicTimeDilation };
 		}
