Sequencer: store camera shake previewers as shared pointers because FGCObject classes are not relocatabble

diff --git a/Engine/Source/Runtime/MovieSceneTracks/Private/Systems/MovieSceneCameraShakeSystem.cpp b/Engine/Source/Runtime/MovieSceneTracks/Private/Systems/MovieSceneCameraShakeSystem.cpp
--- a/Engine/Source/Runtime/MovieSceneTracks/Private/Systems/MovieSceneCameraShakeSystem.cpp
+++ b/Engine/Source/Runtime/MovieSceneTracks/Private/Systems/MovieSceneCameraShakeSystem.cpp
@@ -120,7 +120,7 @@ struct FPreAnimatedCameraSourceShakeTraits : FBoundObjectPreAnimatedStateTraits
 			FCameraShakePreviewerLinkerExtension* PreviewerExtension = Params.Linker->FindExtension<FCameraShakePreviewerLinkerExtension>();
 			if (PreviewerExtension)
 			{
-				if (FCameraShakePreviewer* Previewer = PreviewerExtension->FindPreviewer(Params.TerminalInstanceHandle))
+				if (TSharedPtr<FCameraShakePreviewer> Previewer = PreviewerExtension->FindPreviewer(Params.TerminalInstanceHandle))
 				{
 					Previewer->RemoveAllCameraShakesFromSource(ShakeSourceComponent);
 				}
@@ -186,22 +186,26 @@ FCameraShakePreviewerLinkerExtension::~FCameraShakePreviewerLinkerExtension()
 		GEditor->OnLevelViewportClientListChanged().RemoveAll(this);
 	}
 	
-	for (TPair<FInstanceHandle, FCameraShakePreviewer>& Pair : Previewers)
+	for (TPair<FInstanceHandle, TSharedRef<FCameraShakePreviewer>>& Pair : Previewers)
 	{
-		FCameraShakePreviewer& Previewer = Pair.Value;
-		Previewer.UnRegisterViewModifiers();
+		TSharedRef<FCameraShakePreviewer> Previewer = Pair.Value;
+		Previewer->UnRegisterViewModifiers();
 	}
 	Previewers.Reset();
 }
 
-FCameraShakePreviewer* FCameraShakePreviewerLinkerExtension::FindPreviewer(FInstanceHandle InstanceHandle)
+TSharedPtr<FCameraShakePreviewer> FCameraShakePreviewerLinkerExtension::FindPreviewer(FInstanceHandle InstanceHandle)
 {
-	return Previewers.Find(InstanceHandle);
+	if (TSharedRef<FCameraShakePreviewer>* FoundItem = Previewers.Find(InstanceHandle))
+	{
+		return FoundItem->ToSharedPtr();
+	}
+	return TSharedPtr<FCameraShakePreviewer>();
 }
 
-FCameraShakePreviewer& FCameraShakePreviewerLinkerExtension::GetPreviewer(FInstanceHandle InstanceHandle)
+TSharedRef<FCameraShakePreviewer> FCameraShakePreviewerLinkerExtension::GetPreviewer(FInstanceHandle InstanceHandle)
 {
-	if (FCameraShakePreviewer* Previewer = Previewers.Find(InstanceHandle))
+	if (TSharedRef<FCameraShakePreviewer>* Previewer = Previewers.Find(InstanceHandle))
 	{
 		return *Previewer;
 	}
@@ -218,13 +222,14 @@ FCameraShakePreviewer& FCameraShakePreviewerLinkerExtension::GetPreviewer(FInsta
 	UObject* PlaybackContext = SequenceInstance.GetSharedPlaybackState()->GetPlaybackContext();
 	UWorld* ContextWorld = PlaybackContext ? PlaybackContext->GetWorld() : nullptr;
 
-	FCameraShakePreviewer& NewPreviewer = Previewers.Emplace(InstanceHandle, ContextWorld);
-	NewPreviewer.RegisterViewModifiers([ContextWorld](FLevelEditorViewportClient* LevelVC) -> bool
+	TSharedRef<FCameraShakePreviewer> NewPreviewer = MakeShared<FCameraShakePreviewer>(ContextWorld);
+	NewPreviewer->RegisterViewModifiers([ContextWorld](FLevelEditorViewportClient* LevelVC) -> bool
 				{
 					return LevelVC->AllowsCinematicControl() && LevelVC->GetWorld() == ContextWorld;
 				},
 				// Pass false to help the Mac compiler along.
 				false);
+	Previewers.Add(InstanceHandle, NewPreviewer);
 	return NewPreviewer;
 }
 
@@ -236,8 +241,9 @@ void FCameraShakePreviewerLinkerExtension::UpdateAllPreviewers()
 	for (auto It = Instances.CreateConstIterator(); It; ++It)
 	{
 		FInstanceHandle InstanceHandle = It->GetInstanceHandle();
-		if (FCameraShakePreviewer* Previewer = Previewers.Find(InstanceHandle))
+		if (TSharedRef<FCameraShakePreviewer>* FoundItem = Previewers.Find(InstanceHandle))
 		{
+			TSharedRef<FCameraShakePreviewer> Previewer(*FoundItem);
 			const FMovieSceneContext& Context = It->GetContext();
 			const float DeltaTime = Context.GetFrameRate().AsSeconds(Context.GetDelta());
 			if (DeltaTime > 0.f)
@@ -257,9 +263,10 @@ void FCameraShakePreviewerLinkerExtension::UpdateAllPreviewers()
 bool FCameraShakePreviewerLinkerExtension::HasAnyShake() const
 {
 	TArray<FActiveCameraShakeInfo> TempCameraShakes;
-	for (const TPair<FInstanceHandle, FCameraShakePreviewer>& Pair : Previewers)
+	for (const TPair<FInstanceHandle, TSharedRef<FCameraShakePreviewer>>& Pair : Previewers)
 	{
-		if (Pair.Value.NumActiveCameraShakes() > 0)
+		TSharedRef<FCameraShakePreviewer> Previewer(Pair.Value);
+		if (Previewer->NumActiveCameraShakes() > 0)
 		{
 			return true;
 		}
@@ -273,13 +280,14 @@ void FCameraShakePreviewerLinkerExtension::OnLevelViewportClientListChanged()
 	// However, we need to automatically register any *new* viewport that fits our requirements.
 	UMovieSceneEntitySystemLinker* Linker = WeakLinker.Get();
 	FInstanceRegistry* InstanceRegistry = Linker->GetInstanceRegistry();
-	for (TPair<FInstanceHandle, FCameraShakePreviewer> Pair : Previewers)
+	for (TPair<FInstanceHandle, TSharedRef<FCameraShakePreviewer>> Pair : Previewers)
 	{
 		const FSequenceInstance& SequenceInstance = InstanceRegistry->GetInstance(Pair.Key);
 		UObject* PlaybackContext = SequenceInstance.GetSharedPlaybackState()->GetPlaybackContext();
 		UWorld* ContextWorld = PlaybackContext ? PlaybackContext->GetWorld() : nullptr;
 
-		Pair.Value.RegisterViewModifiers(
+		TSharedRef<FCameraShakePreviewer> Previewer(Pair.Value);
+		Previewer->RegisterViewModifiers(
 				[ContextWorld](FLevelEditorViewportClient* LevelVC)
 				{
 					return LevelVC->AllowsCinematicControl() && LevelVC->GetWorld() == ContextWorld;
@@ -435,7 +443,7 @@ void UMovieSceneCameraShakeInstantiatorSystem::OnRun(FSystemTaskPrerequisites& I
 				// add some shaking ourselves. Let's do that here.
 				// In fact, in the editor, the above StartCameraShake call generally does nothing
 				// since there is no player controller.
-				FCameraShakePreviewer& Previewer = PreviewerExtension->GetPreviewer(InstanceHandle);
+				TSharedRef<FCameraShakePreviewer> Previewer = PreviewerExtension->GetPreviewer(InstanceHandle);
 
 				FCameraShakePreviewerAddParams PreviewParams;
 				PreviewParams.ShakeClass = ShakeClass;
@@ -450,10 +458,10 @@ void UMovieSceneCameraShakeInstantiatorSystem::OnRun(FSystemTaskPrerequisites& I
 				UCameraShakeBase* OldShakeInstance = ShakeInstanceData.ShakeInstance;
 				if (OldShakeInstance)
 				{
-					Previewer.RemoveCameraShake(OldShakeInstance);
+					Previewer->RemoveCameraShake(OldShakeInstance);
 				}
 
-				ShakeInstanceData.ShakeInstance = Previewer.AddCameraShake(PreviewParams);
+				ShakeInstanceData.ShakeInstance = Previewer->AddCameraShake(PreviewParams);
 				ShakeInstanceData.bManagedByPreviewer = true;
 			}
 #endif  // WITH_EDITOR
@@ -582,7 +590,7 @@ void UMovieSceneCameraShakeInstantiatorSystem::TriggerOneShotShakes()
 					if (PreviewerExtension)
 					{
 						// Also start playing the shake in our editor preview.
-						FCameraShakePreviewer& Previewer = PreviewerExtension->GetPreviewer(Pair.Key);
+						TSharedRef<FCameraShakePreviewer> Previewer = PreviewerExtension->GetPreviewer(Pair.Key);
 
 						FCameraShakePreviewerAddParams PreviewParams;
 						PreviewParams.ShakeClass = ShakeClass;
@@ -591,7 +599,7 @@ void UMovieSceneCameraShakeInstantiatorSystem::TriggerOneShotShakes()
 						PreviewParams.Scale = Trigger.Trigger.PlayScale;
 						PreviewParams.PlaySpace = Trigger.Trigger.PlaySpace;
 						PreviewParams.UserPlaySpaceRot = Trigger.Trigger.UserDefinedPlaySpace;
-						Previewer.AddCameraShake(PreviewParams);
+						Previewer->AddCameraShake(PreviewParams);
 					}
 #endif  // WITH_EDITOR
 				}
diff --git a/Engine/Source/Runtime/MovieSceneTracks/Public/Evaluation/MovieSceneCameraShakePreviewer.h b/Engine/Source/Runtime/MovieSceneTracks/Public/Evaluation/MovieSceneCameraShakePreviewer.h
--- a/Engine/Source/Runtime/MovieSceneTracks/Public/Evaluation/MovieSceneCameraShakePreviewer.h
+++ b/Engine/Source/Runtime/MovieSceneTracks/Public/Evaluation/MovieSceneCameraShakePreviewer.h
@@ -35,7 +35,7 @@ struct FCameraShakePreviewerAddParams
 /**
  * A class that owns a gameplay camera shake manager, so that we can us it to preview shakes in editor.
  */
-class FCameraShakePreviewer : public FGCObject
+class FCameraShakePreviewer : public FGCObject, public TSharedFromThis<FCameraShakePreviewer>
 {
 public:
 	using FViewportFilter = TFunctionRef<bool(FLevelEditorViewportClient*)>;
diff --git a/Engine/Source/Runtime/MovieSceneTracks/Public/Systems/MovieSceneCameraShakeSystem.h b/Engine/Source/Runtime/MovieSceneTracks/Public/Systems/MovieSceneCameraShakeSystem.h
--- a/Engine/Source/Runtime/MovieSceneTracks/Public/Systems/MovieSceneCameraShakeSystem.h
+++ b/Engine/Source/Runtime/MovieSceneTracks/Public/Systems/MovieSceneCameraShakeSystem.h
@@ -37,9 +37,9 @@ namespace UE::MovieScene
 		virtual ~FCameraShakePreviewerLinkerExtension();
 
 		/** Find the previewer for the given sequence instance, if it has already been created */
-		FCameraShakePreviewer* FindPreviewer(FInstanceHandle InstanceHandle);
+		TSharedPtr<FCameraShakePreviewer> FindPreviewer(FInstanceHandle InstanceHandle);
 		/** Get or create the previewer for the given sequence instance */
-		FCameraShakePreviewer& GetPreviewer(FInstanceHandle InstanceHandle);
+		TSharedRef<FCameraShakePreviewer> GetPreviewer(FInstanceHandle InstanceHandle);
 		/** Update all previewers by using the update context of each matching sequence instance */
 		void UpdateAllPreviewers();
 		/** Whether there is any previewer with any shake */
@@ -49,7 +49,7 @@ namespace UE::MovieScene
 		void OnLevelViewportClientListChanged();
 
 	private:
-		TMap<FInstanceHandle, FCameraShakePreviewer> Previewers;
+		TMap<FInstanceHandle, TSharedRef<FCameraShakePreviewer>> Previewers;
 	};
 #endif
 }
