Sequencer: Use WeakPointers to store SkelMeshComps in system data.

diff --git a/Engine/Source/Runtime/MovieSceneTracks/Private/Systems/MovieSceneSkeletalAnimationSystem.cpp b/Engine/Source/Runtime/MovieSceneTracks/Private/Systems/MovieSceneSkeletalAnimationSystem.cpp
--- a/Engine/Source/Runtime/MovieSceneTracks/Private/Systems/MovieSceneSkeletalAnimationSystem.cpp
+++ b/Engine/Source/Runtime/MovieSceneTracks/Private/Systems/MovieSceneSkeletalAnimationSystem.cpp
@@ -283,7 +283,7 @@ void FBoneTransformFinalizeData::Register(USkeletalMeshComponent* InSkeleletalMe
 		MeshRelativeRootMotionTransform = InMeshRelativeRootMotionTransform;
 		InitialActorTransform = InInitialActorTransform;
 		// Also store the inverse relative rotation between SkeletalMeshComponent and the Actor if it's not the root component
-		if (SkeletalMeshComponent && SkeletalMeshComponent != SkeletalMeshComponent->GetOwner()->GetRootComponent())
+		if (SkeletalMeshComponent.IsValid() && SkeletalMeshComponent.Get() != SkeletalMeshComponent->GetOwner()->GetRootComponent())
 		{
 			InverseMeshToActorRotation = SkeletalMeshComponent->GetOwner()->GetRootComponent()->GetComponentTransform().GetRelativeTransformReverse(SkeletalMeshComponent->GetComponentTransform()).GetRotation();
 		}
@@ -316,7 +316,7 @@ void FBoneTransformFinalizeData::Register(USkeletalMeshComponent* InSkeleletalMe
 
 void FBoneTransformFinalizeData::Unregister()
 {
-	if (SkeletalMeshComponent)
+	if (SkeletalMeshComponent.IsValid())
 	{
 		SkeletalMeshComponent->UnregisterOnBoneTransformsFinalizedDelegate(OnBoneTransformsFinalizedHandle);
 	}
@@ -341,7 +341,7 @@ void FBoneTransformFinalizeData::BoneTransformFinalized()
 	}
 #endif
 
-	if (SkeletalMeshComponent && SwapRootBone != ESwapRootBone::SwapRootBone_None)
+	if (SkeletalMeshComponent.IsValid() && SwapRootBone != ESwapRootBone::SwapRootBone_None)
 	{
 		FTransform RelativeTransform = MeshRelativeRootMotionTransform;
 
@@ -376,7 +376,7 @@ void FBoneTransformFinalizeData::BoneTransformFinalized()
 void FSkeletalAnimationSystemData::ResetSkeletalAnimations()
 {
 	//clear out the delegates
-	for (TTuple<USkeletalMeshComponent*, FBoundObjectActiveSkeletalAnimations>& Pair : SkeletalAnimations)
+	for (TTuple<TWeakObjectPtr<USkeletalMeshComponent>, FBoundObjectActiveSkeletalAnimations>& Pair : SkeletalAnimations)
 	{
 		Pair.Value.BoneTransformFinalizeData.Unregister();
 	}
@@ -583,9 +583,12 @@ public:
 	}
 	void Run() const
 	{
-		for (const TTuple<USkeletalMeshComponent*, FBoundObjectActiveSkeletalAnimations>& Pair : SystemData->SkeletalAnimations)
+		for (const TTuple<TWeakObjectPtr<USkeletalMeshComponent>, FBoundObjectActiveSkeletalAnimations>& Pair : SystemData->SkeletalAnimations)
 		{
-			EvaluateSkeletalAnimations(Pair.Key, Pair.Value);
+			if (Pair.Key.IsValid())
+			{
+				EvaluateSkeletalAnimations(Pair.Key.Get(), Pair.Value);
+			}
 		}
 	}
 
diff --git a/Engine/Source/Runtime/MovieSceneTracks/Public/Systems/MovieSceneSkeletalAnimationSystem.h b/Engine/Source/Runtime/MovieSceneTracks/Public/Systems/MovieSceneSkeletalAnimationSystem.h
--- a/Engine/Source/Runtime/MovieSceneTracks/Public/Systems/MovieSceneSkeletalAnimationSystem.h
+++ b/Engine/Source/Runtime/MovieSceneTracks/Public/Systems/MovieSceneSkeletalAnimationSystem.h
@@ -39,7 +39,7 @@ struct FActiveSkeletalAnimation
 struct FBoneTransformFinalizeData
 {
 	FBoneTransformFinalizeData();
-	USkeletalMeshComponent* SkeletalMeshComponent;
+	TWeakObjectPtr<USkeletalMeshComponent> SkeletalMeshComponent;
 	ESwapRootBone SwapRootBone;
 	FTransform MeshRelativeRootMotionTransform;
 	TOptional<FTransform> InitialActorTransform;
@@ -84,7 +84,7 @@ struct FSkeletalAnimationSystemData
 	void ResetSkeletalAnimations();
 
 	/** Map of active skeletal animations for each bound object */
-	TMap<USkeletalMeshComponent*, FBoundObjectActiveSkeletalAnimations> SkeletalAnimations;
+	TMap<TWeakObjectPtr<USkeletalMeshComponent>, FBoundObjectActiveSkeletalAnimations> SkeletalAnimations;
 
 	/** Map of persistent montage data */
 	TMap<FObjectKey, TMap<FObjectKey, FMontagePlayerPerSectionData>> MontageData;
