Sequencer: Multiple camera cut track rows

When multiple rows exist, the camera cut track does not try to auto arrange sections

diff --git a/Engine/Source/Runtime/MovieSceneTracks/Private/Tracks/MovieSceneCameraCutTrack.cpp b/Engine/Source/Runtime/MovieSceneTracks/Private/Tracks/MovieSceneCameraCutTrack.cpp
--- a/Engine/Source/Runtime/MovieSceneTracks/Private/Tracks/MovieSceneCameraCutTrack.cpp
+++ b/Engine/Source/Runtime/MovieSceneTracks/Private/Tracks/MovieSceneCameraCutTrack.cpp
@@ -95,7 +95,7 @@ UMovieSceneSection* UMovieSceneCameraCutTrack::CreateNewSection()
 
 bool UMovieSceneCameraCutTrack::SupportsMultipleRows() const
 {
-	return false;
+	return true;
 }
 
 EMovieSceneTrackEasingSupportFlags UMovieSceneCameraCutTrack::SupportsEasing(FMovieSceneSupportsEasingParams& Params) const
@@ -111,7 +111,7 @@ EMovieSceneTrackEasingSupportFlags UMovieSceneCameraCutTrack::SupportsEasing(FMo
 		{
 			return EMovieSceneTrackEasingSupportFlags::AutomaticEasing | EMovieSceneTrackEasingSupportFlags::ManualEasing;
 		}
-		else if (NumSections > 1 && bAutoArrangeSections)
+		else if (NumSections > 1 && IsAutoManagingSections())
 		{
 			// Find the section with the earliest start time, and the section with the latest end time.
 			int32 EdgeSections[2] = { 0, NumSections - 1 };
@@ -156,7 +156,7 @@ EMovieSceneTrackEasingSupportFlags UMovieSceneCameraCutTrack::SupportsEasing(FMo
 			}
 			return Flags;
 		}
-		else if (NumSections > 1 && !bAutoArrangeSections)
+		else if (NumSections > 1 && !IsAutoManagingSections())
 		{
 			// The given section supports manual easing on one side only if it doesn't overlap with any other section
 			// on that side.
@@ -278,9 +278,19 @@ void UMovieSceneCameraCutTrack::PreCompileImpl(FMovieSceneTrackPreCompileResult&
 	}
 }
 
+bool UMovieSceneCameraCutTrack::IsAutoManagingSections() const
+{
+	return bAutoArrangeSections && GetMaxRowIndex() == 0;
+}
+
+void UMovieSceneCameraCutTrack::SetIsAutoManagingSections(bool bInAutoArrangeSections)
+{
+	bAutoArrangeSections = bInAutoArrangeSections;
+}
+
 bool UMovieSceneCameraCutTrack::AutoArrangeSectionsIfNeeded(UMovieSceneSection& ChangedSection, bool bWasDeletion, bool bCleanUp)
 {
-	if (bAutoArrangeSections)
+	if (IsAutoManagingSections())
 	{
 		if (bCanBlend)
 		{
@@ -306,7 +316,7 @@ void UMovieSceneCameraCutTrack::RearrangeAllSections()
 	MovieSceneHelpers::SortConsecutiveSections(MutableView(Sections));
 
 	// Go over the sections and change anything we don't want to support.
-	const bool bRemoveGaps = bAutoArrangeSections;
+	const bool bRemoveGaps = IsAutoManagingSections();
 	const bool bRemoveOverlaps = !bCanBlend;
 	const bool bRemoveEasings = !bCanBlend;
 
diff --git a/Engine/Source/Runtime/MovieSceneTracks/Public/Tracks/MovieSceneCameraCutTrack.h b/Engine/Source/Runtime/MovieSceneTracks/Public/Tracks/MovieSceneCameraCutTrack.h
--- a/Engine/Source/Runtime/MovieSceneTracks/Public/Tracks/MovieSceneCameraCutTrack.h
+++ b/Engine/Source/Runtime/MovieSceneTracks/Public/Tracks/MovieSceneCameraCutTrack.h
@@ -63,16 +63,10 @@ public:
 #endif
 
 	/** @return Whether camera cut sections should automatically resize to fill gaps */
-	bool IsAutoManagingSections() const
-	{
-		return bAutoArrangeSections;
-	}
+	MOVIESCENETRACKS_API bool IsAutoManagingSections() const;
 
 	/** Sets whether camera cut sections should automatically resize to fill gaps */
-	void SetIsAutoManagingSections(bool bInAutoArrangeSections)
-	{
-		bAutoArrangeSections = bInAutoArrangeSections;
-	}
+	MOVIESCENETRACKS_API void SetIsAutoManagingSections(bool bInAutoArrangeSections);
 
 #if WITH_EDITOR
 	virtual EMovieSceneSectionMovedResult OnSectionMoved(UMovieSceneSection& Section, const FMovieSceneSectionMovedParams& Params) override;
