Fix HBias and Animation Sections. The skeletal animation system now adds groupings based on bound objects so sections animating the same actor across the hierarchy can properly be marked as ignored if the hbiases differ.

diff --git a/Engine/Plugins/Experimental/MovieSceneAnimMixer/Source/MovieSceneAnimMixer/Private/Systems/MovieSceneAnimMixerSystem.cpp b/Engine/Plugins/Experimental/MovieSceneAnimMixer/Source/MovieSceneAnimMixer/Private/Systems/MovieSceneAnimMixerSystem.cpp
--- a/Engine/Plugins/Experimental/MovieSceneAnimMixer/Source/MovieSceneAnimMixer/Private/Systems/MovieSceneAnimMixerSystem.cpp
+++ b/Engine/Plugins/Experimental/MovieSceneAnimMixer/Source/MovieSceneAnimMixer/Private/Systems/MovieSceneAnimMixerSystem.cpp
@@ -926,7 +926,7 @@ void UMovieSceneAnimMixerSystem::OnSchedulePersistentTasks(UE::MovieScene::IEnti
 		.Write(AnimMixerComponents->MixerEntry)
 		.ReadOptional(AnimMixerComponents->RootMotionSettings)
 		.ReadOptional(BuiltInComponents->WeightAndEasingResult)
-		.FilterNone({BuiltInComponents->Tags.NeedsUnlink})
+		.FilterNone({BuiltInComponents->Tags.NeedsUnlink, BuiltInComponents->Tags.Ignored})
 		.Schedule_PerEntity<FUpdateTaskEntities>(&Linker->EntityManager, TaskScheduler);
 
 
@@ -935,6 +935,7 @@ void UMovieSceneAnimMixerSystem::OnSchedulePersistentTasks(UE::MovieScene::IEnti
 		.Read(AnimMixerComponents->MeshComponent)
 		.Read(AnimMixerComponents->Target)
 		.Write(AnimMixerComponents->MixerTask)
+		.FilterNone({BuiltInComponents->Tags.Ignored})
 		.Schedule_PerEntity<FEvaluateAnimMixers>(&Linker->EntityManager, TaskScheduler, &Mixers, &RootMotionByObject, Linker);
 
 	TaskScheduler->AddPrerequisite(UpdateTask, MixTask);
diff --git a/Engine/Source/Runtime/MovieSceneTracks/Private/Systems/MovieSceneSkeletalAnimationSystem.cpp b/Engine/Source/Runtime/MovieSceneTracks/Private/Systems/MovieSceneSkeletalAnimationSystem.cpp
--- a/Engine/Source/Runtime/MovieSceneTracks/Private/Systems/MovieSceneSkeletalAnimationSystem.cpp
+++ b/Engine/Source/Runtime/MovieSceneTracks/Private/Systems/MovieSceneSkeletalAnimationSystem.cpp
@@ -10,6 +10,7 @@
 #include "EntitySystem/MovieSceneEntitySystemLinker.h"
 #include "EntitySystem/MovieSceneEntitySystemRunner.h"
 #include "EntitySystem/MovieSceneEntitySystemTask.h"
+#include "EntitySystem/MovieSceneEntityGroupingSystem.h"
 #include "IMovieScenePlayer.h"
 #include "MovieScene.h"
 #include "MovieSceneExecutionToken.h"
@@ -1335,6 +1336,27 @@ void UMovieSceneSkeletalAnimationSystem::OnRun(FSystemTaskPrerequisites& InPrere
 	.Dispatch<FEvaluateSkeletalAnimations>(&Linker->EntityManager, EvalPrereqs, &Subsequents, Linker, &SystemData);
 }
 
+void UMovieSceneSkeletalAnimationSystem::OnLink()
+{
+	using namespace UE::MovieScene;
+	
+	UMovieSceneEntityGroupingSystem* GroupingSystem = Linker->LinkSystem<UMovieSceneEntityGroupingSystem>();
+
+	GroupingKey = GroupingSystem->AddGrouping(FAnimationGroupingPolicy(), FBuiltInComponentTypes::Get()->BoundObject, FMovieSceneTracksComponentTypes::Get()->SkeletalAnimation);
+}
+
+void UMovieSceneSkeletalAnimationSystem::OnUnlink()
+{
+	using namespace UE::MovieScene;
+	
+	UMovieSceneEntityGroupingSystem* GroupingSystem = Linker->FindSystem<UMovieSceneEntityGroupingSystem>();
+	if (ensure(GroupingSystem))
+	{
+		GroupingSystem->RemoveGrouping(GroupingKey);
+	}
+	GroupingKey = FEntityGroupingPolicyKey();
+}
+
 bool UMovieSceneSkeletalAnimationSystem::IsRelevantImpl(UMovieSceneEntitySystemLinker* InLinker) const
 {
 	using namespace UE::MovieScene;
diff --git a/Engine/Source/Runtime/MovieSceneTracks/Public/Systems/MovieSceneSkeletalAnimationSystem.h b/Engine/Source/Runtime/MovieSceneTracks/Public/Systems/MovieSceneSkeletalAnimationSystem.h
--- a/Engine/Source/Runtime/MovieSceneTracks/Public/Systems/MovieSceneSkeletalAnimationSystem.h
+++ b/Engine/Source/Runtime/MovieSceneTracks/Public/Systems/MovieSceneSkeletalAnimationSystem.h
@@ -2,9 +2,11 @@
 
 #pragma once
 
+#include "EntitySystem/MovieSceneEntityGroupingSystem.h"
 #include "EntitySystem/MovieSceneEntitySystem.h"
 #include "EntitySystem/MovieSceneSequenceInstanceHandle.h"
 #include "Evaluation/MovieScenePlayback.h"
+#include "MovieSceneTracksComponentTypes.h"
 #include "UObject/ObjectKey.h"
 #include "MovieSceneSkeletalAnimationSystem.generated.h"
 
@@ -110,10 +112,40 @@ public:
 	MOVIESCENETRACKS_API TOptional<FTransform> GetInitialActorTransform(UObject* InObject) const;
 	MOVIESCENETRACKS_API TOptional<FQuat> GetInverseMeshToActorRotation(UObject* InObject) const;
 
+protected:
+	
+	struct FAnimationGroupingPolicy
+	{
+		using GroupKeyType = FObjectKey;
+
+		bool GetGroupKey(UObject* Object, FMovieSceneSkeletalAnimationComponentData ComponentData, GroupKeyType& OutGroupKey)
+		{
+			// ComponentData is only used to filter the grouping to objects with animations. If it were used for the key, it would only group objects
+			// with identical animations, so instead the object is used for both the key and value of the tuple. 
+			OutGroupKey = FObjectKey(Object);
+			return true;
+		}
+		
+#if WITH_EDITOR
+		bool OnObjectsReplaced(GroupKeyType& InOutKey, const TMap<UObject*, UObject*>& ReplacementMap)
+		{
+			if (UObject* const * NewObject = ReplacementMap.Find(InOutKey.ResolveObjectPtr()))
+			{
+				InOutKey = *NewObject;
+				return true;
+			}
+			return false;
+		}
+#endif
+	
+	};
+
 private:
 
 	virtual void OnSchedulePersistentTasks(UE::MovieScene::IEntitySystemScheduler* TaskScheduler) override;
 	virtual void OnRun(FSystemTaskPrerequisites& InPrerequisites, FSystemSubsequentTasks& Subsequents) override final;
+	virtual void OnLink() override;
+	virtual void OnUnlink() override;
 
 	virtual bool IsRelevantImpl(UMovieSceneEntitySystemLinker* InLinker) const override;
 
@@ -122,5 +154,6 @@ private:
 private:
 
 	UE::MovieScene::FSkeletalAnimationSystemData SystemData;
+	UE::MovieScene::FEntityGroupingPolicyKey GroupingKey;
 };
 
