Sequencer: Speculative fix for possible null world context in GetSpawnName

#rnx

diff --git a/Engine/Source/Runtime/MovieSceneTracks/Private/Bindings/MovieSceneSpawnableActorBinding.cpp b/Engine/Source/Runtime/MovieSceneTracks/Private/Bindings/MovieSceneSpawnableActorBinding.cpp
--- a/Engine/Source/Runtime/MovieSceneTracks/Private/Bindings/MovieSceneSpawnableActorBinding.cpp
+++ b/Engine/Source/Runtime/MovieSceneTracks/Private/Bindings/MovieSceneSpawnableActorBinding.cpp
@@ -275,8 +275,14 @@ FName UMovieSceneSpawnableActorBindingBase::GetSpawnName(const FGuid& BindingId,
 
 #if WITH_EDITOR
 		UClass* ActorClass = GetActorClass();
-		return bUseNetAddressableName ? GetNetAddressableName(SharedPlaybackState, BindingId, TemplateID, DesiredBindingName)
-			: MakeUniqueObjectName(WorldContext->PersistentLevel, ActorClass ? ActorClass : AActor::StaticClass(), *DesiredBindingName);
+		if (bUseNetAddressableName)
+		{
+			return GetNetAddressableName(SharedPlaybackState, BindingId, TemplateID, DesiredBindingName);
+		}
+		else if (ensure(WorldContext && WorldContext->PersistentLevel))
+		{
+			return MakeUniqueObjectName(WorldContext->PersistentLevel, ActorClass ? ActorClass : AActor::StaticClass(), *DesiredBindingName);
+		}
 #else
 		return bUseNetAddressableName ? GetNetAddressableName(SharedPlaybackState, BindingId, TemplateID, DesiredBindingName) : NAME_None;
 #endif
diff --git a/Engine/Source/Runtime/MovieSceneTracks/Private/Bindings/MovieSceneSpawnableDirectorBlueprintBinding.cpp b/Engine/Source/Runtime/MovieSceneTracks/Private/Bindings/MovieSceneSpawnableDirectorBlueprintBinding.cpp
--- a/Engine/Source/Runtime/MovieSceneTracks/Private/Bindings/MovieSceneSpawnableDirectorBlueprintBinding.cpp
+++ b/Engine/Source/Runtime/MovieSceneTracks/Private/Bindings/MovieSceneSpawnableDirectorBlueprintBinding.cpp
@@ -75,10 +75,12 @@ FName UMovieSceneSpawnableDirectorBlueprintBinding::GetSpawnName(const FGuid& Bi
 
 #if WITH_EDITOR
 	UClass* ObjectClass = GetBoundObjectClass();
-	return MakeUniqueObjectName(WorldContext->PersistentLevel.Get(), ObjectClass ? ObjectClass : UObject::StaticClass(), *DesiredBindingName);
-#else
-	return *DesiredBindingName;
+	if (ensure(WorldContext && WorldContext->PersistentLevel))
+	{
+		return MakeUniqueObjectName(WorldContext->PersistentLevel.Get(), ObjectClass ? ObjectClass : UObject::StaticClass(), *DesiredBindingName);
+	}
 #endif
+	return *DesiredBindingName;
 }
 
 bool UMovieSceneSpawnableDirectorBlueprintBinding::SupportsBindingCreationFromObject(const UObject* SourceObject) const
