Sequencer: Fix CIS

#rnx

diff --git a/Engine/Source/Runtime/MovieSceneTracks/Private/Bindings/MovieSceneSpawnableActorBinding.cpp b/Engine/Source/Runtime/MovieSceneTracks/Private/Bindings/MovieSceneSpawnableActorBinding.cpp
--- a/Engine/Source/Runtime/MovieSceneTracks/Private/Bindings/MovieSceneSpawnableActorBinding.cpp
+++ b/Engine/Source/Runtime/MovieSceneTracks/Private/Bindings/MovieSceneSpawnableActorBinding.cpp
@@ -264,7 +264,7 @@ FName UMovieSceneSpawnableActorBindingBase::GetSpawnName(const FGuid& BindingId,
 	// We use the net addressable name for spawnable actors on any non-editor, non-standalone world (ie, all clients, servers and PIE worlds)
 
 	UWorld* WorldContext = GetWorldContext(SharedPlaybackState);
-	const bool bUseNetAddressableName = bNetAddressableName && (WorldContext->WorldType != EWorldType::Editor) && (WorldContext->GetNetMode() != ENetMode::NM_Standalone);
+	const bool bUseNetAddressableName = bNetAddressableName && WorldContext && (WorldContext->WorldType != EWorldType::Editor) && (WorldContext->GetNetMode() != ENetMode::NM_Standalone);
 	FString DesiredBindingName = GetDesiredBindingName();
 	if (FMovieScenePossessable* Possessable = MovieScene.FindPossessable(BindingId))
 	{
@@ -279,7 +279,7 @@ FName UMovieSceneSpawnableActorBindingBase::GetSpawnName(const FGuid& BindingId,
 		{
 			return GetNetAddressableName(SharedPlaybackState, BindingId, TemplateID, DesiredBindingName);
 		}
-		else if (ensure(WorldContext && WorldContext->PersistentLevel))
+		else if (ensure(WorldContext != nullptr && WorldContext->PersistentLevel))
 		{
 			return MakeUniqueObjectName(WorldContext->PersistentLevel, ActorClass ? ActorClass : AActor::StaticClass(), *DesiredBindingName);
 		}
diff --git a/Engine/Source/Runtime/MovieSceneTracks/Private/Bindings/MovieSceneSpawnableDirectorBlueprintBinding.cpp b/Engine/Source/Runtime/MovieSceneTracks/Private/Bindings/MovieSceneSpawnableDirectorBlueprintBinding.cpp
--- a/Engine/Source/Runtime/MovieSceneTracks/Private/Bindings/MovieSceneSpawnableDirectorBlueprintBinding.cpp
+++ b/Engine/Source/Runtime/MovieSceneTracks/Private/Bindings/MovieSceneSpawnableDirectorBlueprintBinding.cpp
@@ -75,7 +75,7 @@ FName UMovieSceneSpawnableDirectorBlueprintBinding::GetSpawnName(const FGuid& Bi
 
 #if WITH_EDITOR
 	UClass* ObjectClass = GetBoundObjectClass();
-	if (ensure(WorldContext && WorldContext->PersistentLevel))
+	if (ensure(WorldContext != nullptr && WorldContext->PersistentLevel))
 	{
 		return MakeUniqueObjectName(WorldContext->PersistentLevel.Get(), ObjectClass ? ObjectClass : UObject::StaticClass(), *DesiredBindingName);
 	}
