PR #13499: Fix for editor crash in MovieSceneAudioSystem by making ScrubbedSound a TObjectPtr instead of a TStrongObjectPtr

[RN] Fixed issue with GC hang after closing PIE after scrubbing an audio sequence while PIE is open and playing a level sequence.

diff --git a/Engine/Source/Runtime/MovieSceneTracks/Private/Systems/MovieSceneAudioSystem.cpp b/Engine/Source/Runtime/MovieSceneTracks/Private/Systems/MovieSceneAudioSystem.cpp
--- a/Engine/Source/Runtime/MovieSceneTracks/Private/Systems/MovieSceneAudioSystem.cpp
+++ b/Engine/Source/Runtime/MovieSceneTracks/Private/Systems/MovieSceneAudioSystem.cpp
@@ -979,6 +979,25 @@ void UMovieSceneAudioSystem::ResetSharedData()
 	}
 }
 
+void UMovieSceneAudioSystem::AddReferencedObjects(UObject* InThis, FReferenceCollector& Collector)
+{
+	Super::AddReferencedObjects(InThis, Collector);
+#if WITH_EDITOR
+	UMovieSceneAudioSystem* This = CastChecked<UMovieSceneAudioSystem>(InThis);
+	for (auto& ActorPair : This->AudioComponentsByActorKey)
+	{
+		for (auto& SectionPair : ActorPair.Value)
+		{
+			FAudioComponentEvaluationData& EvalData = SectionPair.Value;
+			if (EvalData.ScrubbedSound)
+			{
+				Collector.AddReferencedObject(EvalData.ScrubbedSound);
+			}
+		}
+	}
+#endif
+}
+
 void UMovieSceneAudioSystem::OnSchedulePersistentTasks(UE::MovieScene::IEntitySystemScheduler* TaskScheduler)
 {
 	if (!GEngine || !GEngine->UseSound())
@@ -1207,7 +1226,7 @@ UMovieSceneAudioSystem::FAudioComponentEvaluationData* UMovieSceneAudioSystem::A
 		else
 		{
 			ScrubbedSound->SetFlags(RF_Transient);
-			ExistingData->ScrubbedSound = TStrongObjectPtr<UScrubbedSound>(ScrubbedSound);
+			ExistingData->ScrubbedSound = TObjectPtr<UScrubbedSound>(ScrubbedSound);
 		}
 #endif // WITH_EDITOR
 
@@ -1259,7 +1278,7 @@ UMovieSceneAudioSystem::FAudioComponentEvaluationData* UMovieSceneAudioSystem::A
 		else
 		{
 			ScrubbedSound->SetFlags(RF_Transient);
-			ExistingData->ScrubbedSound = TStrongObjectPtr<UScrubbedSound>(ScrubbedSound);
+			ExistingData->ScrubbedSound = TObjectPtr<UScrubbedSound>(ScrubbedSound);
 		}
 #endif // WITH_EDITOR
 	}
diff --git a/Engine/Source/Runtime/MovieSceneTracks/Private/Systems/MovieSceneAudioSystem.h b/Engine/Source/Runtime/MovieSceneTracks/Private/Systems/MovieSceneAudioSystem.h
--- a/Engine/Source/Runtime/MovieSceneTracks/Private/Systems/MovieSceneAudioSystem.h
+++ b/Engine/Source/Runtime/MovieSceneTracks/Private/Systems/MovieSceneAudioSystem.h
@@ -42,7 +42,7 @@ namespace UE::MovieScene
 
 #if WITH_EDITOR
 		/** While in editor, we can scrub the audio in the audio component. */
-		TStrongObjectPtr<UScrubbedSound> ScrubbedSound;
+		TObjectPtr<UScrubbedSound> ScrubbedSound;
 #endif
 
 		/** Volume multiplier to use this frame */
@@ -120,6 +120,9 @@ public:
 	 */
 	void ResetSharedData();
 
+	// To expose the class to GC //
+	static void AddReferencedObjects(UObject* InThis, FReferenceCollector& Collector);
+
 private:
 
 	using FInstanceObjectKey = TTuple<FInstanceHandle, FObjectKey>;
