Fix montages not playing back after playing a sequence that uses custom mode

When restoring state we don't always re-initialize the anim instance on the skeletal mesh component after setting the anim script instance to the cached value, causing montages to not play, since the instance is flagged as uninitialized.

This change now initializes the anim instance if it's reverted back to its cached value.

diff --git a/Engine/Source/Runtime/MovieSceneTracks/Private/Systems/MovieSceneSkeletalAnimationSystem.cpp b/Engine/Source/Runtime/MovieSceneTracks/Private/Systems/MovieSceneSkeletalAnimationSystem.cpp
--- a/Engine/Source/Runtime/MovieSceneTracks/Private/Systems/MovieSceneSkeletalAnimationSystem.cpp
+++ b/Engine/Source/Runtime/MovieSceneTracks/Private/Systems/MovieSceneSkeletalAnimationSystem.cpp
@@ -151,9 +151,10 @@ struct FPreAnimatedSkeletalAnimationTraits : FBoundObjectPreAnimatedStateTraits
 			{
 				//the skeleton may have changed so need to recalc required bones as needed.
 				Component->AnimScriptInstance->CurrentSkeleton = Component->GetSkeletalMeshAsset()->GetSkeleton();
-				//Need at least RecalcRequiredbones and UpdateMorphTargetrs
-				Component->InitializeAnimScriptInstance(true);
 			}
+			// When Changing skeletons RecalcRequiredBones and UpdateMorphTargets need to be called, and InitializeAnimation needs to be called regardless, when changing script instances
+			// Playing montages does not work on uninitialized anim script instances with the EarlyOutMontageWhenUninitialized cvar enabled.
+			Component->InitializeAnimScriptInstance(true);
 		}
 
 		// Restore pose after unbinding to force the restored pose
