Cameras: Change log to verbose

#rnx

diff --git a/Engine/Source/Runtime/MovieSceneTracks/Private/TrackInstances/MovieSceneCameraCutGameHandler.cpp b/Engine/Source/Runtime/MovieSceneTracks/Private/TrackInstances/MovieSceneCameraCutGameHandler.cpp
--- a/Engine/Source/Runtime/MovieSceneTracks/Private/TrackInstances/MovieSceneCameraCutGameHandler.cpp
+++ b/Engine/Source/Runtime/MovieSceneTracks/Private/TrackInstances/MovieSceneCameraCutGameHandler.cpp
@@ -353,7 +353,7 @@ void FCameraCutGameHandler::SetCameraCut(
 		const float PlayRateFactor = Wrapper.CameraCutCapability ? Wrapper.CameraCutCapability->GetCameraBlendPlayRate() : 1.f;
 		const float BlendTime = CameraCutParams.BlendTime * (1 / FMath::Max(UE_SMALL_NUMBER, abs(PlayRateFactor)));
 
-		UE_LOG(LogMovieScene, Log, TEXT("Blending into new camera cut: '%s' -> '%s' (blend time: %f)"),
+		UE_LOG(LogMovieScene, Verbose, TEXT("Blending into new camera cut: '%s' -> '%s' (blend time: %f)"),
 			(ViewTarget ? *ViewTarget->GetName() : TEXT("None")),
 			(CameraActor ? *CameraActor->GetName() : TEXT("None")),
 			BlendTime);
@@ -373,13 +373,13 @@ void FCameraCutGameHandler::SetCameraCut(
 		const AActor* PendingViewTarget = CameraManager->PendingViewTarget.Target;
 		if (CameraActor != nullptr && PendingViewTarget == CameraActor)
 		{
-			UE_LOG(LogMovieScene, Log, TEXT("Camera transition aborted, we are already blending towards the intended camera"));
+			UE_LOG(LogMovieScene, Verbose, TEXT("Camera transition aborted, we are already blending towards the intended camera"));
 			bDoSetViewTarget = false;
 		}
 	}
 	else
 	{
-		UE_LOG(LogMovieScene, Log, TEXT("Starting new camera cut: '%s'"),
+		UE_LOG(LogMovieScene, Verbose, TEXT("Starting new camera cut: '%s'"),
 			(CameraActor ? *CameraActor->GetName() : TEXT("None")));
 	}
 	if (bDoSetViewTarget && ensureMsgf(CameraManager, TEXT("Can't set view target when there is no player controller!")))
