Sequencer: Follow up fixes for FMovieSceneBindingResolveResult. Return the first valid object

#rnx

diff --git a/Engine/Source/Runtime/MovieSceneTracks/Private/Bindings/MovieSceneSpawnableDirectorBlueprintBinding.cpp b/Engine/Source/Runtime/MovieSceneTracks/Private/Bindings/MovieSceneSpawnableDirectorBlueprintBinding.cpp
--- a/Engine/Source/Runtime/MovieSceneTracks/Private/Bindings/MovieSceneSpawnableDirectorBlueprintBinding.cpp
+++ b/Engine/Source/Runtime/MovieSceneTracks/Private/Bindings/MovieSceneSpawnableDirectorBlueprintBinding.cpp
@@ -16,7 +16,14 @@ UObject* UMovieSceneSpawnableDirectorBlueprintBinding::SpawnObjectInternal(UWorl
 {
 	FMovieSceneDynamicBindingResolveResult ResolveResult = FMovieSceneDynamicBindingInvoker::ResolveDynamicBinding(SharedPlaybackState, MovieScene.GetTypedOuter<UMovieSceneSequence>(), TemplateID, BindingId, DynamicBinding);
 
-	return ResolveResult.Objects.Num() > 0 ? ResolveResult.Objects[0].Get() : nullptr;
+	for (TObjectPtr<UObject> Object : ResolveResult.Objects)
+	{
+		if (Object.Get())
+		{
+			return Object.Get();
+		}
+	}
+	return nullptr;
 }
 
 void UMovieSceneSpawnableDirectorBlueprintBinding::DestroySpawnedObjectInternal(UObject* Object)
diff --git a/Engine/Source/Runtime/UMG/Private/Animation/WidgetAnimationBinding.cpp b/Engine/Source/Runtime/UMG/Private/Animation/WidgetAnimationBinding.cpp
--- a/Engine/Source/Runtime/UMG/Private/Animation/WidgetAnimationBinding.cpp
+++ b/Engine/Source/Runtime/UMG/Private/Animation/WidgetAnimationBinding.cpp
@@ -22,9 +22,12 @@ UObject* FWidgetAnimationBinding::FindRuntimeObject(const UWidgetTree& WidgetTre
 	if (Sequence && SharedPlaybackState && DynamicBinding.Function)
 	{
 		FMovieSceneDynamicBindingResolveResult ResolveResult = FMovieSceneDynamicBindingInvoker::ResolveDynamicBinding(SharedPlaybackState.ToSharedRef(), const_cast<UMovieSceneSequence*>(Sequence), MovieSceneSequenceID::Root, AnimationGuid, DynamicBinding);
-		if (ResolveResult.Objects.Num())
+		for (TObjectPtr<UObject> Object : ResolveResult.Objects)
 		{
-			return ResolveResult.Objects[0];
+			if (Object.Get())
+			{
+				return Object.Get();
+			}
 		}
 	}
 
