Sequencer: Ported boolean properties to ECS

- This is a little more convoluted than most property sections because we have to implement it all in the track rather than just putting it on the section type, owing to the fact that the bool section has been re-used in a few other contexts.
- Legacy template can be re-enabled using Sequencer.UseLegacyBooleanTemplate, and completely compiled out using Sequencer.CanUseLegacyBooleanTemplate.

diff --git a/Engine/Source/Runtime/MovieSceneTracks/Private/Evaluation/MovieScenePropertyTemplates.cpp b/Engine/Source/Runtime/MovieSceneTracks/Private/Evaluation/MovieScenePropertyTemplates.cpp
--- a/Engine/Source/Runtime/MovieSceneTracks/Private/Evaluation/MovieScenePropertyTemplates.cpp
+++ b/Engine/Source/Runtime/MovieSceneTracks/Private/Evaluation/MovieScenePropertyTemplates.cpp
@@ -26,6 +26,11 @@ namespace
 	}
 }
 
+namespace UE::MovieScene
+{
+	MOVIESCENETRACKS_API extern bool GBooleanPropertyLegacyTemplate;
+}
+
 //	----------------------------------------------------------------------------
 //	Boolean Property Template
 FMovieSceneBoolPropertySectionTemplate::FMovieSceneBoolPropertySectionTemplate(const UMovieSceneBoolSection& Section, const UMovieScenePropertyTrack& Track)
@@ -37,6 +42,11 @@ FMovieSceneBoolPropertySectionTemplate::FMovieSceneBoolPropertySectionTemplate(c
 
 void FMovieSceneBoolPropertySectionTemplate::Evaluate(const FMovieSceneEvaluationOperand& Operand, const FMovieSceneContext& Context, const FPersistentEvaluationData& PersistentData, FMovieSceneExecutionTokens& ExecutionTokens) const
 {
+	if (!UE::MovieScene::GBooleanPropertyLegacyTemplate)
+	{
+		return;
+	}
+
 	// Only evaluate if the curve has any data
 	bool Result = false;
 	if (BoolCurve.Evaluate(Context.GetTime(), Result))
diff --git a/Engine/Source/Runtime/MovieSceneTracks/Private/Systems/MovieSceneBoolPropertySystem.cpp b/Engine/Source/Runtime/MovieSceneTracks/Private/Systems/MovieSceneBoolPropertySystem.cpp
--- a/Engine/Source/Runtime/MovieSceneTracks/Private/Systems/MovieSceneBoolPropertySystem.cpp
+++ b/Engine/Source/Runtime/MovieSceneTracks/Private/Systems/MovieSceneBoolPropertySystem.cpp
@@ -2,6 +2,7 @@
 
 #include "Systems/MovieSceneBoolPropertySystem.h"
 #include "Systems/MovieScenePiecewiseBoolBlenderSystem.h"
+#include "Systems/BoolChannelEvaluatorSystem.h"
 #include "MovieSceneTracksComponentTypes.h"
 
 #include UE_INLINE_GENERATED_CPP_BY_NAME(MovieSceneBoolPropertySystem)
@@ -13,6 +14,7 @@ UMovieSceneBoolPropertySystem::UMovieSceneBoolPropertySystem(const FObjectInitia
 
 	if (HasAnyFlags(RF_ClassDefaultObject))
 	{
+		DefineImplicitPrerequisite(UBoolChannelEvaluatorSystem::StaticClass(), GetClass());
 		DefineImplicitPrerequisite(UMovieScenePiecewiseBoolBlenderSystem::StaticClass(), GetClass());
 
 		DefineComponentConsumer(GetClass(), UE::MovieScene::FMovieSceneTracksComponentTypes::Get()->Bool.PropertyTag);
diff --git a/Engine/Source/Runtime/MovieSceneTracks/Private/Tracks/MovieSceneBoolTrack.cpp b/Engine/Source/Runtime/MovieSceneTracks/Private/Tracks/MovieSceneBoolTrack.cpp
--- a/Engine/Source/Runtime/MovieSceneTracks/Private/Tracks/MovieSceneBoolTrack.cpp
+++ b/Engine/Source/Runtime/MovieSceneTracks/Private/Tracks/MovieSceneBoolTrack.cpp
@@ -4,9 +4,26 @@
 #include "MovieSceneCommonHelpers.h"
 #include "Sections/MovieSceneBoolSection.h"
 #include "Evaluation/MovieScenePropertyTemplates.h"
+#include "Algo/IndexOf.h"
 
 #include UE_INLINE_GENERATED_CPP_BY_NAME(MovieSceneBoolTrack)
 
+namespace UE::MovieScene
+{
+	MOVIESCENETRACKS_API bool GBooleanPropertyLegacyTemplate = false;
+	static FAutoConsoleVariableRef CVarBooleanPropertyECS(
+		TEXT("Sequencer.UseLegacyBooleanTemplate"),
+		GBooleanPropertyLegacyTemplate,
+		TEXT("(Default: false) DEPRECATED Whether to use Sequencer's legacy template evauation for boolean properties or not."));
+
+	MOVIESCENETRACKS_API bool GCompileBooleanPropertyLegacyTemplate = true;
+	static FAutoConsoleVariableRef CVarCompileBooleanPropertyECS(
+		TEXT("Sequencer.CanUseLegacyBooleanTemplate"),
+		GCompileBooleanPropertyLegacyTemplate,
+		TEXT("(Default: true) DEPRECATED Whether to compile Sequencer's legacy template boolean properties at all. Disabling this cvar will completely disable the legacy template."));
+
+}
+
 bool UMovieSceneBoolTrack::SupportsType(TSubclassOf<UMovieSceneSection> SectionClass) const
 {
 	return SectionClass == UMovieSceneBoolSection::StaticClass();
@@ -20,5 +37,80 @@ UMovieSceneSection* UMovieSceneBoolTrack::CreateNewSection()
 
 FMovieSceneEvalTemplatePtr UMovieSceneBoolTrack::CreateTemplateForSection(const UMovieSceneSection& InSection) const
 {
-	return FMovieSceneBoolPropertySectionTemplate(*CastChecked<const UMovieSceneBoolSection>(&InSection), *this);
+	if (UE::MovieScene::GCompileBooleanPropertyLegacyTemplate)
+	{
+		return FMovieSceneBoolPropertySectionTemplate(*CastChecked<const UMovieSceneBoolSection>(&InSection), *this);
+	}
+	return FMovieSceneEvalTemplatePtr();
+}
+
+void UMovieSceneBoolTrack::ImportEntityImpl(UMovieSceneEntitySystemLinker* EntityLinker, const FEntityImportParams& Params, FImportedEntity* OutImportedEntity)
+{
+	using namespace UE::MovieScene;
+
+	if (GBooleanPropertyLegacyTemplate)
+	{
+		return;
+	}
+
+	const int32 SectionIndex = static_cast<int32>(Params.EntityID);
+
+	UMovieSceneBoolSection* BoolSection = CastChecked<UMovieSceneBoolSection>(Sections[SectionIndex]);
+
+	if (!BoolSection->GetChannel().HasAnyData())
+	{
+		return;
+	}
+
+	const FBuiltInComponentTypes* Components = FBuiltInComponentTypes::Get();
+	const FMovieSceneTracksComponentTypes* TracksComponents = FMovieSceneTracksComponentTypes::Get();
+
+	// Entity ID for the track is the index of the section, but for FPropertyTrackEntityImportHelper it expects 0/1.
+	//   Boolean tracks do not support edit condition overrides (for now), so we just force it to 0 for the main property entity
+	FEntityImportParams SectionParams = Params;
+	SectionParams.EntityID = 0;
+
+	FPropertyTrackEntityImportHelper(TracksComponents->Bool)
+		.Add(Components->BoolChannel, &BoolSection->GetChannel())
+		.Add(Components->BlenderType, UMovieScenePiecewiseBoolBlenderSystem::StaticClass())
+		.Commit(BoolSection, SectionParams, OutImportedEntity);
+
+	BoolSection->BuildDefaultComponents(EntityLinker, Params, OutImportedEntity);
 }
+
+bool UMovieSceneBoolTrack::PopulateEvaluationFieldImpl(const TRange<FFrameNumber>& InEffectiveRange, const FMovieSceneEvaluationFieldEntityMetaData& InMetaData, FMovieSceneEntityComponentFieldBuilder* OutFieldBuilder)
+{
+	for (const FMovieSceneTrackEvaluationFieldEntry& Entry : GetEvaluationField().Entries)
+	{
+		if (Entry.Section && IsRowEvalDisabled(Entry.Section->GetRowIndex()))
+		{
+			continue;
+		}
+
+		// This codepath should only ever execute for the highest level so we do not need to do any transformations
+		TRange<FFrameNumber> EffectiveRange = TRange<FFrameNumber>::Intersection(InEffectiveRange, Entry.Range);
+		if (!EffectiveRange.IsEmpty())
+		{
+			const int32 SectionIndex = Algo::IndexOf(Sections, Entry.Section);
+			if (SectionIndex == INDEX_NONE)
+			{
+				continue;
+			}
+
+			FMovieSceneEvaluationFieldEntityMetaData MetaData = InMetaData;
+
+			MetaData.ForcedTime = Entry.ForcedTime;
+			MetaData.Flags      = Entry.Flags;
+			MetaData.bEvaluateInSequencePreRoll  = EvalOptions.bEvaluateInPreroll;
+			MetaData.bEvaluateInSequencePostRoll = EvalOptions.bEvaluateInPostroll;
+			MetaData.Condition = MovieSceneHelpers::GetSequenceCondition(this, Entry.Section, true);
+
+			const int32 EntityIndex   = OutFieldBuilder->FindOrAddEntity(this, static_cast<uint32>(SectionIndex));
+			const int32 MetaDataIndex = OutFieldBuilder->AddMetaData(MetaData);
+
+			OutFieldBuilder->AddPersistentEntity(EffectiveRange, EntityIndex, MetaDataIndex);
+		}
+	}
+
+	return true;
+}
\ No newline at end of file
diff --git a/Engine/Source/Runtime/MovieSceneTracks/Public/Tracks/MovieSceneBoolTrack.h b/Engine/Source/Runtime/MovieSceneTracks/Public/Tracks/MovieSceneBoolTrack.h
--- a/Engine/Source/Runtime/MovieSceneTracks/Public/Tracks/MovieSceneBoolTrack.h
+++ b/Engine/Source/Runtime/MovieSceneTracks/Public/Tracks/MovieSceneBoolTrack.h
@@ -6,6 +6,7 @@
 #include "UObject/ObjectMacros.h"
 #include "Tracks/MovieScenePropertyTrack.h"
 #include "Compilation/IMovieSceneTrackTemplateProducer.h"
+#include "EntitySystem/IMovieSceneEntityProvider.h"
 #include "MovieSceneBoolTrack.generated.h"
 
 /**
@@ -15,6 +16,7 @@ UCLASS( MinimalAPI )
 class UMovieSceneBoolTrack
 	: public UMovieScenePropertyTrack
 	, public IMovieSceneTrackTemplateProducer
+	, public IMovieSceneEntityProvider
 {
 	GENERATED_BODY()
 
@@ -25,4 +27,8 @@ public:
 	virtual bool SupportsType(TSubclassOf<UMovieSceneSection> SectionClass) const override;
 	virtual UMovieSceneSection* CreateNewSection() override;
 	virtual FMovieSceneEvalTemplatePtr CreateTemplateForSection(const UMovieSceneSection& InSection) const override;
+
+	//~ IMovieSceneEntityProvider interface
+	virtual void ImportEntityImpl(UMovieSceneEntitySystemLinker* EntityLinker, const FEntityImportParams& Params, FImportedEntity* OutImportedEntity) override;
+	virtual bool PopulateEvaluationFieldImpl(const TRange<FFrameNumber>& EffectiveRange, const FMovieSceneEvaluationFieldEntityMetaData& InMetaData, FMovieSceneEntityComponentFieldBuilder* OutFieldBuilder) override;
 };
