Sequencer: Improved property traits variants, added type-erased property values, and unified many property tack editors

diff --git a/Engine/Source/Runtime/MovieSceneTracks/Private/MovieSceneTracksComponentTypes.cpp b/Engine/Source/Runtime/MovieSceneTracks/Private/MovieSceneTracksComponentTypes.cpp
--- a/Engine/Source/Runtime/MovieSceneTracks/Private/MovieSceneTracksComponentTypes.cpp
+++ b/Engine/Source/Runtime/MovieSceneTracks/Private/MovieSceneTracksComponentTypes.cpp
@@ -5,19 +5,28 @@
 #include "Camera/CameraShakeSourceComponent.h"
 #include "Components/SkeletalMeshComponent.h"
 #include "MovieSceneTracksCustomAccessors.h"
+#include "Tracks/MovieSceneIntegerTrack.h"
+#include "Tracks/MovieSceneFloatTrack.h"
+#include "Tracks/MovieSceneDoubleTrack.h"
+#include "Tracks/MovieSceneColorTrack.h"
+#include "Tracks/MovieSceneVectorTrack.h"
+#include "Tracks/MovieSceneTransformTrack.h"
+#include "Tracks/MovieSceneEulerTransformTrack.h"
 #include "EntitySystem/BuiltInComponentTypes.h"
 #include "EntitySystem/MovieSceneEntityManager.h"
 #include "EntitySystem/MovieSceneEntitySystemLinker.h"
 #include "EntitySystem/MovieSceneComponentRegistry.h"
 #include "EntitySystem/MovieSceneBlenderSystem.h"
+#include "Channels/MovieSceneUnpackedChannelValues.h"
 #include "Systems/MovieScenePiecewiseBoolBlenderSystem.h"
 #include "Systems/MovieScenePiecewiseByteBlenderSystem.h"
 #include "Systems/MovieScenePiecewiseIntegerBlenderSystem.h"
 #include "Systems/MovieScenePiecewiseDoubleBlenderSystem.h"
 #include "Systems/MovieScenePiecewiseDoubleBlenderSystem.h"
 #include "EntitySystem/MovieScenePropertyComponentHandler.h"
+#include "EntitySystem/MovieSceneVariantPropertyComponentHandler.h"
 #include "EntitySystem/MovieSceneEntityFactoryTemplates.h"
-#include "EntitySystem/MovieScenePropertyMetaDataTraits.inl"
+#include "EntitySystem/MovieScenePropertyMetaData.inl"
 #include "PreAnimatedState/MovieScenePreAnimatedComponentTransformStorage.h"
 #include "Systems/MovieSceneColorPropertySystem.h"
 #include "Systems/MovieSceneVectorPropertySystem.h"
@@ -26,6 +35,7 @@
 #include "Materials/MaterialParameterCollection.h"
 #include "Misc/App.h"
 #include "PhysicsEngine/BodyInstance.h"
+#include "EntitySystem/MovieSceneComponentTypeIDs.h"
 
 #include UE_INLINE_GENERATED_CPP_BY_NAME(MovieSceneTracksComponentTypes)
 
@@ -33,6 +43,60 @@ namespace UE
 {
 namespace MovieScene
 {
+/* ---------------------------------------------------------------------------
+ * Integer conversion functions
+ * ---------------------------------------------------------------------------*/
+void ConvertOperationalProperty(int64 In, int8&  Out)
+{
+	static constexpr int8 Min = std::numeric_limits<int8>::lowest();
+	static constexpr int8 Max = std::numeric_limits<int8>::max();
+	if (In > static_cast<int64>(Max))
+	{
+		Out = Max;
+	}
+	else if (In < static_cast<int64>(Min))
+	{
+		Out = Min;
+	}
+	else
+	{
+		Out = static_cast<int8>(In);
+	}
+}
+void ConvertOperationalProperty(int64 In, int16& Out)
+{
+	static constexpr int16 Min = std::numeric_limits<int16>::lowest();
+	static constexpr int16 Max = std::numeric_limits<int16>::max();
+	if (In > static_cast<int64>(Max))
+	{
+		Out = Max;
+	}
+	else if (In < static_cast<int64>(Min))
+	{
+		Out = Min;
+	}
+	else
+	{
+		Out = static_cast<int16>(In);
+	}
+}
+void ConvertOperationalProperty(int64 In, int32& Out)
+{
+	static constexpr int32 Min = std::numeric_limits<int32>::lowest();
+	static constexpr int32 Max = std::numeric_limits<int32>::max();
+	if (In > static_cast<int64>(Max))
+	{
+		Out = Max;
+	}
+	else if (In < static_cast<int64>(Min))
+	{
+		Out = Min;
+	}
+	else
+	{
+		Out = static_cast<int32>(In);
+	}
+}
 
 /* ---------------------------------------------------------------------------
  * Transform conversion functions
@@ -326,49 +390,6 @@ struct FFloatHandler : TPropertyComponentHandler<FFloatPropertyTraits, double>
 {
 };
 
-struct FColorHandler : TPropertyComponentHandler<FColorPropertyTraits, double, double, double, double>
-{
-	virtual void DispatchInitializePropertyMetaDataTasks(const FPropertyDefinition& Definition, FSystemTaskPrerequisites& InPrerequisites, FSystemSubsequentTasks& Subsequents, UMovieSceneEntitySystemLinker* Linker) override
-	{
-		FBuiltInComponentTypes* BuiltInComponents = FBuiltInComponentTypes::Get();
-		FMovieSceneTracksComponentTypes* TrackComponents = FMovieSceneTracksComponentTypes::Get();
-
-		FEntityTaskBuilder()
-		.Read(BuiltInComponents->BoundObject)
-		.Read(BuiltInComponents->PropertyBinding)
-		.Write(TrackComponents->Color.MetaDataComponents.GetType<0>())
-		.FilterAll({ BuiltInComponents->Tags.NeedsLink })
-		.Iterate_PerEntity(&Linker->EntityManager, [](UObject* Object, const FMovieScenePropertyBinding& Binding, EColorPropertyType& OutType)
-		{
-			FTrackInstancePropertyBindings Bindings(Binding.PropertyName, Binding.PropertyPath.ToString());
-			const UStruct* PropertyStruct = Bindings.GetPropertyStruct(*Object);
-			if (ensure(PropertyStruct))
-			{
-				if (PropertyStruct == TBaseStructure<FColor>::Get())
-				{
-					// We assume the color we get back is in sRGB, assigning it to a linear color will implicitly
-					// convert it to a linear color instead of using ReinterpretAsLinear which will just change the
-					// bytes into floats using divide by 255.
-					OutType = EColorPropertyType::Color;
-				}
-				else if (PropertyStruct == TBaseStructure<FSlateColor>::Get())
-				{
-					OutType = EColorPropertyType::Slate;
-				}
-				else
-				{
-					ensure(PropertyStruct == TBaseStructure<FLinearColor>::Get());
-					OutType = EColorPropertyType::Linear;
-				}
-			}
-			else
-			{
-				OutType = EColorPropertyType::Linear;
-			}
-		});
-	}
-};
-
 struct FFloatParameterHandler : TPropertyComponentHandler<FFloatParameterTraits, double>
 {
 	virtual TSharedPtr<IInitialValueProcessor> MakeInitialValueProcessor(const FPropertyDefinition& Definition) override
@@ -385,88 +406,6 @@ struct FColorParameterHandler : TPropertyComponentHandler<FColorParameterTraits,
 	}
 };
 
-
-struct FFloatVectorHandler : TPropertyComponentHandler<FFloatVectorPropertyTraits, double, double, double, double>
-{
-	virtual void DispatchInitializePropertyMetaDataTasks(const FPropertyDefinition& Definition, FSystemTaskPrerequisites& InPrerequisites, FSystemSubsequentTasks& Subsequents, UMovieSceneEntitySystemLinker* Linker) override
-	{
-		FBuiltInComponentTypes* BuiltInComponents = FBuiltInComponentTypes::Get();
-		FMovieSceneTracksComponentTypes* TrackComponents = FMovieSceneTracksComponentTypes::Get();
-
-		FEntityTaskBuilder()
-		.Read(BuiltInComponents->BoundObject)
-		.Read(BuiltInComponents->PropertyBinding)
-		.Write(TrackComponents->FloatVector.MetaDataComponents.GetType<0>())
-		.FilterAll({ BuiltInComponents->Tags.NeedsLink })
-		.Iterate_PerEntity(&Linker->EntityManager, [](UObject* Object, const FMovieScenePropertyBinding& Binding, FVectorPropertyMetaData& OutMetaData)
-		{
-			FTrackInstancePropertyBindings Bindings(Binding.PropertyName, Binding.PropertyPath.ToString());
-			const UStruct* PropertyStruct = Bindings.GetPropertyStruct(*Object);
-			if (ensure(PropertyStruct))
-			{
-				if (PropertyStruct == TVariantStructure<FVector2f>::Get())
-				{
-					OutMetaData.NumChannels = 2;
-				}
-				else if (PropertyStruct == TVariantStructure<FVector3f>::Get())
-				{
-					OutMetaData.NumChannels = 3;
-				}
-				else
-				{
-					ensure(PropertyStruct == TVariantStructure<FVector4f>::Get());
-					OutMetaData.NumChannels = 4;
-				}
-			}
-			else
-			{
-				OutMetaData.NumChannels = 4;
-			}
-		});
-	}
-};
-
-
-struct FDoubleVectorHandler : TPropertyComponentHandler<FDoubleVectorPropertyTraits, double, double, double, double>
-{
-	virtual void DispatchInitializePropertyMetaDataTasks(const FPropertyDefinition& Definition, FSystemTaskPrerequisites& InPrerequisites, FSystemSubsequentTasks& Subsequents, UMovieSceneEntitySystemLinker* Linker) override
-	{
-		FBuiltInComponentTypes* BuiltInComponents = FBuiltInComponentTypes::Get();
-		FMovieSceneTracksComponentTypes* TrackComponents = FMovieSceneTracksComponentTypes::Get();
-
-		FEntityTaskBuilder()
-		.Read(BuiltInComponents->BoundObject)
-		.Read(BuiltInComponents->PropertyBinding)
-		.Write(TrackComponents->DoubleVector.MetaDataComponents.GetType<0>())
-		.FilterAll({ BuiltInComponents->Tags.NeedsLink })
-		.Iterate_PerEntity(&Linker->EntityManager, [](UObject* Object, const FMovieScenePropertyBinding& Binding, FVectorPropertyMetaData& OutMetaData)
-		{
-			FTrackInstancePropertyBindings Bindings(Binding.PropertyName, Binding.PropertyPath.ToString());
-			const UStruct* PropertyStruct = Bindings.GetPropertyStruct(*Object);
-			if (ensure(PropertyStruct))
-			{
-				if (PropertyStruct == TBaseStructure<FVector2D>::Get())
-				{
-					OutMetaData.NumChannels = 2;
-				}
-				else if (PropertyStruct == TBaseStructure<FVector>::Get() || PropertyStruct == TVariantStructure<FVector3d>::Get())
-				{
-					OutMetaData.NumChannels = 3;
-				}
-				else
-				{
-					ensure(PropertyStruct == TBaseStructure<FVector4>::Get() || PropertyStruct == TVariantStructure<FVector4d>::Get());
-					OutMetaData.NumChannels = 4;
-				}
-			}
-			else
-			{
-				OutMetaData.NumChannels = 4;
-			}
-		});
-	}
-};
-
 struct FBoolHandler : TPropertyComponentHandler<FBoolPropertyTraits, bool>
 {
 	virtual void DispatchInitializePropertyMetaDataTasks(const FPropertyDefinition& Definition, FSystemTaskPrerequisites& InPrerequisites, FSystemSubsequentTasks& Subsequents, UMovieSceneEntitySystemLinker* Linker) override
@@ -550,6 +489,7 @@ struct FObjectHandler : TPropertyComponentHandler<FObjectPropertyTraits, FObject
 
 FMovieSceneTracksComponentTypes::FMovieSceneTracksComponentTypes()
 {
+	FBuiltInComponentTypes* BuiltInComponents = FBuiltInComponentTypes::Get();
 	FComponentRegistry* ComponentRegistry = UMovieSceneEntitySystemLinker::GetComponents();
 
 	ComponentRegistry->NewPropertyType(Bool, TEXT("bool"));
@@ -583,10 +523,14 @@ FMovieSceneTracksComponentTypes::FMovieSceneTracksComponentTypes()
 	ComponentRegistry->NewPropertyType(Parameters.Color, TEXT("Color parameter"));
 	ComponentRegistry->NewPropertyType(Parameters.Transform, TEXT("Transform parameter"));
 
+	Color.MetaDataComponents.Initialize(BuiltInComponents->VariantPropertyTypeIndex);
+	Double.MetaDataComponents.Initialize(BuiltInComponents->VariantPropertyTypeIndex);
+	Integer.MetaDataComponents.Initialize(BuiltInComponents->VariantPropertyTypeIndex);
+	Transform.MetaDataComponents.Initialize(BuiltInComponents->VariantPropertyTypeIndex);
+	FloatVector.MetaDataComponents.Initialize(BuiltInComponents->VariantPropertyTypeIndex);
+	DoubleVector.MetaDataComponents.Initialize(BuiltInComponents->VariantPropertyTypeIndex);
+
 	Bool.MetaDataComponents.Initialize(ComponentRegistry, TEXT("Bool Bitfield"));
-	Color.MetaDataComponents.Initialize(ComponentRegistry, TEXT("Color Type"));
-	FloatVector.MetaDataComponents.Initialize(ComponentRegistry, TEXT("Num Float Vector Channels"));
-	DoubleVector.MetaDataComponents.Initialize(ComponentRegistry, TEXT("Num Double Vector Channels"));
 	Object.MetaDataComponents.Initialize(ComponentRegistry, TEXT("Object Class"));
 
 	ComponentRegistry->NewComponentType(&QuaternionRotationChannel[0], TEXT("Quaternion Rotation Channel 0"));
@@ -650,8 +594,6 @@ FMovieSceneTracksComponentTypes::FMovieSceneTracksComponentTypes()
 
 	Tags.AnimMixerPoseProducer = ComponentRegistry->NewTag(TEXT("Anim Mixer Pose Producer"));
 
-	FBuiltInComponentTypes* BuiltInComponents = FBuiltInComponentTypes::Get();
-
 	// --------------------------------------------------------------------------------------------
 	// Set up bool properties
 	BuiltInComponents->PropertyRegistry.DefineProperty(Bool, TEXT("Apply bool Properties"))
@@ -662,7 +604,7 @@ FMovieSceneTracksComponentTypes::FMovieSceneTracksComponentTypes()
 
 	// --------------------------------------------------------------------------------------------
 	// Set up FTransform properties
-	BuiltInComponents->PropertyRegistry.DefineCompositeProperty(Transform, TEXT("Apply FTransform Properties"))
+	BuiltInComponents->PropertyRegistry.DefineCompositeProperty(Transform, TEXT("Apply Transform Properties"))
 	.AddComposite(BuiltInComponents->DoubleResult[0], &FIntermediate3DTransform::T_X)
 	.AddComposite(BuiltInComponents->DoubleResult[1], &FIntermediate3DTransform::T_Y)
 	.AddComposite(BuiltInComponents->DoubleResult[2], &FIntermediate3DTransform::T_Z)
@@ -672,6 +614,8 @@ FMovieSceneTracksComponentTypes::FMovieSceneTracksComponentTypes()
 	.AddComposite(BuiltInComponents->DoubleResult[6], &FIntermediate3DTransform::S_X)
 	.AddComposite(BuiltInComponents->DoubleResult[7], &FIntermediate3DTransform::S_Y)
 	.AddComposite(BuiltInComponents->DoubleResult[8], &FIntermediate3DTransform::S_Z)
+	.SetCustomAccessors(&Accessors.Transform)
+	.SetDefaultTrackType(UMovieSceneTransformTrack::StaticClass())
 	.SetBlenderSystem<UMovieScenePiecewiseDoubleBlenderSystem>()
 	.Commit();
 
@@ -705,6 +649,7 @@ FMovieSceneTracksComponentTypes::FMovieSceneTracksComponentTypes()
 	// Set up integer properties
 	BuiltInComponents->PropertyRegistry.DefineProperty(Integer, TEXT("Apply Integer Properties"))
 	.AddSoleChannel(BuiltInComponents->IntegerResult)
+	.SetDefaultTrackType(UMovieSceneIntegerTrack::StaticClass())
 	.SetBlenderSystem<UMovieScenePiecewiseIntegerBlenderSystem>()
 	.SetCustomAccessors(&Accessors.Integer)
 	.Commit();
@@ -713,6 +658,7 @@ FMovieSceneTracksComponentTypes::FMovieSceneTracksComponentTypes()
 	// Set up float properties
 	BuiltInComponents->PropertyRegistry.DefineProperty(Float, TEXT("Apply float Properties"))
 	.AddSoleChannel(BuiltInComponents->DoubleResult[0])
+	.SetDefaultTrackType(UMovieSceneFloatTrack::StaticClass())
 	.SetBlenderSystem<UMovieScenePiecewiseDoubleBlenderSystem>()
 	.SetCustomAccessors(&Accessors.Float)
 	.Commit(FFloatHandler());
@@ -721,6 +667,7 @@ FMovieSceneTracksComponentTypes::FMovieSceneTracksComponentTypes()
 	// Set up double properties
 	BuiltInComponents->PropertyRegistry.DefineProperty(Double, TEXT("Apply Double Properties"))
 	.AddSoleChannel(BuiltInComponents->DoubleResult[0])
+	.SetDefaultTrackType(UMovieSceneDoubleTrack::StaticClass())
 	.SetBlenderSystem<UMovieScenePiecewiseDoubleBlenderSystem>()
 	.SetCustomAccessors(&Accessors.Double)
 	.Commit();
@@ -732,9 +679,10 @@ FMovieSceneTracksComponentTypes::FMovieSceneTracksComponentTypes()
 	.AddComposite(BuiltInComponents->DoubleResult[1], &FIntermediateColor::G)
 	.AddComposite(BuiltInComponents->DoubleResult[2], &FIntermediateColor::B)
 	.AddComposite(BuiltInComponents->DoubleResult[3], &FIntermediateColor::A)
+	.SetDefaultTrackType(UMovieSceneColorTrack::StaticClass())
 	.SetBlenderSystem<UMovieScenePiecewiseDoubleBlenderSystem>()
 	.SetCustomAccessors(&Accessors.Color)
-	.Commit(FColorHandler());
+	.Commit();
 
 	// --------------------------------------------------------------------------------------------
 	// Set up string properties
@@ -799,9 +747,10 @@ FMovieSceneTracksComponentTypes::FMovieSceneTracksComponentTypes()
 	.AddComposite(BuiltInComponents->DoubleResult[1], &FFloatIntermediateVector::Y)
 	.AddComposite(BuiltInComponents->DoubleResult[2], &FFloatIntermediateVector::Z)
 	.AddComposite(BuiltInComponents->DoubleResult[3], &FFloatIntermediateVector::W)
+	.SetDefaultTrackType(UMovieSceneFloatVectorTrack::StaticClass())
 	.SetBlenderSystem<UMovieScenePiecewiseDoubleBlenderSystem>()
 	.SetCustomAccessors(&Accessors.FloatVector)
-	.Commit(FFloatVectorHandler());
+	.Commit();
 
 	BuiltInComponents->PropertyRegistry.DefineCompositeProperty(Parameters.Vector2, TEXT("Vector2 Parameters"))
 	.AddComposite(BuiltInComponents->DoubleResult[0], &FFloatIntermediateVector::X)
@@ -821,9 +770,10 @@ FMovieSceneTracksComponentTypes::FMovieSceneTracksComponentTypes()
 	.AddComposite(BuiltInComponents->DoubleResult[1], &FDoubleIntermediateVector::Y)
 	.AddComposite(BuiltInComponents->DoubleResult[2], &FDoubleIntermediateVector::Z)
 	.AddComposite(BuiltInComponents->DoubleResult[3], &FDoubleIntermediateVector::W)
+	.SetDefaultTrackType(UMovieSceneDoubleVectorTrack::StaticClass())
 	.SetBlenderSystem<UMovieScenePiecewiseDoubleBlenderSystem>()
 	.SetCustomAccessors(&Accessors.DoubleVector)
-	.Commit(FDoubleVectorHandler());
+	.Commit();
 
 	// --------------------------------------------------------------------------------------------
 	// Set up FEulerTransform properties
@@ -837,6 +787,7 @@ FMovieSceneTracksComponentTypes::FMovieSceneTracksComponentTypes()
 	.AddComposite(BuiltInComponents->DoubleResult[6], &FIntermediate3DTransform::S_X)
 	.AddComposite(BuiltInComponents->DoubleResult[7], &FIntermediate3DTransform::S_Y)
 	.AddComposite(BuiltInComponents->DoubleResult[8], &FIntermediate3DTransform::S_Z)
+	.SetDefaultTrackType(UMovieSceneEulerTransformTrack::StaticClass())
 	.SetBlenderSystem<UMovieScenePiecewiseDoubleBlenderSystem>()
 	.Commit();
 
diff --git a/Engine/Source/Runtime/MovieSceneTracks/Private/MovieSceneTracksCustomAccessors.cpp b/Engine/Source/Runtime/MovieSceneTracks/Private/MovieSceneTracksCustomAccessors.cpp
--- a/Engine/Source/Runtime/MovieSceneTracks/Private/MovieSceneTracksCustomAccessors.cpp
+++ b/Engine/Source/Runtime/MovieSceneTracks/Private/MovieSceneTracksCustomAccessors.cpp
@@ -73,21 +73,18 @@ void SetComponentTransformAndVelocity(UObject* Object, const FIntermediate3DTran
 	UE::MovieScene::FIntermediate3DTransform::ApplyTransformTo(CastChecked<USceneComponent>(Object), InTransform);
 }
 
-FIntermediateColor GetLightComponentLightColor(const UObject* Object, EColorPropertyType InColorType)
+FIntermediateColor GetLightComponentLightColor(const UObject* Object)
 {
-	ensure(InColorType == EColorPropertyType::Color);
-
 	const ULightComponent* LightComponent = CastChecked<const ULightComponent>(Object);
 	return FIntermediateColor(LightComponent->GetLightColor());
 }
 
-void SetLightComponentLightColor(UObject* Object, EColorPropertyType InColorType, const FIntermediateColor& InColor)
+void SetLightComponentLightColor(UObject* Object, const FIntermediateColor& InColor)
 {
 	// This is a little esoteric - ULightComponentBase::LightColor is the UPROPERTY that generates the meta-data
 	// for this custom callback, but it is an FColor, even though the public get/set functions expose it as an
 	// FLinearColor. FIntermediateColor is always blended and dealt with in linear space, so it's fine to 
 	// simply reinterpret the color
-	ensure(InColorType == EColorPropertyType::Color);
 
 	const bool bConvertBackToSRgb = true;
 	ULightComponent* LightComponent = CastChecked<ULightComponent>(Object);
@@ -142,22 +139,18 @@ void SetDirectionalLightComponentLightSourceAngle(UObject* Object, float InLight
 	DirectionalLightComponent->SetLightSourceAngle(InLightSourceAngle);
 }
 
-FIntermediateColor GetSkyLightComponentLightColor(const UObject* Object, EColorPropertyType InColorType)
+FIntermediateColor GetSkyLightComponentLightColor(const UObject* Object)
 {
-	ensure(InColorType == EColorPropertyType::Color);
-
 	const USkyLightComponent* SkyLightComponent = CastChecked<const USkyLightComponent>(Object);
 	return FIntermediateColor(SkyLightComponent->GetLightColor());
 }
 
-void SetSkyLightComponentLightColor(UObject* Object, EColorPropertyType InColorType, const FIntermediateColor& InColor)
+void SetSkyLightComponentLightColor(UObject* Object, const FIntermediateColor& InColor)
 {
 	// This is a little esoteric - ULightComponentBase::LightColor is the UPROPERTY that generates the meta-data
 	// for this custom callback, but it is an FColor, even though the public get/set functions expose it as an
 	// FLinearColor. FIntermediateColor is always blended and dealt with in linear space, so it's fine to 
 	// simply reinterpret the color
-	ensure(InColorType == EColorPropertyType::Color);
-
 	USkyLightComponent* SkyLightComponent = CastChecked<USkyLightComponent>(Object);
 	SkyLightComponent->SetLightColor(InColor.GetLinearColor());
 }
@@ -210,18 +203,14 @@ void SetSkyAtmosphereComponentOtherAbsorptionScale(UObject* Object, float InOthe
 	SkyAtmosphereComponent->SetOtherAbsorptionScale(InOtherAbsorptionScale);
 }
 
-FIntermediateColor GetSkyAtmosphereComponentOtherAbsorption(const UObject* Object, EColorPropertyType InColorType)
+FIntermediateColor GetSkyAtmosphereComponentOtherAbsorption(const UObject* Object)
 {
-	ensure(InColorType == EColorPropertyType::Linear);
-	
 	const USkyAtmosphereComponent* SkyAtmosphereComponent = CastChecked<const USkyAtmosphereComponent>(Object);
 	return FIntermediateColor(SkyAtmosphereComponent->OtherAbsorption);
 }
 
-void SetSkyAtmosphereComponentOtherAbsorption(UObject* Object, EColorPropertyType InColorType, const FIntermediateColor& InColor)
+void SetSkyAtmosphereComponentOtherAbsorption(UObject* Object, const FIntermediateColor& InColor)
 {
-	ensure(InColorType == EColorPropertyType::Linear);
-
 	USkyAtmosphereComponent* SkyAtmosphereComponent = CastChecked<USkyAtmosphereComponent>(Object);
 	SkyAtmosphereComponent->SetOtherAbsorption(InColor.GetLinearColor());
 }
@@ -262,82 +251,66 @@ void SetSkyAtmosphereComponentAerialPerspectiveViewDistanceScale(UObject* Object
 	SkyAtmosphereComponent->SetAerialPespectiveViewDistanceScale(InAerialPerspectiveViewDistanceScale);
 }
 
-FIntermediateColor GetSkyAtmosphereComponentMieAbsorption(const UObject* Object, EColorPropertyType InColorType)
+FIntermediateColor GetSkyAtmosphereComponentMieAbsorption(const UObject* Object)
 {
 	// USkyAtmosphereComponent::MieAbsorption is a FLinearColor. FIntermediateColor is always blended and dealt with
 	// in linear space so it should be safe to reinterpret it as linear.
-	ensure(InColorType == EColorPropertyType::Linear);
-
 	const USkyAtmosphereComponent* SkyAtmosphereComponent = CastChecked<const USkyAtmosphereComponent>(Object);
 	return FIntermediateColor(SkyAtmosphereComponent->MieAbsorption);
 }
 
-void SetSkyAtmosphereComponentMieAbsorption(UObject* Object, EColorPropertyType InColorType, const FIntermediateColor& InColor)
+void SetSkyAtmosphereComponentMieAbsorption(UObject* Object, const FIntermediateColor& InColor)
 {
 	// USkyAtmosphereComponent::MieAbsorption is a FLinearColor. FIntermediateColor is always blended and dealt with
 	// in linear space so it should be safe to reinterpret it as linear.
-	ensure(InColorType == EColorPropertyType::Linear);
-
 	USkyAtmosphereComponent* SkyAtmosphereComponent = CastChecked<USkyAtmosphereComponent>(Object);
 	SkyAtmosphereComponent->SetMieAbsorption(InColor.GetLinearColor());
 }
 
-FIntermediateColor GetSkyAtmosphereComponentMieScattering(const UObject* Object, EColorPropertyType InColorType)
+FIntermediateColor GetSkyAtmosphereComponentMieScattering(const UObject* Object)
 {
 	// USkyAtmosphereComponent::MieScattering is a FLinearColor. FIntermediateColor is always blended and dealt with
 	// in linear space so it should be safe to reinterpret it as linear.
-	ensure(InColorType == EColorPropertyType::Linear);
-
 	const USkyAtmosphereComponent* SkyAtmosphereComponent = CastChecked<const USkyAtmosphereComponent>(Object);
 	return FIntermediateColor(SkyAtmosphereComponent->MieScattering);
 }
 
-void SetSkyAtmosphereComponentMieScattering(UObject* Object, EColorPropertyType InColorType, const FIntermediateColor& InColor)
+void SetSkyAtmosphereComponentMieScattering(UObject* Object, const FIntermediateColor& InColor)
 {
 	// USkyAtmosphereComponent::MieScattering is a FLinearColor. FIntermediateColor is always blended and dealt with
 	// in linear space so it should be safe to reinterpret it as linear.
-	ensure(InColorType == EColorPropertyType::Linear);
-
 	USkyAtmosphereComponent* SkyAtmosphereComponent = CastChecked<USkyAtmosphereComponent>(Object);
 	SkyAtmosphereComponent->SetMieScattering(InColor.GetLinearColor());
 }
 
-FIntermediateColor GetSkyAtmosphereComponentRayleighScattering(const UObject* Object, EColorPropertyType InColorType)
+FIntermediateColor GetSkyAtmosphereComponentRayleighScattering(const UObject* Object)
 {
 	// USkyAtmosphereComponent::RayleighScattering is a FLinearColor. FIntermediateColor is always blended and dealt with
 	// in linear space so it should be safe to reinterpret it as linear.
-	ensure(InColorType == EColorPropertyType::Linear);
-
 	const USkyAtmosphereComponent* SkyAtmosphereComponent = CastChecked<const USkyAtmosphereComponent>(Object);
 	return FIntermediateColor(SkyAtmosphereComponent->RayleighScattering);
 }
 
-void SetSkyAtmosphereComponentRayleighScattering(UObject* Object, EColorPropertyType InColorType, const FIntermediateColor& InColor)
+void SetSkyAtmosphereComponentRayleighScattering(UObject* Object, const FIntermediateColor& InColor)
 {
 	// USkyAtmosphereComponent::RayleighScattering is a FLinearColor. FIntermediateColor is always blended and dealt with
 	// in linear space so it should be safe to reinterpret it as linear.
-	ensure(InColorType == EColorPropertyType::Linear);
-
 	USkyAtmosphereComponent* SkyAtmosphereComponent = CastChecked<USkyAtmosphereComponent>(Object);
 	SkyAtmosphereComponent->SetRayleighScattering(InColor.GetLinearColor());
 }
 
-FIntermediateColor GetSkyAtmosphereComponentSkyLuminanceFactor(const UObject* Object, EColorPropertyType InColorType)
+FIntermediateColor GetSkyAtmosphereComponentSkyLuminanceFactor(const UObject* Object)
 {
 	// USkyAtmosphereComponent::SkyLuminanceFactor is a FLinearColor. FIntermediateColor is always blended and dealt with
 	// in linear space so it should be safe to reinterpret it as linear.
-	ensure(InColorType == EColorPropertyType::Linear);
-
 	const USkyAtmosphereComponent* SkyAtmosphereComponent = CastChecked<const USkyAtmosphereComponent>(Object);
 	return FIntermediateColor(SkyAtmosphereComponent->SkyLuminanceFactor);
 }
 
-void SetSkyAtmosphereComponentSkyLuminanceFactor(UObject* Object, EColorPropertyType InColorType, const FIntermediateColor& InColor)
+void SetSkyAtmosphereComponentSkyLuminanceFactor(UObject* Object, const FIntermediateColor& InColor)
 {
 	// USkyAtmosphereComponent::SkyLuminanceFactor is a FLinearColor. FIntermediateColor is always blended and dealt with
 	// in linear space so it should be safe to reinterpret it as linear.
-	ensure(InColorType == EColorPropertyType::Linear);
-
 	USkyAtmosphereComponent* SkyAtmosphereComponent = CastChecked<USkyAtmosphereComponent>(Object);
 	SkyAtmosphereComponent->SetSkyLuminanceFactor(InColor.GetLinearColor());
 }
@@ -450,54 +423,42 @@ void SetExponentialHeightFogComponentVolumetricFogNearFadeInDistance(UObject* Ob
 	ExponentialHeightFogComponent->SetVolumetricFogNearFadeInDistance(InVolumetricFogNearFadeInDistance);
 }
 	
-FIntermediateColor GetExponentialHeightFogComponentDirectionalInscatteringColor(const UObject* Object, EColorPropertyType InColorType)
+FIntermediateColor GetExponentialHeightFogComponentDirectionalInscatteringColor(const UObject* Object)
 {
-	ensure(InColorType == EColorPropertyType::Linear);
-	
 	const UExponentialHeightFogComponent* ExponentialHeightFogComponent = CastChecked<const UExponentialHeightFogComponent>(Object);
 	return FIntermediateColor(ExponentialHeightFogComponent->DirectionalInscatteringLuminance);
 }
 
-void SetExponentialHeightFogComponentDirectionalInscatteringColor(UObject* Object, EColorPropertyType InColorType, const FIntermediateColor& InColor)
+void SetExponentialHeightFogComponentDirectionalInscatteringColor(UObject* Object, const FIntermediateColor& InColor)
 {
-	ensure(InColorType == EColorPropertyType::Linear);
-	
 	UExponentialHeightFogComponent* ExponentialHeightFogComponent = CastChecked<UExponentialHeightFogComponent>(Object);
 	ExponentialHeightFogComponent->SetDirectionalInscatteringColor(InColor.GetLinearColor());
 }
 
-FIntermediateColor GetExponentialHeightFogComponentFogInscatteringColor(const UObject* Object, EColorPropertyType InColorType)
+FIntermediateColor GetExponentialHeightFogComponentFogInscatteringColor(const UObject* Object)
 {
-	ensure(InColorType == EColorPropertyType::Linear);
-	
 	const UExponentialHeightFogComponent* ExponentialHeightFogComponent = CastChecked<const UExponentialHeightFogComponent>(Object);
 	return FIntermediateColor(ExponentialHeightFogComponent->FogInscatteringLuminance);
 }
 
-void SetExponentialHeightFogComponentFogInscatteringColor(UObject* Object, EColorPropertyType InColorType, const FIntermediateColor& InColor)
+void SetExponentialHeightFogComponentFogInscatteringColor(UObject* Object, const FIntermediateColor& InColor)
 {
-	ensure(InColorType == EColorPropertyType::Linear);
-	
 	UExponentialHeightFogComponent* ExponentialHeightFogComponent = CastChecked<UExponentialHeightFogComponent>(Object);
 	ExponentialHeightFogComponent->SetFogInscatteringColor(InColor.GetLinearColor());
 }
 
-FIntermediateColor GetExponentialHeightFogVolumetricFogAlbedo(const UObject* Object, EColorPropertyType InColorType)
+FIntermediateColor GetExponentialHeightFogVolumetricFogAlbedo(const UObject* Object)
 {
 	// UExponentialHeightFog::VolumetricFogAlbedo is an FColor. FIntermediateColor is always blended and dealt with
 	// in linear space, so reinterpret the color.
-	ensure(InColorType == EColorPropertyType::Color);
-
 	const UExponentialHeightFogComponent* ExponentialHeightFogComponent = CastChecked<const UExponentialHeightFogComponent>(Object);
 	return FIntermediateColor(ExponentialHeightFogComponent->VolumetricFogAlbedo);
 }
 
-void SetExponentialHeightFogVolumetricFogAlbedo(UObject* Object, EColorPropertyType InColorType, const FIntermediateColor& InColor)
+void SetExponentialHeightFogVolumetricFogAlbedo(UObject* Object, const FIntermediateColor& InColor)
 {
 	// UExponentialHeightFog::VolumetricFogAlbedo is an FColor. FIntermediateColor is always blended and dealt with
 	// in linear space, so reinterpret the color back to FColor when setting via the public interface.
-	ensure(InColorType == EColorPropertyType::Color);
-
 	UExponentialHeightFogComponent* ExponentialHeightFogComponent = CastChecked<UExponentialHeightFogComponent>(Object);
 	ExponentialHeightFogComponent->SetVolumetricFogAlbedo(InColor.GetColor());
 }
@@ -550,18 +511,14 @@ void SetLocalFogVolumeComponentRadialFogExtinction(UObject* Object, float InRadi
 	LocalFogVolumeComponent->SetRadialFogExtinction(InRadialFogExtinction);
 }
 
-FIntermediateColor GetLocalFogVolumeComponentFogAlbedo(const UObject* Object, EColorPropertyType InColorType)
+FIntermediateColor GetLocalFogVolumeComponentFogAlbedo(const UObject* Object)
 {
-	ensure(InColorType == EColorPropertyType::Linear);
-
 	const ULocalFogVolumeComponent* LocalFogVolumeComponent = CastChecked<ULocalFogVolumeComponent>(Object);
 	return FIntermediateColor(LocalFogVolumeComponent->FogAlbedo);
 }
 
-void SetLocalFogVolumeComponentFogAlbedo(UObject* Object, EColorPropertyType InColorType, const FIntermediateColor& InColor)
+void SetLocalFogVolumeComponentFogAlbedo(UObject* Object, const FIntermediateColor& InColor)
 {
-	ensure(InColorType == EColorPropertyType::Linear);
-
 	ULocalFogVolumeComponent* LocalFogVolumeComponent = CastChecked<ULocalFogVolumeComponent>(Object);
 	LocalFogVolumeComponent->SetFogAlbedo(InColor.GetLinearColor());
 }
diff --git a/Engine/Source/Runtime/MovieSceneTracks/Private/Systems/IntegerChannelEvaluatorSystem.cpp b/Engine/Source/Runtime/MovieSceneTracks/Private/Systems/IntegerChannelEvaluatorSystem.cpp
--- a/Engine/Source/Runtime/MovieSceneTracks/Private/Systems/IntegerChannelEvaluatorSystem.cpp
+++ b/Engine/Source/Runtime/MovieSceneTracks/Private/Systems/IntegerChannelEvaluatorSystem.cpp
@@ -21,11 +21,11 @@ namespace MovieScene
 // Do we need to optimize for this case using something like the code below, while pessimizing the common (non-multi-bind) codepath??
 struct FEvaluateIntegerChannels
 {
-	static void ForEachEntity(FSourceIntegerChannel IntegerChannel, FFrameTime FrameTime, int32& OutResult)
+	static void ForEachEntity(FSourceIntegerChannel IntegerChannel, FFrameTime FrameTime, int64& OutResult)
 	{
 		if (!IntegerChannel.Source->Evaluate(FrameTime, OutResult))
 		{
-			OutResult = MIN_int32;
+			OutResult = MIN_int64;
 		}
 	}
 };
diff --git a/Engine/Source/Runtime/MovieSceneTracks/Private/Systems/MovieSceneAudioSystem.cpp b/Engine/Source/Runtime/MovieSceneTracks/Private/Systems/MovieSceneAudioSystem.cpp
--- a/Engine/Source/Runtime/MovieSceneTracks/Private/Systems/MovieSceneAudioSystem.cpp
+++ b/Engine/Source/Runtime/MovieSceneTracks/Private/Systems/MovieSceneAudioSystem.cpp
@@ -164,7 +164,7 @@ struct FGatherAudioInputs
 				double, double, double,
 				double, double, double,
 				double, double, double,
-				FString, int32, bool>
+				FString, int64, bool>
 			AudioInputResults) const
 	{
 		FAudioInputsBySectionKey& AudioInputsBySectionKey = AudioSystem->AudioInputsBySectionKey;
@@ -182,7 +182,7 @@ struct FGatherAudioInputs
 			DoubleResults[8] = AudioInputResults.Get<8>();
 		}
 		const FString* StringResults = AudioInputResults.Get<9>();
-		const int32* IntegerResults = AudioInputResults.Get<10>();
+		const int64* IntegerResults = AudioInputResults.Get<10>();
 		const bool* BoolResults = AudioInputResults.Get<11>();
 
 		for (int32 Index = 0; Index < Allocation->Num(); ++Index)
diff --git a/Engine/Source/Runtime/MovieSceneTracks/Private/Systems/MovieSceneBaseValueEvaluatorSystem.cpp b/Engine/Source/Runtime/MovieSceneTracks/Private/Systems/MovieSceneBaseValueEvaluatorSystem.cpp
--- a/Engine/Source/Runtime/MovieSceneTracks/Private/Systems/MovieSceneBaseValueEvaluatorSystem.cpp
+++ b/Engine/Source/Runtime/MovieSceneTracks/Private/Systems/MovieSceneBaseValueEvaluatorSystem.cpp
@@ -48,11 +48,11 @@ struct FEvaluateBaseByteValues
 
 struct FEvaluateBaseIntegerValues
 {
-	void ForEachEntity(FSourceIntegerChannel IntegerChannel, FFrameTime FrameTime, int32& OutResult)
+	void ForEachEntity(FSourceIntegerChannel IntegerChannel, FFrameTime FrameTime, int64& OutResult)
 	{
 		if (!IntegerChannel.Source->Evaluate(FrameTime, OutResult))
 		{
-			OutResult = MIN_int32;
+			OutResult = MIN_int64;
 		}
 	}
 };
diff --git a/Engine/Source/Runtime/MovieSceneTracks/Private/Systems/MovieScenePiecewiseIntegerBlenderSystem.cpp b/Engine/Source/Runtime/MovieSceneTracks/Private/Systems/MovieScenePiecewiseIntegerBlenderSystem.cpp
--- a/Engine/Source/Runtime/MovieSceneTracks/Private/Systems/MovieScenePiecewiseIntegerBlenderSystem.cpp
+++ b/Engine/Source/Runtime/MovieSceneTracks/Private/Systems/MovieScenePiecewiseIntegerBlenderSystem.cpp
@@ -27,8 +27,8 @@ struct FIntegerAccumulationTask
 		: AccumulationBuffer(*InAccumulationBuffer)
 	{}
 
-	/** Task entry point - iterates the allocation's headers and accumulates int32 results for any required components */
-	void ForEachAllocation(const FEntityAllocation* Allocation, TRead<FMovieSceneBlendChannelID> BlendIDs, TRead<int32> Integers, TReadOptional<double> EasingAndWeights) const
+	/** Task entry point - iterates the allocation's headers and accumulates int64 results for any required components */
+	void ForEachAllocation(const FEntityAllocation* Allocation, TRead<FMovieSceneBlendChannelID> BlendIDs, TRead<int64> Integers, TReadOptional<double> EasingAndWeights) const
 	{
 		static const FMovieSceneBlenderSystemID IntegerBlenderSystemID = UMovieSceneBlenderSystem::GetBlenderSystemID<UMovieScenePiecewiseIntegerBlenderSystem>();
 
@@ -45,7 +45,7 @@ struct FIntegerAccumulationTask
 				FIntegerBlendResult& Result = AccumulationBuffer[BlendID.ChannelID];
 
 				const double Weight = EasingAndWeights[Index];
-				Result.Total  += (int32)(Integers[Index] * Weight);
+				Result.Total  += (int64)(Integers[Index] * Weight);
 				Result.Weight += Weight;
 			}
 		}
@@ -86,7 +86,7 @@ struct FIntegerAdditiveFromBaseBlendTask
 		}
 	}
 
-	void ForEachAllocation(const FEntityAllocation* Allocation, TRead<FMovieSceneBlendChannelID> BlendIDs, TRead<int32> Integers, TRead<int32> BaseValues, TReadOptional<double> EasingAndWeights) const
+	void ForEachAllocation(const FEntityAllocation* Allocation, TRead<FMovieSceneBlendChannelID> BlendIDs, TRead<int64> Integers, TRead<int64> BaseValues, TReadOptional<double> EasingAndWeights) const
 	{
 		static const FMovieSceneBlenderSystemID IntegerBlenderSystemID = UMovieSceneBlenderSystem::GetBlenderSystemID<UMovieScenePiecewiseIntegerBlenderSystem>();
 
@@ -102,7 +102,7 @@ struct FIntegerAdditiveFromBaseBlendTask
 				FIntegerBlendResult& Result = AccumulationBuffer[BlendID.ChannelID];
 
 				const double Weight = EasingAndWeights[Index];
-				Result.Total  += (int32)((Integers[Index] - BaseValues[Index]) * Weight);
+				Result.Total  += (int64)((Integers[Index] - BaseValues[Index]) * Weight);
 				Result.Weight += Weight;
 			}
 		}
@@ -132,7 +132,7 @@ struct FIntegerCombineBlends
 		: AccumulationBuffers(*InAccumulationBuffers)
 	{}
 
-	void ForEachAllocation(FEntityAllocation* Allocation, TRead<FMovieSceneBlendChannelID> BlendIDs, TReadOptional<int32> InitialValues, TWrite<int32> IntegerResults) const
+	void ForEachAllocation(FEntityAllocation* Allocation, TRead<FMovieSceneBlendChannelID> BlendIDs, TReadOptional<int64> InitialValues, TWrite<int64> IntegerResults) const
 	{
 		static const FMovieSceneBlenderSystemID IntegerBlenderSystemID = UMovieSceneBlenderSystem::GetBlenderSystemID<UMovieScenePiecewiseIntegerBlenderSystem>();
 
@@ -154,7 +154,7 @@ struct FIntegerCombineBlends
 		}
 	}
 
-	void BlendResultsWithInitial(uint16 BlendID, const int32 InitialValue, int32& OutFinalBlendResult) const
+	void BlendResultsWithInitial(uint16 BlendID, const int64 InitialValue, int64& OutFinalBlendResult) const
 	{
 		FIntegerBlendResult AbsoluteResult = AccumulationBuffers.Absolute.Num() > 0 ? AccumulationBuffers.Absolute[BlendID] : FIntegerBlendResult();
 		FIntegerBlendResult RelativeResult = AccumulationBuffers.Relative.Num() > 0 ? AccumulationBuffers.Relative[BlendID] : FIntegerBlendResult();
@@ -178,12 +178,12 @@ struct FIntegerCombineBlends
 			// Note that "partial weighting" means strictly between 0 and 100%. At 100% and above, we don't need to do this thing with the initial
 			// value. At 0%, we have no absolute value (only a relative value) and we therefore don't want to include the initial value either.
 			const bool bInitialValueContributes = (0.f < AbsoluteResult.Weight && AbsoluteResult.Weight < 1.f);
-			const int32 AbsoluteBlendedValue = bInitialValueContributes ?
-				((int32)(InitialValue * (1.f - AbsoluteResult.Weight)) + AbsoluteResult.Total) :
+			const int64 AbsoluteBlendedValue = bInitialValueContributes ?
+				((int64)(InitialValue * (1.f - AbsoluteResult.Weight)) + AbsoluteResult.Total) :
 				AbsoluteResult.Total;
 			const double FinalTotalWeight = bInitialValueContributes ? (TotalWeight + (1.f - AbsoluteResult.Weight)) : TotalWeight;
 
-			const int32 Value = (int32)((AbsoluteBlendedValue + RelativeResult.Total) / FinalTotalWeight) + TotalAdditiveResult.Total;
+			const int64 Value = (int64)((AbsoluteBlendedValue + RelativeResult.Total) / FinalTotalWeight) + TotalAdditiveResult.Total;
 			OutFinalBlendResult = Value;
 		}
 		else if (TotalAdditiveResult.Weight != 0)
@@ -196,7 +196,7 @@ struct FIntegerCombineBlends
 		}
 	}
 
-	void BlendResults(uint16 BlendID, int32& OutFinalBlendResult) const
+	void BlendResults(uint16 BlendID, int64& OutFinalBlendResult) const
 	{
 		FIntegerBlendResult AbsoluteResult = AccumulationBuffers.Absolute.Num() > 0 ? AccumulationBuffers.Absolute[BlendID] : FIntegerBlendResult();
 		FIntegerBlendResult AdditiveResult = AccumulationBuffers.Additive.Num() > 0 ? AccumulationBuffers.Additive[BlendID] : FIntegerBlendResult();
@@ -209,7 +209,7 @@ struct FIntegerCombineBlends
 		const double TotalWeight = AbsoluteResult.Weight;
 		if (TotalWeight != 0)
 		{
-			const int32 Value = (int32)(AbsoluteResult.Total / AbsoluteResult.Weight) + AdditiveResult.Total + AdditiveFromBaseResult.Total;
+			const int64 Value = (int64)(AbsoluteResult.Total / AbsoluteResult.Weight) + AdditiveResult.Total + AdditiveFromBaseResult.Total;
 			OutFinalBlendResult = Value;
 		}
 	}
@@ -468,7 +468,7 @@ void UMovieScenePiecewiseIntegerBlenderSystem::ReinitializeAccumulationBuffers()
 	AccumulationBuffers.Reset();
 
 	// Find if we have any integer results to blend using any supported blend types.
-	const TComponentTypeID<int32> Component = BuiltInComponents->IntegerResult;
+	const TComponentTypeID<int64> Component = BuiltInComponents->IntegerResult;
 
 	const bool bHasAbsolutes         = Linker->EntityManager.Contains(FEntityComponentFilter().All({ GetBlenderTypeTag(), Component, BuiltInComponents->BlendChannelInput, BuiltInComponents->Tags.AbsoluteBlend }));
 	const bool bHasRelatives         = Linker->EntityManager.Contains(FEntityComponentFilter().All({ GetBlenderTypeTag(), Component, BuiltInComponents->BlendChannelInput, BuiltInComponents->Tags.RelativeBlend }));
diff --git a/Engine/Source/Runtime/MovieSceneTracks/Private/Systems/MovieScenePropertyInstantiator.cpp b/Engine/Source/Runtime/MovieSceneTracks/Private/Systems/MovieScenePropertyInstantiator.cpp
--- a/Engine/Source/Runtime/MovieSceneTracks/Private/Systems/MovieScenePropertyInstantiator.cpp
+++ b/Engine/Source/Runtime/MovieSceneTracks/Private/Systems/MovieScenePropertyInstantiator.cpp
@@ -1132,8 +1132,26 @@ UE::MovieScene::FPropertyRecomposerPropertyInfo UMovieScenePropertyInstantiatorS
 	using namespace UE::MovieScene;
 
 	TOptionalComponentReader<FMovieScenePropertyBinding> PropertyBinding = Linker->EntityManager.ReadComponent(EntityID, BuiltInComponents->PropertyBinding);
-	TOptionalComponentReader<FEntityGroupID>             GroupID         = Linker->EntityManager.ReadComponent(EntityID, BuiltInComponents->Group);
-	if (!PropertyBinding || !GroupID)
+	if (!PropertyBinding)
+	{
+		return FPropertyRecomposerPropertyInfo::Invalid();
+	}
+
+	TOptionalComponentReader<FEntityGroupID> GroupID = Linker->EntityManager.ReadComponent(EntityID, BuiltInComponents->Group);
+	if (!GroupID)
+	{
+		// If this Entity didn't have a group, see if any of its children are bound to the same object and use that
+		Linker->EntityManager.IterateImmediateChildren(EntityID, [this, Object, &GroupID](FMovieSceneEntityID ChildEntityID){
+			TOptionalComponentReader<UObject*>       BoundObject = this->Linker->EntityManager.ReadComponent(ChildEntityID, BuiltInComponents->BoundObject);
+			TOptionalComponentReader<FEntityGroupID> ChildGroup  = this->Linker->EntityManager.ReadComponent(ChildEntityID, BuiltInComponents->Group);
+			if (BoundObject && *BoundObject == Object && ChildGroup)
+			{
+				GroupID = MoveTemp(ChildGroup);
+			}
+		});
+	}
+
+	if (!GroupID)
 	{
 		return FPropertyRecomposerPropertyInfo::Invalid();
 	}
diff --git a/Engine/Source/Runtime/MovieSceneTracks/Private/Tracks/MovieSceneVectorTrack.cpp b/Engine/Source/Runtime/MovieSceneTracks/Private/Tracks/MovieSceneVectorTrack.cpp
--- a/Engine/Source/Runtime/MovieSceneTracks/Private/Tracks/MovieSceneVectorTrack.cpp
+++ b/Engine/Source/Runtime/MovieSceneTracks/Private/Tracks/MovieSceneVectorTrack.cpp
@@ -26,6 +26,30 @@ UMovieSceneSection* UMovieSceneFloatVectorTrack::CreateNewSection()
 	return NewSection;
 }
 
+void UMovieSceneFloatVectorTrack::InitializeFromProperty(const FProperty* Property, const UE::MovieScene::FPropertyDefinition* Definition)
+{
+	const FStructProperty* StructProp = CastField<const FStructProperty>(Property);
+
+	FName StructName = StructProp ? StructProp->Struct->GetFName() : NAME_None;
+
+	if (StructName == NAME_Vector2f)
+	{
+		SetNumChannelsUsed(2);
+	}
+	else if (StructName == NAME_Vector3f)
+	{
+		SetNumChannelsUsed(3);
+	}
+	else if (StructName == NAME_Vector4f)
+	{
+		SetNumChannelsUsed(4);
+	}
+	else
+	{
+		ensureAlwaysMsgf(false, TEXT("Attempting to create a vector track with an unsupported property type"));
+	}
+}
+
 
 UMovieSceneDoubleVectorTrack::UMovieSceneDoubleVectorTrack( const FObjectInitializer& ObjectInitializer )
 	: Super( ObjectInitializer )
@@ -47,3 +71,26 @@ UMovieSceneSection* UMovieSceneDoubleVectorTrack::CreateNewSection()
 	return NewSection;
 }
 
+void UMovieSceneDoubleVectorTrack::InitializeFromProperty(const FProperty* Property, const UE::MovieScene::FPropertyDefinition* Definition)
+{
+	const FStructProperty* StructProp = CastField<const FStructProperty>(Property);
+
+	FName StructName = StructProp ? StructProp->Struct->GetFName() : NAME_None;
+	
+	if (StructName == NAME_Vector2D)
+	{
+		SetNumChannelsUsed(2);
+	}
+	else if (StructName == NAME_Vector3d || StructName == NAME_Vector)
+	{
+		SetNumChannelsUsed(3);
+	}	
+	else if (StructName == NAME_Vector4d || StructName == NAME_Vector4)
+	{
+		SetNumChannelsUsed(4);
+	}
+	else
+	{
+		ensureAlwaysMsgf(false, TEXT("Attempting to create a vector track with an unsupported property type"));
+	}
+}
\ No newline at end of file
diff --git a/Engine/Source/Runtime/MovieSceneTracks/Public/MovieSceneTracksComponentTypes.h b/Engine/Source/Runtime/MovieSceneTracks/Public/MovieSceneTracksComponentTypes.h
--- a/Engine/Source/Runtime/MovieSceneTracks/Public/MovieSceneTracksComponentTypes.h
+++ b/Engine/Source/Runtime/MovieSceneTracks/Public/MovieSceneTracksComponentTypes.h
@@ -7,7 +7,7 @@
 #include "EntitySystem/BuiltInComponentTypes.h"
 #include "EntitySystem/MovieSceneEntityIDs.h"
 #include "EntitySystem/MovieScenePropertySystemTypes.h"
-#include "EntitySystem/MovieScenePropertyMetaDataTraits.h"
+#include "EntitySystem/MovieScenePropertyMetaData.h"
 #include "EntitySystem/MovieScenePropertyTraits.h"
 #include "Engine/EngineTypes.h"
 #include "EulerTransform.h"
@@ -18,7 +18,11 @@
 #include "Materials/MaterialParameters.h"
 #include "Sections/MovieSceneCameraShakeSection.h"
 #include "Misc/Guid.h"
+#include "MovieSceneTracksPropertyTraits.h"
 #include "Tracks/MovieSceneMaterialTrack.h"
+#include "Channels/MovieSceneFloatChannel.h"
+#include "Channels/MovieSceneDoubleChannel.h"
+#include "Channels/MovieSceneObjectPathChannel.h"
 #include "MovieSceneTracksComponentTypes.generated.h"
 
 class UMaterialParameterCollection;
@@ -31,6 +35,19 @@ struct FCameraShakeBaseStartParams;
 struct FCameraShakeSourceComponentStartParams;
 struct FMovieSceneObjectBindingID;
 
+namespace UE::MovieScene
+{
+	struct FIntermediate3DTransform;
+
+	template<>
+	struct TPropertyMatch<FIntermediate3DTransform>
+	{
+		static bool SupportsProperty(const FProperty& InProperty)
+		{
+			return false;
+		}
+	};
+}
 
 /** Component data for Perlin Noise channels */
 USTRUCT()
@@ -233,17 +250,29 @@ struct FFadeComponentData
 	bool bFadeAudio;
 };
 
-struct FFloatPropertyTraits
+struct FFloatPropertyTraits : IPropertyTraits
 {
 	static constexpr bool bIsComposite = false;
 
 	using StorageType  = double;
 	using CustomAccessorStorageType = float;
-	using MetaDataType = TPropertyMetaData<>;
-	using TraitsType = FFloatPropertyTraits;
 
 	using FloatTraitsImpl = TDirectPropertyTraits<float>;
 
+	FORCEINLINE static bool SupportsProperty(const FProperty& InProperty)
+	{
+		return InProperty.IsA<FFloatProperty>();
+	}
+	FORCEINLINE static FIntermediatePropertyValue CoercePropertyValue(const FProperty& InProperty, const FSourcePropertyValue& InPropertyValue)
+	{
+		check(InProperty.IsA<FFloatProperty>());
+		return FIntermediatePropertyValue::FromValue(*InPropertyValue.Cast<float>());
+	}
+	FORCEINLINE static void UnpackChannels(float InValue, const FProperty& Property, FUnpackedChannelValues& OutUnpackedValues)
+	{
+		OutUnpackedValues.Add(FUnpackedChannelValue(TIndexedChannelValue<FMovieSceneFloatChannel>(InValue, 0), NAME_None));
+	}
+
 	static void GetObjectPropertyValue(const UObject* InObject, const FCustomPropertyAccessor& BaseCustomAccessor, double& OutValue)
 	{
 		const TCustomPropertyAccessor<FFloatPropertyTraits>& CustomAccessor = static_cast<const TCustomPropertyAccessor<FFloatPropertyTraits>&>(BaseCustomAccessor);
@@ -287,237 +316,36 @@ struct FFloatPropertyTraits
 	}
 };
 
-struct FColorPropertyTraits
-{
-	using StorageType  = FIntermediateColor;
-	using MetaDataType = TPropertyMetaData<EColorPropertyType>;
-
-	static void GetObjectPropertyValue(const UObject* InObject, EColorPropertyType ColorType, const FCustomPropertyAccessor& BaseCustomAccessor, FIntermediateColor& OutValue)
-	{
-		const TCustomPropertyAccessor<FColorPropertyTraits>& CustomAccessor = static_cast<const TCustomPropertyAccessor<FColorPropertyTraits>&>(BaseCustomAccessor);
-		OutValue = (*CustomAccessor.Functions.Getter)(InObject, ColorType);
-	}
-	static void GetObjectPropertyValue(const UObject* InObject, EColorPropertyType ColorType, uint16 PropertyOffset, FIntermediateColor& OutValue)
-	{
-		switch (ColorType)
-		{
-			case EColorPropertyType::Slate:  TIndirectPropertyTraits<FSlateColor, FIntermediateColor>::GetObjectPropertyValue(InObject, PropertyOffset, OutValue);  return;
-			case EColorPropertyType::Linear: TIndirectPropertyTraits<FLinearColor, FIntermediateColor>::GetObjectPropertyValue(InObject, PropertyOffset, OutValue); return;
-			case EColorPropertyType::Color:  TIndirectPropertyTraits<FColor, FIntermediateColor>::GetObjectPropertyValue(InObject, PropertyOffset, OutValue);       return;
-		}
-	}
-	static void GetObjectPropertyValue(const UObject* InObject, EColorPropertyType ColorType, FTrackInstancePropertyBindings* PropertyBindings, FIntermediateColor& OutValue)
-	{
-		switch (ColorType)
-		{
-			case EColorPropertyType::Slate:  TIndirectPropertyTraits<FSlateColor, FIntermediateColor>::GetObjectPropertyValue(InObject, PropertyBindings, OutValue);  return;
-			case EColorPropertyType::Linear: TIndirectPropertyTraits<FLinearColor, FIntermediateColor>::GetObjectPropertyValue(InObject, PropertyBindings, OutValue); return;
-			case EColorPropertyType::Color:  TIndirectPropertyTraits<FColor, FIntermediateColor>::GetObjectPropertyValue(InObject, PropertyBindings, OutValue);       return;
-		}
-	}
-	static void GetObjectPropertyValue(const UObject* InObject, EColorPropertyType ColorType, const FName& PropertyPath, StorageType& OutValue)
-	{
-		switch (ColorType)
-		{
-			case EColorPropertyType::Slate:  TIndirectPropertyTraits<FSlateColor, FIntermediateColor>::GetObjectPropertyValue(InObject, PropertyPath, OutValue);  return;
-			case EColorPropertyType::Linear: TIndirectPropertyTraits<FLinearColor, FIntermediateColor>::GetObjectPropertyValue(InObject, PropertyPath, OutValue); return;
-			case EColorPropertyType::Color:  TIndirectPropertyTraits<FColor, FIntermediateColor>::GetObjectPropertyValue(InObject, PropertyPath, OutValue);       return;
-		}
-	}
-
-	static void SetObjectPropertyValue(UObject* InObject, EColorPropertyType ColorType, const FCustomPropertyAccessor& BaseCustomAccessor, const FIntermediateColor& InValue)
-	{
-		const TCustomPropertyAccessor<FColorPropertyTraits>& CustomAccessor = static_cast<const TCustomPropertyAccessor<FColorPropertyTraits>&>(BaseCustomAccessor);
-		(*CustomAccessor.Functions.Setter)(InObject, ColorType, InValue);
-	}
-	static void SetObjectPropertyValue(UObject* InObject, EColorPropertyType ColorType, uint16 PropertyOffset, const FIntermediateColor& InValue)
-	{
-		switch (ColorType)
-		{
-			case EColorPropertyType::Slate:  TIndirectPropertyTraits<FSlateColor, FIntermediateColor>::SetObjectPropertyValue(InObject, PropertyOffset, InValue);  return;
-			case EColorPropertyType::Linear: TIndirectPropertyTraits<FLinearColor, FIntermediateColor>::SetObjectPropertyValue(InObject, PropertyOffset, InValue); return;
-			case EColorPropertyType::Color:  TIndirectPropertyTraits<FColor, FIntermediateColor>::SetObjectPropertyValue(InObject, PropertyOffset, InValue);       return;
-		}
-	}
-	static void SetObjectPropertyValue(UObject* InObject, EColorPropertyType ColorType, FTrackInstancePropertyBindings* PropertyBindings, const FIntermediateColor& InValue)
-	{
-		switch (ColorType)
-		{
-			case EColorPropertyType::Slate:  TIndirectPropertyTraits<FSlateColor, FIntermediateColor>::SetObjectPropertyValue(InObject, PropertyBindings, InValue);  return;
-			case EColorPropertyType::Linear: TIndirectPropertyTraits<FLinearColor, FIntermediateColor>::SetObjectPropertyValue(InObject, PropertyBindings, InValue); return;
-			case EColorPropertyType::Color:  TIndirectPropertyTraits<FColor, FIntermediateColor>::SetObjectPropertyValue(InObject, PropertyBindings, InValue);       return;
-		}
-	}
-
-	static FIntermediateColor CombineComposites(EColorPropertyType InType, double InR, double InG, double InB, double InA)
-	{
-		return FIntermediateColor((float)InR, (float)InG, (float)InB, (float)InA);
-	}
-};
-
-struct FDoubleVectorPropertyTraits
+struct FBoolPropertyTraits : IPropertyTraits
 {
-	using StorageType  = FDoubleIntermediateVector;
-	using MetaDataType = TPropertyMetaData<FVectorPropertyMetaData>;
-
-	static void GetObjectPropertyValue(const UObject* InObject, FVectorPropertyMetaData MetaData, const FCustomPropertyAccessor& BaseCustomAccessor, FDoubleIntermediateVector& OutValue)
-	{
-		const TCustomPropertyAccessor<FDoubleVectorPropertyTraits>& CustomAccessor = static_cast<const TCustomPropertyAccessor<FDoubleVectorPropertyTraits>&>(BaseCustomAccessor);
-		OutValue = (*CustomAccessor.Functions.Getter)(InObject, MetaData);
-	}
-	static void GetObjectPropertyValue(const UObject* InObject, FVectorPropertyMetaData MetaData, uint16 PropertyOffset, FDoubleIntermediateVector& OutValue)
-	{
-		switch (MetaData.NumChannels)
-		{
-			case 2: TIndirectPropertyTraits<FVector2d, FDoubleIntermediateVector>::GetObjectPropertyValue(InObject, PropertyOffset, OutValue); return;
-			case 3: TIndirectPropertyTraits<FVector3d, FDoubleIntermediateVector>::GetObjectPropertyValue(InObject, PropertyOffset, OutValue); return;
-			case 4: TIndirectPropertyTraits<FVector4d, FDoubleIntermediateVector>::GetObjectPropertyValue(InObject, PropertyOffset, OutValue); return;
-		}
-	}
-	static void GetObjectPropertyValue(const UObject* InObject, FVectorPropertyMetaData MetaData, FTrackInstancePropertyBindings* PropertyBindings, FDoubleIntermediateVector& OutValue)
-	{
-		switch (MetaData.NumChannels)
-		{
-			case 2: TIndirectPropertyTraits<FVector2d, FDoubleIntermediateVector>::GetObjectPropertyValue(InObject, PropertyBindings, OutValue); return;
-			case 3: TIndirectPropertyTraits<FVector3d, FDoubleIntermediateVector>::GetObjectPropertyValue(InObject, PropertyBindings, OutValue); return;
-			case 4: TIndirectPropertyTraits<FVector4d, FDoubleIntermediateVector>::GetObjectPropertyValue(InObject, PropertyBindings, OutValue); return;
-		}
-	}
-	static void GetObjectPropertyValue(const UObject* InObject, FVectorPropertyMetaData MetaData, const FName& PropertyPath, StorageType& OutValue)
-	{
-		switch (MetaData.NumChannels)
-		{
-			case 2: TIndirectPropertyTraits<FVector2d, FDoubleIntermediateVector>::GetObjectPropertyValue(InObject, PropertyPath, OutValue); return;
-			case 3: TIndirectPropertyTraits<FVector3d, FDoubleIntermediateVector>::GetObjectPropertyValue(InObject, PropertyPath, OutValue); return;
-			case 4: TIndirectPropertyTraits<FVector4d, FDoubleIntermediateVector>::GetObjectPropertyValue(InObject, PropertyPath, OutValue); return;
-		}
-	}
-
-	static void SetObjectPropertyValue(UObject* InObject, FVectorPropertyMetaData MetaData, const FCustomPropertyAccessor& BaseCustomAccessor, const FDoubleIntermediateVector& InValue)
-	{
-		const TCustomPropertyAccessor<FDoubleVectorPropertyTraits>& CustomAccessor = static_cast<const TCustomPropertyAccessor<FDoubleVectorPropertyTraits>&>(BaseCustomAccessor);
-		(*CustomAccessor.Functions.Setter)(InObject, MetaData, InValue);
-	}
-	static void SetObjectPropertyValue(UObject* InObject, FVectorPropertyMetaData MetaData, uint16 PropertyOffset, const FDoubleIntermediateVector& InValue)
-	{
-		switch (MetaData.NumChannels)
-		{
-			case 2: TIndirectPropertyTraits<FVector2d, FDoubleIntermediateVector>::SetObjectPropertyValue(InObject, PropertyOffset, InValue); return;
-			case 3: TIndirectPropertyTraits<FVector3d, FDoubleIntermediateVector>::SetObjectPropertyValue(InObject, PropertyOffset, InValue); return;
-			case 4: TIndirectPropertyTraits<FVector4d, FDoubleIntermediateVector>::SetObjectPropertyValue(InObject, PropertyOffset, InValue); return;
-		}
-
-		checkf(false, TEXT("Invalid number of channels"));
-	}
-	static void SetObjectPropertyValue(UObject* InObject, FVectorPropertyMetaData MetaData, FTrackInstancePropertyBindings* PropertyBindings, const FDoubleIntermediateVector& InValue)
-	{
-		switch (MetaData.NumChannels)
-		{
-			case 2: TIndirectPropertyTraits<FVector2d, FDoubleIntermediateVector>::SetObjectPropertyValue(InObject, PropertyBindings, InValue); return;
-			case 3: TIndirectPropertyTraits<FVector3d, FDoubleIntermediateVector>::SetObjectPropertyValue(InObject, PropertyBindings, InValue); return;
-			case 4: TIndirectPropertyTraits<FVector4d, FDoubleIntermediateVector>::SetObjectPropertyValue(InObject, PropertyBindings, InValue); return;
-		}
-
-		checkf(false, TEXT("Invalid number of channels"));
-	}
-
-	static FDoubleIntermediateVector CombineComposites(FVectorPropertyMetaData MetaData, double InX, double InY, double InZ, double InW)
-	{
-		return FDoubleIntermediateVector(InX, InY, InZ, InW);
-	}
-};
-
-struct FFloatVectorPropertyTraits
-{
-	using StorageType = FFloatIntermediateVector;
-	using MetaDataType = TPropertyMetaData<FVectorPropertyMetaData>;
+	static constexpr bool bIsComposite = false;
 
-	static void GetObjectPropertyValue(const UObject* InObject, FVectorPropertyMetaData MetaData, const FCustomPropertyAccessor& BaseCustomAccessor, FFloatIntermediateVector& OutValue)
-	{
-		const TCustomPropertyAccessor<FFloatVectorPropertyTraits>& CustomAccessor = static_cast<const TCustomPropertyAccessor<FFloatVectorPropertyTraits>&>(BaseCustomAccessor);
-		OutValue = (*CustomAccessor.Functions.Getter)(InObject, MetaData);
-	}
-	static void GetObjectPropertyValue(const UObject* InObject, FVectorPropertyMetaData MetaData, uint16 PropertyOffset, FFloatIntermediateVector& OutValue)
-	{
-		switch (MetaData.NumChannels)
-		{
-		case 2: TIndirectPropertyTraits<FVector2f, FFloatIntermediateVector>::GetObjectPropertyValue(InObject, PropertyOffset, OutValue); return;
-		case 3: TIndirectPropertyTraits<FVector3f, FFloatIntermediateVector>::GetObjectPropertyValue(InObject, PropertyOffset, OutValue); return;
-		case 4: TIndirectPropertyTraits<FVector4f, FFloatIntermediateVector>::GetObjectPropertyValue(InObject, PropertyOffset, OutValue); return;
-		}
-	}
-	static void GetObjectPropertyValue(const UObject* InObject, FVectorPropertyMetaData MetaData, FTrackInstancePropertyBindings* PropertyBindings, FFloatIntermediateVector& OutValue)
-	{
-		switch (MetaData.NumChannels)
-		{
-		case 2: TIndirectPropertyTraits<FVector2f, FFloatIntermediateVector>::GetObjectPropertyValue(InObject, PropertyBindings, OutValue); return;
-		case 3: TIndirectPropertyTraits<FVector3f, FFloatIntermediateVector>::GetObjectPropertyValue(InObject, PropertyBindings, OutValue); return;
-		case 4: TIndirectPropertyTraits<FVector4f, FFloatIntermediateVector>::GetObjectPropertyValue(InObject, PropertyBindings, OutValue); return;
-		}
-	}
-	static void GetObjectPropertyValue(const UObject* InObject, FVectorPropertyMetaData MetaData, const FName& PropertyPath, StorageType& OutValue)
+	struct FBoolMetaData
 	{
-		switch (MetaData.NumChannels)
-		{
-		case 2: TIndirectPropertyTraits<FVector2f, FFloatIntermediateVector>::GetObjectPropertyValue(InObject, PropertyPath, OutValue); return;
-		case 3: TIndirectPropertyTraits<FVector3f, FFloatIntermediateVector>::GetObjectPropertyValue(InObject, PropertyPath, OutValue); return;
-		case 4: TIndirectPropertyTraits<FVector4f, FFloatIntermediateVector>::GetObjectPropertyValue(InObject, PropertyPath, OutValue); return;
-		}
-	}
+		/** Size of the bitfield/bool property in bytes. Must be one of 1(uint8), 2(uint16), 4(uint32), 8(uint64), or 0 signifying that this is not a bitfield, but a regular bool.. */
+		uint8 BitFieldSize = 0;
+		uint8 BitIndex = 0;
+	};
+	using StorageType         = bool;
+	using MetaDataType        = TPropertyMetaData<FBoolMetaData>;
 
-	static void SetObjectPropertyValue(UObject* InObject, FVectorPropertyMetaData MetaData, const FCustomPropertyAccessor& BaseCustomAccessor, const FFloatIntermediateVector& InValue)
-	{
-		const TCustomPropertyAccessor<FFloatVectorPropertyTraits>& CustomAccessor = static_cast<const TCustomPropertyAccessor<FFloatVectorPropertyTraits>&>(BaseCustomAccessor);
-		(*CustomAccessor.Functions.Setter)(InObject, MetaData, InValue);
-	}
-	static void SetObjectPropertyValue(UObject* InObject, FVectorPropertyMetaData MetaData, uint16 PropertyOffset, const FFloatIntermediateVector& InValue)
+	FORCEINLINE static bool SupportsProperty(const FProperty& InProperty)
 	{
-		switch (MetaData.NumChannels)
-		{
-		case 2: TIndirectPropertyTraits<FVector2f, FFloatIntermediateVector>::SetObjectPropertyValue(InObject, PropertyOffset, InValue); return;
-		case 3: TIndirectPropertyTraits<FVector3f, FFloatIntermediateVector>::SetObjectPropertyValue(InObject, PropertyOffset, InValue); return;
-		case 4: TIndirectPropertyTraits<FVector4f, FFloatIntermediateVector>::SetObjectPropertyValue(InObject, PropertyOffset, InValue); return;
-		}
-
-		checkf(false, TEXT("Invalid number of channels"));
+		return InProperty.IsA<FBoolProperty>();
 	}
-	static void SetObjectPropertyValue(UObject* InObject, FVectorPropertyMetaData MetaData, FTrackInstancePropertyBindings* PropertyBindings, const FFloatIntermediateVector& InValue)
+	FORCEINLINE static FIntermediatePropertyValue CoercePropertyValue(const FProperty& InProperty, const FSourcePropertyValue& InPropertyValue)
 	{
-		switch (MetaData.NumChannels)
-		{
-		case 2: TIndirectPropertyTraits<FVector2f, FFloatIntermediateVector>::SetObjectPropertyValue(InObject, PropertyBindings, InValue); return;
-		case 3: TIndirectPropertyTraits<FVector3f, FFloatIntermediateVector>::SetObjectPropertyValue(InObject, PropertyBindings, InValue); return;
-		case 4: TIndirectPropertyTraits<FVector4f, FFloatIntermediateVector>::SetObjectPropertyValue(InObject, PropertyBindings, InValue); return;
-		}
-
-		checkf(false, TEXT("Invalid number of channels"));
+		return FIntermediatePropertyValue::FromValue(*InPropertyValue.Cast<bool>());
 	}
-
-	static FFloatIntermediateVector CombineComposites(FVectorPropertyMetaData MetaData, float InX, float InY, float InZ, float InW)
+	FORCEINLINE static void UnpackChannels(bool bInValue, const FProperty& Property, FUnpackedChannelValues& OutUnpackedValues)
 	{
-		return FFloatIntermediateVector(InX, InY, InZ, InW);
+		OutUnpackedValues.Add(FUnpackedChannelValue(TIndexedChannelValue<FMovieSceneBoolChannel>(bInValue, 0), NAME_None));
 	}
-};
-
-struct FBoolPropertyTraits
-{
-	static constexpr bool bIsComposite = false;
-
-	struct FBoolMetaData
-	{
-		/** Size of the bitfield/bool property in bytes. Must be one of 1(uint8), 2(uint16), 4(uint32), 8(uint64), or 0 signifying that this is not a bitfield, but a regular bool.. */
-		uint8 BitFieldSize = 0;
-		uint8 BitIndex = 0;
-	};
-	using StorageType  = bool;
-	using MetaDataType = TPropertyMetaData<FBoolMetaData>;
-	using TraitsType   = TDirectPropertyTraits<bool>;
-	using ParamType    = bool;
 
 	/** Property Value Getters  */
 	static void GetObjectPropertyValue(const UObject* InObject, FBoolMetaData MetaData, const FCustomPropertyAccessor& BaseCustomAccessor, bool& OutValue)
 	{
-		const TCustomPropertyAccessor<TraitsType>& CustomAccessor = static_cast<const TCustomPropertyAccessor<TraitsType>&>(BaseCustomAccessor);
+		const TCustomPropertyAccessor<TDirectPropertyTraits<bool>>& CustomAccessor = static_cast<const TCustomPropertyAccessor<TDirectPropertyTraits<bool>>&>(BaseCustomAccessor);
 		OutValue = (*CustomAccessor.Functions.Getter)(InObject);
 	}
 	static void GetObjectPropertyValue(const UObject* InObject, FBoolMetaData MetaData, uint16 PropertyOffset, bool& OutValue)
@@ -548,7 +376,7 @@ struct FBoolPropertyTraits
 	/** Property Value Setters  */
 	static void SetObjectPropertyValue(UObject* InObject, FBoolMetaData MetaData, const FCustomPropertyAccessor& BaseCustomAccessor, bool InValue)
 	{
-		const TCustomPropertyAccessor<TraitsType>& CustomAccessor = static_cast<const TCustomPropertyAccessor<TraitsType>&>(BaseCustomAccessor);
+		const TCustomPropertyAccessor<TDirectPropertyTraits<bool>>& CustomAccessor = static_cast<const TCustomPropertyAccessor<TDirectPropertyTraits<bool>>&>(BaseCustomAccessor);
 		(*CustomAccessor.Functions.Setter)(InObject, InValue);
 	}
 	static void SetObjectPropertyValue(UObject* InObject, FBoolMetaData MetaData, uint16 PropertyOffset, bool InValue)
@@ -574,7 +402,7 @@ struct FBoolPropertyTraits
 	}
 };
 
-struct FObjectPropertyTraits
+struct FObjectPropertyTraits : IPropertyTraits
 {
 	static constexpr bool bIsComposite = false;
 
@@ -585,12 +413,23 @@ struct FObjectPropertyTraits
 	};
 
 	using StorageType = FObjectComponent;
-	//using CustomAccessorStorageType = float;
 	using MetaDataType = TPropertyMetaData<FObjectMetadata>;
-	using TraitsType = FObjectPropertyTraits;
-	using ParamType = UObject*;
 
-	using ObjectTraitsImpl = TIndirectPropertyTraits<UObject*, FObjectComponent>;
+	using ObjectTraitsImpl = TPropertyTraits<UObject*, FObjectComponent>;
+
+	FORCEINLINE static bool SupportsProperty(const FProperty& InProperty)
+	{
+		return InProperty.IsA<FObjectPropertyBase>();
+	}
+	FORCEINLINE static FIntermediatePropertyValue CoercePropertyValue(const FProperty& InProperty, const FSourcePropertyValue& InPropertyValue)
+	{
+		return FIntermediatePropertyValue::FromValue(FObjectComponent::Weak(*InPropertyValue.Cast<UObject*>()));
+	}
+	FORCEINLINE static void UnpackChannels(const FObjectComponent& InValue, const FProperty& Property, FUnpackedChannelValues& OutUnpackedValues)
+	{
+		FMovieSceneObjectPathChannelKeyValue ChannelValue = InValue.GetObject();
+		OutUnpackedValues.Add(FUnpackedChannelValue(TIndexedChannelValue<FMovieSceneObjectPathChannel>(ChannelValue, 0), NAME_None));
+	}
 
 	static void GetObjectPropertyValue(const UObject* InObject, FObjectMetadata ObjectMetadata, const FCustomPropertyAccessor& BaseCustomAccessor, FObjectComponent& OutValue)
 	{
@@ -648,25 +487,29 @@ struct FObjectPropertyTraits
 	}
 };
 
-using FBytePropertyTraits               = TDirectPropertyTraits<uint8, false>;
-using FEnumPropertyTraits               = TDirectPropertyTraits<uint8, false>;
-using FIntPropertyTraits                = TDirectPropertyTraits<int32, false>;
-using FDoublePropertyTraits             = TDirectPropertyTraits<double, false>;
-using FTransformPropertyTraits          = TIndirectPropertyTraits<FTransform, FIntermediate3DTransform>;
-using FEulerTransformPropertyTraits     = TIndirectPropertyTraits<FEulerTransform, FIntermediate3DTransform>;
+using FBytePropertyTraits               = TPropertyTraits<uint8, uint8>;
+using FEnumPropertyTraits               = TPropertyTraits<uint8, uint8>;
+using FIntPropertyTraits                = TDynamicVariantPropertyTraits<int64, int64, int32, int16, int8>;
+using FDoublePropertyTraits             = TDynamicVariantPropertyTraits<double, double>;
+using FTransformPropertyTraits          = TDynamicVariantPropertyTraits<FIntermediate3DTransform, FTransform>;
+using FEulerTransformPropertyTraits     = TPropertyTraits<FEulerTransform, FIntermediate3DTransform>;
 using FComponentTransformPropertyTraits = TDirectPropertyTraits<FIntermediate3DTransform>;
 using FRotatorPropertyTraits            = TDirectPropertyTraits<FRotator>;
 using FStringPropertyTraits             = TDirectPropertyTraits<FString>;
 using FTextPropertyTraits               = TDirectPropertyTraits<FText>;
 
-using FBoolParameterTraits              = TDirectPropertyTraits<bool, false>;
-using FByteParameterTraits              = TDirectPropertyTraits<uint8, false>;
-using FIntegerParameterTraits           = TDirectPropertyTraits<int32, false>;
-using FFloatParameterTraits             = TIndirectPropertyTraits<float, double, false>;
-using FColorParameterTraits             = TIndirectPropertyTraits<FLinearColor, FIntermediateColor>;
-using FVector2ParameterTraits           = TIndirectPropertyTraits<FVector2f, FFloatIntermediateVector>;
-using FVector3ParameterTraits           = TIndirectPropertyTraits<FVector3f, FFloatIntermediateVector>;
-using FTransformParameterTraits         = TIndirectPropertyTraits<FEulerTransform, FIntermediate3DTransform>;
+using FBoolParameterTraits              = TDirectPropertyTraits<bool>;
+using FByteParameterTraits              = TDirectPropertyTraits<uint8>;
+using FIntegerParameterTraits           = TDirectPropertyTraits<int32>;
+using FFloatParameterTraits             = TPropertyTraits<float, double>;
+using FColorParameterTraits             = TPropertyTraits<FLinearColor, FIntermediateColor>;
+using FVector2ParameterTraits           = TPropertyTraits<FVector2f, FFloatIntermediateVector>;
+using FVector3ParameterTraits           = TPropertyTraits<FVector3f, FFloatIntermediateVector>;
+using FTransformParameterTraits         = TPropertyTraits<FEulerTransform, FIntermediate3DTransform>;
+
+using FColorPropertyTraits              = TDynamicVariantPropertyTraits<FIntermediateColor, FLinearColor, FColor, FSlateColor>;
+using FDoubleVectorPropertyTraits       = TDynamicVariantPropertyTraits<FDoubleIntermediateVector, FVector2D, FVector, FVector3d, FVector4d, FVector4>;
+using FFloatVectorPropertyTraits        = TDynamicVariantPropertyTraits<FFloatIntermediateVector, FVector2f, FVector3f, FVector4f>;
 
 struct FMovieSceneTracksComponentTypes
 {
@@ -778,6 +621,7 @@ struct FMovieSceneTracksComponentTypes
 		TCustomPropertyRegistration<FColorPropertyTraits> Color;
 		TCustomPropertyRegistration<FFloatVectorPropertyTraits> FloatVector;
 		TCustomPropertyRegistration<FDoubleVectorPropertyTraits> DoubleVector;
+		TCustomPropertyRegistration<FTransformPropertyTraits> Transform;
 		TCustomPropertyRegistration<FComponentTransformPropertyTraits, 1> ComponentTransform;
 		FObjectPropertyRegistration Object;
 
diff --git a/Engine/Source/Runtime/MovieSceneTracks/Public/MovieSceneTracksPropertyTypes.h b/Engine/Source/Runtime/MovieSceneTracks/Public/MovieSceneTracksPropertyTypes.h
--- a/Engine/Source/Runtime/MovieSceneTracks/Public/MovieSceneTracksPropertyTypes.h
+++ b/Engine/Source/Runtime/MovieSceneTracks/Public/MovieSceneTracksPropertyTypes.h
@@ -20,6 +20,7 @@ namespace MovieScene
 {
 
 struct FObjectComponent;
+struct FGeneratedTrackKeys;
 
 struct FVectorPropertyMetaData
 {
@@ -74,6 +75,7 @@ struct FFloatIntermediateVector
 		return FVector4d(X, Y, Z, W);
 	}
 };
+
 struct FDoubleIntermediateVector
 {
 	double X, Y, Z, W;
@@ -95,6 +97,7 @@ struct FDoubleIntermediateVector
 	{}
 };
 
+
 /** Color type for the color property system */
 enum class EColorPropertyType : uint8
 {
@@ -165,6 +168,8 @@ struct FIntermediateColor
 	}
 };
 
+
+
 /** Intermediate type used for applying partially animated transforms. Saves us from repteatedly recomposing quaternions from euler angles */
 struct FIntermediate3DTransform
 {
@@ -222,6 +227,10 @@ MOVIESCENETRACKS_API void SetComponentTransformAndVelocity(UObject* Object, cons
 MOVIESCENETRACKS_API void ConvertOperationalProperty(float In, double& Out);
 MOVIESCENETRACKS_API void ConvertOperationalProperty(double In, float& Out);
 
+MOVIESCENETRACKS_API void ConvertOperationalProperty(int64 In, int8&  Out);
+MOVIESCENETRACKS_API void ConvertOperationalProperty(int64 In, int16& Out);
+MOVIESCENETRACKS_API void ConvertOperationalProperty(int64 In, int32& Out);
+
 MOVIESCENETRACKS_API void ConvertOperationalProperty(const FIntermediate3DTransform& In, FEulerTransform& Out);
 MOVIESCENETRACKS_API void ConvertOperationalProperty(const FEulerTransform& In, FIntermediate3DTransform& Out);
 MOVIESCENETRACKS_API void ConvertOperationalProperty(const FIntermediate3DTransform& In, FTransform& Out);
diff --git a/Engine/Source/Runtime/MovieSceneTracks/Public/Sections/MovieSceneActorReferenceSection.h b/Engine/Source/Runtime/MovieSceneTracks/Public/Sections/MovieSceneActorReferenceSection.h
--- a/Engine/Source/Runtime/MovieSceneTracks/Public/Sections/MovieSceneActorReferenceSection.h
+++ b/Engine/Source/Runtime/MovieSceneTracks/Public/Sections/MovieSceneActorReferenceSection.h
@@ -219,3 +219,9 @@ inline bool EvaluateChannel(const FMovieSceneActorReferenceData* InChannel, FFra
 {
 	return InChannel->Evaluate(InTime, OutValue);
 }
+
+inline bool GetChannelDefault(const FMovieSceneActorReferenceData* Channel, FMovieSceneActorReferenceKey& OutDefaultValue)
+{
+	OutDefaultValue = Channel->GetDefault();
+	return true;
+}
\ No newline at end of file
diff --git a/Engine/Source/Runtime/MovieSceneTracks/Public/Systems/MovieScenePiecewiseIntegerBlenderSystem.h b/Engine/Source/Runtime/MovieSceneTracks/Public/Systems/MovieScenePiecewiseIntegerBlenderSystem.h
--- a/Engine/Source/Runtime/MovieSceneTracks/Public/Systems/MovieScenePiecewiseIntegerBlenderSystem.h
+++ b/Engine/Source/Runtime/MovieSceneTracks/Public/Systems/MovieScenePiecewiseIntegerBlenderSystem.h
@@ -17,7 +17,7 @@ namespace MovieScene
 	struct FIntegerBlendResult
 	{
 		/** Cumulative sum of blend values pre-multiplied with each value's weight. */
-		int32 Total = 0;
+		int64 Total = 0;
 		/** Cumulative sum of weights. */
 		double Weight = 0.f;
 	};
diff --git a/Engine/Source/Runtime/MovieSceneTracks/Public/Systems/MovieScenePropertyInstantiator.h b/Engine/Source/Runtime/MovieSceneTracks/Public/Systems/MovieScenePropertyInstantiator.h
--- a/Engine/Source/Runtime/MovieSceneTracks/Public/Systems/MovieScenePropertyInstantiator.h
+++ b/Engine/Source/Runtime/MovieSceneTracks/Public/Systems/MovieScenePropertyInstantiator.h
@@ -94,6 +94,8 @@ public:
 	template<typename PropertyTraits>
 	UE::MovieScene::TRecompositionResult<typename PropertyTraits::StorageType> RecomposeBlendOperational(const UE::MovieScene::TPropertyComponents<PropertyTraits>& InComponents, const UE::MovieScene::FDecompositionQuery& InQuery, const typename PropertyTraits::StorageType& InCurrentValue);
 
+	MOVIESCENETRACKS_API UE::MovieScene::FPropertyRecomposerPropertyInfo FindPropertyFromSource(FMovieSceneEntityID EntityID, UObject* Object) const;
+
 private:
 
 	MOVIESCENETRACKS_API virtual void OnRun(FSystemTaskPrerequisites& InPrerequisites, FSystemSubsequentTasks& Subsequents) override;
@@ -217,8 +219,6 @@ private:
 
 	bool ResolveProperty(UE::MovieScene::FCustomAccessorView CustomAccessors, UObject* Object, const FMovieScenePropertyBinding& PropertyBinding, const UE::MovieScene::FEntityGroupID& GroupID, int32 PropertyDefinitionIndex);
 
-	UE::MovieScene::FPropertyRecomposerPropertyInfo FindPropertyFromSource(FMovieSceneEntityID EntityID, UObject* Object) const;
-
 	void InitializePropertyMetaData(FSystemTaskPrerequisites& InPrerequisites, FSystemSubsequentTasks& Subsequents);
 
 private:
diff --git a/Engine/Source/Runtime/MovieSceneTracks/Public/Tracks/MovieScenePropertyTrack.h b/Engine/Source/Runtime/MovieSceneTracks/Public/Tracks/MovieScenePropertyTrack.h
--- a/Engine/Source/Runtime/MovieSceneTracks/Public/Tracks/MovieScenePropertyTrack.h
+++ b/Engine/Source/Runtime/MovieSceneTracks/Public/Tracks/MovieScenePropertyTrack.h
@@ -16,6 +16,11 @@
 #include "UObject/ObjectMacros.h"
 #include "MovieScenePropertyTrack.generated.h"
 
+namespace UE::MovieScene
+{
+	struct FPropertyDefinition;
+}
+
 /**
  * Base class for tracks that animate an object property
  */
@@ -36,6 +41,9 @@ public:
 	MOVIESCENETRACKS_API virtual void RemoveSectionAt(int32 SectionIndex) override;
 	MOVIESCENETRACKS_API virtual bool IsEmpty() const override;
 	MOVIESCENETRACKS_API virtual const TArray<UMovieSceneSection*>& GetAllSections() const override;
+	virtual void InitializeFromProperty(const FProperty* Property, const UE::MovieScene::FPropertyDefinition* Definition)
+	{
+	}
 
 #if WITH_EDITORONLY_DATA
 	MOVIESCENETRACKS_API virtual FText GetDefaultDisplayName() const override;
diff --git a/Engine/Source/Runtime/MovieSceneTracks/Public/Tracks/MovieSceneVectorTrack.h b/Engine/Source/Runtime/MovieSceneTracks/Public/Tracks/MovieSceneVectorTrack.h
--- a/Engine/Source/Runtime/MovieSceneTracks/Public/Tracks/MovieSceneVectorTrack.h
+++ b/Engine/Source/Runtime/MovieSceneTracks/Public/Tracks/MovieSceneVectorTrack.h
@@ -21,6 +21,8 @@ public:
 	virtual bool SupportsType(TSubclassOf<UMovieSceneSection> SectionClass) const override;
 	virtual UMovieSceneSection* CreateNewSection() override;
 
+	virtual void InitializeFromProperty(const FProperty* Property, const UE::MovieScene::FPropertyDefinition* Definition) override;
+
 	/** @return Get the number of channels used by the vector */
 	int32 GetNumChannelsUsed() const { return NumChannelsUsed; }
 
@@ -46,6 +48,8 @@ public:
 	virtual bool SupportsType(TSubclassOf<UMovieSceneSection> SectionClass) const override;
 	virtual UMovieSceneSection* CreateNewSection() override;
 
+	virtual void InitializeFromProperty(const FProperty* Property, const UE::MovieScene::FPropertyDefinition* Definition) override;
+
 	/** @return Get the number of channels used by the vector */
 	int32 GetNumChannelsUsed() const { return NumChannelsUsed; }
 
diff --git a/Engine/Source/Runtime/UMG/Private/Animation/MovieSceneUMGComponentTypes.cpp b/Engine/Source/Runtime/UMG/Private/Animation/MovieSceneUMGComponentTypes.cpp
--- a/Engine/Source/Runtime/UMG/Private/Animation/MovieSceneUMGComponentTypes.cpp
+++ b/Engine/Source/Runtime/UMG/Private/Animation/MovieSceneUMGComponentTypes.cpp
@@ -9,6 +9,8 @@
 #include "EntitySystem/MovieScenePropertyComponentHandler.h"
 #include "MovieSceneTracksComponentTypes.h"
 #include "Systems/MovieScenePiecewiseDoubleBlenderSystem.h"
+#include "Animation/MovieSceneMarginSection.h"
+#include "Animation/MovieScene2DTransformSection.h"
 
 #include "Components/Widget.h"
 #include "Components/CanvasPanelSlot.h"
@@ -21,6 +23,24 @@ namespace MovieScene
 static bool GMovieSceneUMGComponentTypesDestroyed = false;
 static TUniquePtr<FMovieSceneUMGComponentTypes> GMovieSceneUMGComponentTypes;
 
+void UnpackChannelsFromOperational(const FIntermediateMargin& In, const FProperty& Property, FUnpackedChannelValues& OutUnpackedValues)
+{
+	OutUnpackedValues.Add(UE_MOVIESCENE_UNPACKED_MEMBER(decltype(UMovieSceneMarginSection::LeftCurve)   , 0, In, Left  ));
+	OutUnpackedValues.Add(UE_MOVIESCENE_UNPACKED_MEMBER(decltype(UMovieSceneMarginSection::TopCurve)    , 1, In, Top   ));
+	OutUnpackedValues.Add(UE_MOVIESCENE_UNPACKED_MEMBER(decltype(UMovieSceneMarginSection::RightCurve)  , 2, In, Right ));
+	OutUnpackedValues.Add(UE_MOVIESCENE_UNPACKED_MEMBER(decltype(UMovieSceneMarginSection::BottomCurve) , 3, In, Bottom));
+}
+
+void UnpackChannelsFromOperational(const FIntermediateWidgetTransform& In, const FProperty& Property, FUnpackedChannelValues& OutUnpackedValues)
+{
+	OutUnpackedValues.Add(FUnpackedChannelValue(TIndexedChannelValue<TDecay<decltype(UMovieScene2DTransformSection::Translation[0])>::Type>(In.TranslationX, 0), FName("Translation.X")));
+	OutUnpackedValues.Add(FUnpackedChannelValue(TIndexedChannelValue<TDecay<decltype(UMovieScene2DTransformSection::Translation[1])>::Type>(In.TranslationY, 1), FName("Translation.Y")));
+	OutUnpackedValues.Add(FUnpackedChannelValue(TIndexedChannelValue<decltype(UMovieScene2DTransformSection::Rotation)>(                    In.Rotation,     2), FName("Rotation")));
+	OutUnpackedValues.Add(FUnpackedChannelValue(TIndexedChannelValue<TDecay<decltype(UMovieScene2DTransformSection::Scale[0])>::Type>(      In.ScaleX,       3), FName("Scale.X")));
+	OutUnpackedValues.Add(FUnpackedChannelValue(TIndexedChannelValue<TDecay<decltype(UMovieScene2DTransformSection::Scale[1])>::Type>(      In.ScaleY,       4), FName("Scale.Y")));
+	OutUnpackedValues.Add(FUnpackedChannelValue(TIndexedChannelValue<TDecay<decltype(UMovieScene2DTransformSection::Shear[0])>::Type>(      In.ShearX,       5), FName("Shear.X")));
+	OutUnpackedValues.Add(FUnpackedChannelValue(TIndexedChannelValue<TDecay<decltype(UMovieScene2DTransformSection::Shear[1])>::Type>(      In.ShearY,       6), FName("Shear.Y")));
+}
 
 void ConvertOperationalProperty(const FIntermediateWidgetTransform& In, FWidgetTransform& Out)
 {
diff --git a/Engine/Source/Runtime/UMG/Public/Animation/MovieSceneUMGComponentTypes.h b/Engine/Source/Runtime/UMG/Public/Animation/MovieSceneUMGComponentTypes.h
--- a/Engine/Source/Runtime/UMG/Public/Animation/MovieSceneUMGComponentTypes.h
+++ b/Engine/Source/Runtime/UMG/Public/Animation/MovieSceneUMGComponentTypes.h
@@ -8,7 +8,7 @@
 #include "EntitySystem/MovieSceneEntityIDs.h"
 #include "EntitySystem/MovieScenePropertySystemTypes.h"
 #include "EntitySystem/MovieScenePropertyTraits.h"
-#include "EntitySystem/MovieScenePropertyMetaDataTraits.h"
+#include "EntitySystem/MovieScenePropertyMetaData.h"
 
 #include "Containers/ArrayView.h"
 
@@ -51,8 +51,11 @@ struct FIntermediateMargin
 UMG_API void ConvertOperationalProperty(const FIntermediateMargin& In, FMargin& Out);
 UMG_API void ConvertOperationalProperty(const FMargin& In, FIntermediateMargin& Out);
 
-using FMarginTraits = TIndirectPropertyTraits<FMargin, FIntermediateMargin>;
-using FWidgetTransformPropertyTraits = TIndirectPropertyTraits<FWidgetTransform, FIntermediateWidgetTransform>;
+UMG_API void UnpackChannelsFromOperational(const FIntermediateMargin& In, const FProperty& Property, FUnpackedChannelValues& OutUnpackedValues);
+UMG_API void UnpackChannelsFromOperational(const FIntermediateWidgetTransform& In, const FProperty& Property, FUnpackedChannelValues& OutUnpackedValues);
+
+using FMarginTraits = TPropertyTraits<FMargin, FIntermediateMargin>;
+using FWidgetTransformPropertyTraits = TPropertyTraits<FWidgetTransform, FIntermediateWidgetTransform>;
 
 struct FMovieSceneUMGComponentTypes
 {
