Sequencer: Fixed float property traits passing single-precision floats to their intermediate representations

diff --git a/Engine/Source/Runtime/MovieSceneTracks/Public/MovieSceneTracksComponentTypes.h b/Engine/Source/Runtime/MovieSceneTracks/Public/MovieSceneTracksComponentTypes.h
--- a/Engine/Source/Runtime/MovieSceneTracks/Public/MovieSceneTracksComponentTypes.h
+++ b/Engine/Source/Runtime/MovieSceneTracks/Public/MovieSceneTracksComponentTypes.h
@@ -266,7 +266,7 @@ struct FFloatPropertyTraits : IPropertyTraits
 	FORCEINLINE static FIntermediatePropertyValue CoercePropertyValue(const FProperty& InProperty, const FSourcePropertyValue& InPropertyValue)
 	{
 		check(InProperty.IsA<FFloatProperty>());
-		return FIntermediatePropertyValue::FromValue(*InPropertyValue.Cast<float>());
+		return FIntermediatePropertyValue::FromValue(double(*InPropertyValue.Cast<float>()));
 	}
 	FORCEINLINE static void UnpackChannels(float InValue, const FProperty& Property, FUnpackedChannelValues& OutUnpackedValues)
 	{
