Sequencer: Fixed boolean tracks not being able to create track editors

diff --git a/Engine/Source/Editor/MovieSceneTools/Private/TrackEditors/PropertyTrackEditors/VisibilityPropertyTrackEditor.h b/Engine/Source/Editor/MovieSceneTools/Private/TrackEditors/PropertyTrackEditors/VisibilityPropertyTrackEditor.h
--- a/Engine/Source/Editor/MovieSceneTools/Private/TrackEditors/PropertyTrackEditors/VisibilityPropertyTrackEditor.h
+++ b/Engine/Source/Editor/MovieSceneTools/Private/TrackEditors/PropertyTrackEditors/VisibilityPropertyTrackEditor.h
@@ -35,17 +35,9 @@ public:
 
 	/** Constructor. */
 	FVisibilityPropertyTrackEditor(TSharedRef<ISequencer> InSequencer)
-		: FPropertyTrackEditor(InSequencer, GetAnimatedPropertyTypes())
+		: FPropertyTrackEditor(InSequencer)
 	{ }
 
-	/**
-	 * Retrieve a list of all property types that this track editor animates
-	 */
-	static TArray<FAnimatedPropertyKey, TInlineAllocator<1>> GetAnimatedPropertyTypes()
-	{
-		return TArray<FAnimatedPropertyKey, TInlineAllocator<1>>({ FAnimatedPropertyKey::FromPropertyTypeName(NAME_BoolProperty) });
-	}
-
 	/**
 	 * Creates an instance of this class.  Called by a sequencer 
 	 *
diff --git a/Engine/Source/Runtime/MovieSceneTracks/Private/MovieSceneTracksComponentTypes.cpp b/Engine/Source/Runtime/MovieSceneTracks/Private/MovieSceneTracksComponentTypes.cpp
--- a/Engine/Source/Runtime/MovieSceneTracks/Private/MovieSceneTracksComponentTypes.cpp
+++ b/Engine/Source/Runtime/MovieSceneTracks/Private/MovieSceneTracksComponentTypes.cpp
@@ -12,6 +12,7 @@
 #include "Tracks/MovieSceneVectorTrack.h"
 #include "Tracks/MovieSceneTransformTrack.h"
 #include "Tracks/MovieSceneEulerTransformTrack.h"
+#include "Tracks/MovieSceneBoolTrack.h"
 #include "EntitySystem/BuiltInComponentTypes.h"
 #include "EntitySystem/MovieSceneEntityManager.h"
 #include "EntitySystem/MovieSceneEntitySystemLinker.h"
@@ -599,6 +600,7 @@ FMovieSceneTracksComponentTypes::FMovieSceneTracksComponentTypes()
 	BuiltInComponents->PropertyRegistry.DefineProperty(Bool, TEXT("Apply bool Properties"))
 	.AddSoleChannel(BuiltInComponents->BoolResult)
 	.SetBlenderSystem<UMovieScenePiecewiseBoolBlenderSystem>()
+	.SetDefaultTrackType(UMovieSceneBoolTrack::StaticClass())
 	.SetCustomAccessors(&Accessors.Bool)
 	.Commit(FBoolHandler());
 
