Sequencer: Allow property tracks to be renamed

The tooltip will still show the actual name of the property being animated.

diff --git a/Engine/Source/Runtime/MovieSceneTracks/Public/Tracks/MovieScenePropertyTrack.h b/Engine/Source/Runtime/MovieSceneTracks/Public/Tracks/MovieScenePropertyTrack.h
--- a/Engine/Source/Runtime/MovieSceneTracks/Public/Tracks/MovieScenePropertyTrack.h
+++ b/Engine/Source/Runtime/MovieSceneTracks/Public/Tracks/MovieScenePropertyTrack.h
@@ -49,7 +49,6 @@ public:
 	MOVIESCENETRACKS_API virtual FText GetDefaultDisplayName() const override;
 	MOVIESCENETRACKS_API virtual FText GetDisplayNameToolTipText(const FMovieSceneLabelParams& LabelParams) const override;
 	MOVIESCENETRACKS_API virtual FSlateColor GetLabelColor(const FMovieSceneLabelParams& LabelParams) const override;
-	virtual bool CanRename() const override { return false; }
 #endif
 	MOVIESCENETRACKS_API virtual FName GetTrackName() const override;
 
