Sequencer: add support for notification callbacks when animating properties

If an object property has an InterpNotify metadata on it, we lookup the specified function name when binding to that property or to any of its sub/inner properties. We then call the notify function at the end up each sequence update.

Remove old notification function code which was effectively unused since nothing set the function name anymore.

diff --git a/Engine/Source/Runtime/MovieScene/Public/Evaluation/MovieSceneEvaluationField.h b/Engine/Source/Runtime/MovieScene/Public/Evaluation/MovieSceneEvaluationField.h
--- a/Engine/Source/Runtime/MovieScene/Public/Evaluation/MovieSceneEvaluationField.h
+++ b/Engine/Source/Runtime/MovieScene/Public/Evaluation/MovieSceneEvaluationField.h
@@ -128,6 +128,7 @@ struct FMovieSceneEvaluationFieldEntityMetaData
 			A.bEvaluateInSequencePreRoll == B.bEvaluateInSequencePreRoll && 
 			A.bEvaluateInSequencePostRoll == B.bEvaluateInSequencePostRoll &&
 			A.OverrideBoundPropertyPath == B.OverrideBoundPropertyPath &&
+			A.NotifyFunctionName == B.NotifyFunctionName &&
 			A.Condition == B.Condition;
 	}
 
@@ -140,6 +141,11 @@ struct FMovieSceneEvaluationFieldEntityMetaData
 	UPROPERTY()
 	FString OverrideBoundPropertyPath;
 
+	/** Opt-in - when this value is set, the entity should call the specified function on the bound object (if it exists) */
+	UPROPERTY()
+	FName NotifyFunctionName;
+
+	/** Opt-in - when this value is not MIN_int32, the entity should use this time to evaluate */
 	UPROPERTY()
 	FFrameNumber ForcedTime;
 
diff --git a/Engine/Source/Runtime/MovieSceneTracks/Private/MovieSceneTracksComponentTypes.cpp b/Engine/Source/Runtime/MovieSceneTracks/Private/MovieSceneTracksComponentTypes.cpp
--- a/Engine/Source/Runtime/MovieSceneTracks/Private/MovieSceneTracksComponentTypes.cpp
+++ b/Engine/Source/Runtime/MovieSceneTracks/Private/MovieSceneTracksComponentTypes.cpp
@@ -573,6 +573,8 @@ FMovieSceneTracksComponentTypes::FMovieSceneTracksComponentTypes()
 
 	ComponentRegistry->NewComponentType(&Fade,                   TEXT("Fade"), EComponentTypeFlags::CopyToChildren);
 
+	ComponentRegistry->NewComponentType(&PropertyNotify,         TEXT("PropertyNotify"), EComponentTypeFlags::CopyToChildren);
+
 	ComponentRegistry->NewComponentType(&Audio,                  TEXT("Audio"), EComponentTypeFlags::CopyToChildren);
 	ComponentRegistry->NewComponentType(&AudioInputs,            TEXT("Audio Inputs"), EComponentTypeFlags::CopyToChildren);
 	ComponentRegistry->NewComponentType(&AudioTriggerName,       TEXT("Audio Trigger Name"), EComponentTypeFlags::CopyToChildren);
diff --git a/Engine/Source/Runtime/MovieSceneTracks/Private/Systems/MovieScenePropertyNotificationSystem.cpp b/Engine/Source/Runtime/MovieSceneTracks/Private/Systems/MovieScenePropertyNotificationSystem.cpp
new file mode 100644
--- /dev/null
+++ b/Engine/Source/Runtime/MovieSceneTracks/Private/Systems/MovieScenePropertyNotificationSystem.cpp
@@ -0,0 +1,98 @@
+// Copyright Epic Games, Inc. All Rights Reserved.
+
+#include "Systems/MovieScenePropertyNotificationSystem.h"
+
+#include "EntitySystem/BuiltInComponentTypes.h"
+#include "EntitySystem/MovieSceneEntitySystemLinker.h"
+#include "EntitySystem/MovieSceneEntitySystemRunner.h"
+#include "EntitySystem/MovieSceneEntitySystemTask.h"
+#include "MovieSceneTracksComponentTypes.h"
+
+namespace UE::MovieScene
+{
+
+struct FCallNotificationsTask
+{
+	void ForEachEntity(UObject* const BoundObject, FPropertyNotifyComponentData& PropertyNotify)
+	{
+		EnsureNotifyFunctionCached(BoundObject, PropertyNotify);
+
+		UFunction* NotifyFunction = PropertyNotify.WeakNotifyFunction.Get();
+		const FNotification Notification{ BoundObject, NotifyFunction };
+		if (NotifyFunction && !ProcessedNotifications.Contains(Notification))
+		{
+			BoundObject->ProcessEvent(NotifyFunction, nullptr);
+			ProcessedNotifications.Add(Notification);
+		}
+	}
+
+private:
+
+	void EnsureNotifyFunctionCached(UObject* const BoundObject, FPropertyNotifyComponentData& PropertyNotify)
+	{
+		if (PropertyNotify.bNotifyFunctionCached)
+		{
+			return;
+		}
+
+		const FName NotifyFuncName = PropertyNotify.NotifyFunctionName;
+		if (UFunction* NotifyFunction = BoundObject->FindFunction(NotifyFuncName))
+		{
+			if (NotifyFunction->NumParms == 0 && NotifyFunction->ReturnValueOffset == MAX_uint16)
+			{
+				PropertyNotify.WeakNotifyFunction = NotifyFunction;
+			}
+			else
+			{
+				UE_LOG(LogMovieScene, Warning, 
+						TEXT("Notification function '%s' on object class '%s' must take no parameters and return no value."),
+						*NotifyFuncName.ToString(), *BoundObject->GetClass()->GetName());
+			}
+		}
+		else
+		{
+			UE_LOG(LogMovieScene, Warning, 
+					TEXT("Notification function '%s' not found on object class '%s'."),
+					*NotifyFuncName.ToString(), *BoundObject->GetClass()->GetName());
+		}
+
+		PropertyNotify.bNotifyFunctionCached = true;
+	}
+
+private:
+
+	using FNotification = TTuple<UObject*, UFunction*>;
+	TSet<FNotification> ProcessedNotifications;
+};
+
+}  // namespace UE::MovieScene
+
+UMovieScenePropertyNotificationSystem::UMovieScenePropertyNotificationSystem(const FObjectInitializer& ObjInit)
+	: Super(ObjInit)
+{
+	using namespace UE::MovieScene;
+
+	const FMovieSceneTracksComponentTypes* TracksComponents = FMovieSceneTracksComponentTypes::Get();
+
+	SystemCategories = EEntitySystemCategory::Core;
+	Phase = ESystemPhase::Finalization;
+	RelevantComponent = TracksComponents->PropertyNotify;
+}
+
+void UMovieScenePropertyNotificationSystem::OnRun(FSystemTaskPrerequisites& InPrerequisites, FSystemSubsequentTasks& Subsequents)
+{
+	using namespace UE::MovieScene;
+
+	const FBuiltInComponentTypes* BuiltInComponents = FBuiltInComponentTypes::Get();
+	const FMovieSceneTracksComponentTypes* TracksComponents = FMovieSceneTracksComponentTypes::Get();
+
+	FCallNotificationsTask CallNotificationsTask;
+
+	FEntityTaskBuilder()
+	.Read(BuiltInComponents->BoundObject)
+	.Write(TracksComponents->PropertyNotify)
+	.FilterAny({ BuiltInComponents->CustomPropertyIndex, BuiltInComponents->FastPropertyOffset, BuiltInComponents->SlowProperty })
+	.FilterNone({ BuiltInComponents->Tags.NeedsUnlink, BuiltInComponents->Tags.Ignored })
+	.RunInline_PerEntity(&Linker->EntityManager, CallNotificationsTask);
+}
+
diff --git a/Engine/Source/Runtime/MovieSceneTracks/Private/Tracks/MovieScenePropertyTrack.cpp b/Engine/Source/Runtime/MovieSceneTracks/Private/Tracks/MovieScenePropertyTrack.cpp
--- a/Engine/Source/Runtime/MovieSceneTracks/Private/Tracks/MovieScenePropertyTrack.cpp
+++ b/Engine/Source/Runtime/MovieSceneTracks/Private/Tracks/MovieScenePropertyTrack.cpp
@@ -376,6 +376,59 @@ void FMovieScenePropertyTrackEntityImportHelper::PopulateEvaluationField(UMovieS
 {
 	using namespace UE::MovieScene;
 
+	// Parse the property path, and get the expected object type for the property binding, if any.
+	TArray<FString> PropertyPathSegments;
+#if WITH_EDITORONLY_DATA
+	const UClass* ParentBoundClass = nullptr;
+#endif
+
+	UMovieScenePropertyTrack* PropertyTrack = Section.GetTypedOuter<UMovieScenePropertyTrack>();
+	UMovieScene* MovieScene = Section.GetTypedOuter<UMovieScene>();
+	if (PropertyTrack && MovieScene)
+	{
+		const FMovieScenePropertyBinding& PropertyBinding = PropertyTrack->GetPropertyBinding();
+		PropertyBinding.PropertyPath.ToString().ParseIntoArray(PropertyPathSegments, TEXT("."), true);
+
+#if WITH_EDITORONLY_DATA
+		FGuid ParentBindingGuid = OutFieldBuilder->GetSharedMetaData().ObjectBindingID;
+		if (ParentBindingGuid.IsValid())
+		{
+			if (const FMovieScenePossessable* ParentPossessable = MovieScene->FindPossessable(ParentBindingGuid))
+			{
+				ParentBoundClass = ParentPossessable->GetPossessedObjectClass();
+			}
+			else if (const FMovieSceneSpawnable* ParentSpawnable = MovieScene->FindSpawnable(ParentBindingGuid))
+			{
+				ParentBoundClass = ParentSpawnable->GetObjectTemplate() ? 
+					ParentSpawnable->GetObjectTemplate()->GetClass() : nullptr;
+			}
+		}
+#endif
+	}
+
+	FMovieSceneEvaluationFieldEntityMetaData MainMetaData(InMetaData);
+
+#if WITH_EDITORONLY_DATA
+
+	// Check if this section is animating a property with a notify function. If so, add that information to the metadata.
+	if (ParentBoundClass && PropertyPathSegments.Num() > 0)
+	{
+		if (FProperty* HeadProperty = ParentBoundClass->FindPropertyByName(*PropertyPathSegments[0]))
+		{
+			const FString InterpNotifyMetaData = HeadProperty->GetMetaData(TEXT("InterpNotify"));
+			if (!InterpNotifyMetaData.IsEmpty())
+			{
+				MainMetaData.NotifyFunctionName = FName(InterpNotifyMetaData);
+			}
+		}
+	}
+
+#else
+
+	// No support for notify functions for sequences dynamically built at runtime.
+
+#endif  // WITH_EDITORONLY_DATA
+
 	int32 NumOverridenChannels = 0;
 	IMovieSceneChannelOverrideProvider* RegistryProvider = Cast<IMovieSceneChannelOverrideProvider>(&Section);
 	if (RegistryProvider)
@@ -383,7 +436,7 @@ void FMovieScenePropertyTrackEntityImportHelper::PopulateEvaluationField(UMovieS
 		if (UMovieSceneSectionChannelOverrideRegistry* OverrideRegistry = RegistryProvider->GetChannelOverrideRegistry(false))
 		{
 			NumOverridenChannels = OverrideRegistry->NumChannels();
-			OverrideRegistry->PopulateEvaluationFieldImpl(EffectiveRange, InMetaData, OutFieldBuilder, Section);
+			OverrideRegistry->PopulateEvaluationFieldImpl(EffectiveRange, MainMetaData, OutFieldBuilder, Section);
 		}
 	}
 
@@ -392,22 +445,14 @@ void FMovieScenePropertyTrackEntityImportHelper::PopulateEvaluationField(UMovieS
 	{
 		// Add the default entity for this section.
 		const int32 EntityIndex = OutFieldBuilder->FindOrAddEntity(&Section, SectionPropertyValueImportingID);
-		const int32 MetaDataIndex = OutFieldBuilder->AddMetaData(InMetaData);
+		const int32 MetaDataIndex = OutFieldBuilder->AddMetaData(MainMetaData);
 		OutFieldBuilder->AddPersistentEntity(EffectiveRange, EntityIndex, MetaDataIndex);
 	}
 
 	// Check if this section is animating a property with an edit-condition. If so, we need to also animate a boolean toggle
 	// that will be set to true while the main property is animated.
-	UMovieScenePropertyTrack* PropertyTrack = Section.GetTypedOuter<UMovieScenePropertyTrack>();
-	UMovieScene* MovieScene = Section.GetTypedOuter<UMovieScene>();
-	if (PropertyTrack && MovieScene)
+	if (PropertyPathSegments.Num() > 0)
 	{
-		const FMovieScenePropertyBinding& PropertyBinding = PropertyTrack->GetPropertyBinding();
-
-		TArray<FString> PropertyPathSegments;
-		PropertyBinding.PropertyPath.ToString().ParseIntoArray(PropertyPathSegments, TEXT("."), true);
-		FCachedPropertyPath PropertyPath(PropertyPathSegments);
-
 		// Prepare the edit condition toggle property path by taking the beginning part of the
 		// main property path. We'll append the toggle name to it. This is necessary for nested stuff, like:
 		//
@@ -424,36 +469,21 @@ void FMovieScenePropertyTrackEntityImportHelper::PopulateEvaluationField(UMovieS
 
 #if WITH_EDITORONLY_DATA
 
-		FGuid ParentBindingGuid;
-		const bool bFoundParentBinding = MovieScene->FindTrackBinding(*PropertyTrack, ParentBindingGuid);
-		if (bFoundParentBinding && ParentBindingGuid.IsValid())
+		if (ParentBoundClass)
 		{
-			const UClass* ParentBoundClass = nullptr;
-			if (const FMovieScenePossessable* ParentPossessable = MovieScene->FindPossessable(ParentBindingGuid))
-			{
-				ParentBoundClass = ParentPossessable->GetPossessedObjectClass();
-			}
-			else if (const FMovieSceneSpawnable* ParentSpawnable = MovieScene->FindSpawnable(ParentBindingGuid))
+			FCachedPropertyPath PropertyPath(PropertyPathSegments);
+			PropertyPath.Resolve(ParentBoundClass->GetDefaultObject());
+			if (const FProperty* LeafProperty = PropertyPath.GetFProperty())
 			{
-				ParentBoundClass = ParentSpawnable->GetObjectTemplate() ? 
-					ParentSpawnable->GetObjectTemplate()->GetClass() : nullptr;
-			}
-
-			if (ParentBoundClass) // This line was previously an ensure() but we removed it since that would fail cooks
-			{
-				PropertyPath.Resolve(ParentBoundClass->GetDefaultObject());
-				if (const FProperty* LeafProperty = PropertyPath.GetFProperty())
+				const FString EditConditionPropertyName = LeafProperty->GetMetaData("EditCondition");
+				if (!EditConditionPropertyName.IsEmpty())
 				{
-					const FString EditConditionPropertyName = LeafProperty->GetMetaData("EditCondition");
-					if (!EditConditionPropertyName.IsEmpty())
+					if (!EditConditionPropertyPath.IsEmpty())
 					{
-						if (!EditConditionPropertyPath.IsEmpty())
-						{
-							EditConditionPropertyPath.Append(".");
-						}
-						EditConditionPropertyPath.Append(EditConditionPropertyName);
-						bHasEditCondition = true;
+						EditConditionPropertyPath.Append(".");
 					}
+					EditConditionPropertyPath.Append(EditConditionPropertyName);
+					bHasEditCondition = true;
 				}
 			}
 		}
@@ -468,7 +498,6 @@ void FMovieScenePropertyTrackEntityImportHelper::PopulateEvaluationField(UMovieS
 			FString EditConditionPropertyName = PropertyPathSegments[PropertyPathSegments.Num() - 1];
 			EditConditionPropertyName.InsertAt(0, TEXT("bOverride_"));
 
-
 			if (!EditConditionPropertyPath.IsEmpty())
 			{
 				EditConditionPropertyPath.Append(".");
@@ -477,7 +506,7 @@ void FMovieScenePropertyTrackEntityImportHelper::PopulateEvaluationField(UMovieS
 			bHasEditCondition = true;
 		}
 
-#endif
+#endif  // WITH_EDITORONLY_DATA
 
 		if (bHasEditCondition)
 		{
diff --git a/Engine/Source/Runtime/MovieSceneTracks/Public/MovieSceneTracksComponentTypes.h b/Engine/Source/Runtime/MovieSceneTracks/Public/MovieSceneTracksComponentTypes.h
--- a/Engine/Source/Runtime/MovieSceneTracks/Public/MovieSceneTracksComponentTypes.h
+++ b/Engine/Source/Runtime/MovieSceneTracks/Public/MovieSceneTracksComponentTypes.h
@@ -250,6 +250,13 @@ struct FFadeComponentData
 	bool bFadeAudio;
 };
 
+struct FPropertyNotifyComponentData
+{
+	FName NotifyFunctionName;
+	TWeakObjectPtr<UFunction> WeakNotifyFunction;
+	bool bNotifyFunctionCached = false;
+};
+
 struct FFloatPropertyTraits : IPropertyTraits
 {
 	static constexpr bool bIsComposite = false;
@@ -580,6 +587,8 @@ struct FMovieSceneTracksComponentTypes
 
 	TComponentTypeID<FFadeComponentData> Fade;
 
+	TComponentTypeID<FPropertyNotifyComponentData> PropertyNotify;
+
 	TComponentTypeID<FMovieSceneAudioComponentData> Audio;
 	TComponentTypeID<FMovieSceneAudioInputData> AudioInputs;
 	TComponentTypeID<FName> AudioTriggerName;
diff --git a/Engine/Source/Runtime/MovieSceneTracks/Public/Systems/MovieScenePropertyNotificationSystem.h b/Engine/Source/Runtime/MovieSceneTracks/Public/Systems/MovieScenePropertyNotificationSystem.h
new file mode 100644
--- /dev/null
+++ b/Engine/Source/Runtime/MovieSceneTracks/Public/Systems/MovieScenePropertyNotificationSystem.h
@@ -0,0 +1,25 @@
+// Copyright Epic Games, Inc. All Rights Reserved.
+
+#pragma once
+
+#include "EntitySystem/MovieSceneEntitySystem.h"
+
+#include "MovieScenePropertyNotificationSystem.generated.h"
+
+/**
+ * A system that can notify an object that some of its properties are being animated.
+ */
+UCLASS(MinimalAPI)
+class UMovieScenePropertyNotificationSystem : public UMovieSceneEntitySystem
+{
+	GENERATED_BODY()
+
+public:
+
+	MOVIESCENETRACKS_API UMovieScenePropertyNotificationSystem(const FObjectInitializer& ObjInit);
+
+protected:
+
+	MOVIESCENETRACKS_API virtual void OnRun(FSystemTaskPrerequisites& InPrerequisites, FSystemSubsequentTasks& Subsequents) override;
+};
+
diff --git a/Engine/Source/Runtime/MovieSceneTracks/Public/Tracks/MovieScenePropertyTrack.h b/Engine/Source/Runtime/MovieSceneTracks/Public/Tracks/MovieScenePropertyTrack.h
--- a/Engine/Source/Runtime/MovieSceneTracks/Public/Tracks/MovieScenePropertyTrack.h
+++ b/Engine/Source/Runtime/MovieSceneTracks/Public/Tracks/MovieScenePropertyTrack.h
@@ -333,6 +333,7 @@ struct TPropertyTrackWithOverridableChannelsEntityImportHelperImpl
 	void Commit(const UMovieSceneSection* InSection, const UE::MovieScene::FEntityImportParams& Params, UE::MovieScene::FImportedEntity* OutImportedEntity)
 	{
 		const FGuid ObjectBindingID = Params.GetObjectBindingID();
+		const FMovieSceneEvaluationFieldEntityMetaData* MetaData = Params.EntityMetaData;
 		const FBuiltInComponentTypes* BuiltInComponents = FBuiltInComponentTypes::Get();
 		const FMovieSceneTracksComponentTypes* TracksComponents = FMovieSceneTracksComponentTypes::Get();
 
@@ -340,10 +341,13 @@ struct TPropertyTrackWithOverridableChannelsEntityImportHelperImpl
 		{
 			if (FMovieScenePropertyTrackEntityImportHelper::IsPropertyValueID(Params))
 			{
+				const FName NotifyFunctionName = (MetaData ? MetaData->NotifyFunctionName : NAME_None);
+
 				OutImportedEntity->AddBuilder(
 					Builder
 					.Add(BuiltInComponents->PropertyBinding, PropertyTrack->GetPropertyBinding())
-					.AddConditional(BuiltInComponents->GenericObjectBinding, ObjectBindingID, ObjectBindingID.IsValid()));
+					.AddConditional(BuiltInComponents->GenericObjectBinding, ObjectBindingID, ObjectBindingID.IsValid())
+					.AddConditional(TracksComponents->PropertyNotify, FPropertyNotifyComponentData{ NotifyFunctionName }, !NotifyFunctionName.IsNone()));
 			}
 			else if (FMovieScenePropertyTrackEntityImportHelper::IsEditConditionToggleID(Params))
 			{
