Sequencer: Fix issue where when changing shots we could sometimes get one frame of t-pose on a character when using bForceCustomMode.

The issue was that we were checking for custom mode too late in the frame relative to when we were potentially doing the binding to a custom anim instance.

diff --git a/Engine/Plugins/Experimental/MovieSceneAnimMixer/Source/MovieSceneAnimMixer/Private/Systems/MovieSceneAnimInstanceTargetSystem.cpp b/Engine/Plugins/Experimental/MovieSceneAnimMixer/Source/MovieSceneAnimMixer/Private/Systems/MovieSceneAnimInstanceTargetSystem.cpp
--- a/Engine/Plugins/Experimental/MovieSceneAnimMixer/Source/MovieSceneAnimMixer/Private/Systems/MovieSceneAnimInstanceTargetSystem.cpp
+++ b/Engine/Plugins/Experimental/MovieSceneAnimMixer/Source/MovieSceneAnimMixer/Private/Systems/MovieSceneAnimInstanceTargetSystem.cpp
@@ -447,7 +447,12 @@ namespace UE::MovieScene
 			TWeakObjectPtr CurrentAnimInstanceWeakPtr (SkeletalMeshComponent->GetAnimInstance());
 			TWeakObjectPtr SkeletalMeshComponentWeakPtr (SkeletalMeshComponent);
 #endif
-			
+			// By using the anim instance target, we are implicitly using AnimationCustomMode. We set it here explicitly to ensure the first frame initializes properly.
+			if (SkeletalMeshComponent->GetAnimationMode() != EAnimationMode::AnimationCustomMode)
+			{
+				SkeletalMeshComponent->SetAnimationMode(EAnimationMode::AnimationCustomMode);
+			}
+
 			USequencerMixedAnimInstance* SequencerInstance = FAnimCustomInstanceHelper::BindToSkeletalMeshComponent<USequencerMixedAnimInstance>(SkeletalMeshComponent, bWasCreated);
 
 			if (SequencerInstance)
diff --git a/Engine/Source/Runtime/MovieSceneTracks/Private/Systems/MovieSceneSkeletalAnimationSystem.cpp b/Engine/Source/Runtime/MovieSceneTracks/Private/Systems/MovieSceneSkeletalAnimationSystem.cpp
--- a/Engine/Source/Runtime/MovieSceneTracks/Private/Systems/MovieSceneSkeletalAnimationSystem.cpp
+++ b/Engine/Source/Runtime/MovieSceneTracks/Private/Systems/MovieSceneSkeletalAnimationSystem.cpp
@@ -619,6 +619,25 @@ private:
 			PreAnimatedStorage->CachePreAnimatedValue(CacheParams, SkeletalMeshComponent);
 		}
 
+		// Check to see if we need to switch modes on the anim instance
+		if (Algo::AnyOf(InSkeletalAnimations.Animations, [&SkeletalMeshComponent](const FActiveSkeletalAnimation& ActiveAnim)
+			{
+				UAnimSequenceBase* Animation = ActiveAnim.AnimSection->GetPlaybackAnimation();
+				if (CanPlayAnimation(SkeletalMeshComponent, Animation))
+				{
+					const FMovieSceneSkeletalAnimationParams& AnimParams = ActiveAnim.AnimSection->Params;
+					TScriptInterface<ISequencerAnimationOverride> SequencerAnimOverride = ISequencerAnimationOverride::GetSequencerAnimOverride(SkeletalMeshComponent);
+					if (AnimParams.bForceCustomMode || (SequencerAnimOverride.GetObject() && ISequencerAnimationOverride::Execute_AllowsCinematicOverride(SequencerAnimOverride.GetObject())))
+					{
+						return true;
+					}
+				}
+				return false;
+			}))
+		{
+			SkeletalMeshComponent->SetAnimationMode(EAnimationMode::AnimationCustomMode);
+		}
+
 		// Setup any needed animation nodes for sequencer playback.
 		UAnimInstance* ExistingAnimInstance = GetSourceAnimInstance(SkeletalMeshComponent);
 		bool bWasCreated = false;
@@ -877,11 +896,6 @@ private:
 		{
 			return;
 		}
-		TScriptInterface<ISequencerAnimationOverride> SequencerAnimOverride = ISequencerAnimationOverride::GetSequencerAnimOverride(Params.SkeletalMeshComponent);
-		if (AnimParams.bForceCustomMode || (SequencerAnimOverride.GetObject() && ISequencerAnimationOverride::Execute_AllowsCinematicOverride(SequencerAnimOverride.GetObject())))
-		{
-			Params.SkeletalMeshComponent->SetAnimationMode(EAnimationMode::AnimationCustomMode);
-		}
 
 		UAnimSequencerInstance* SequencerInst = GetAnimSequencerInstance(Params.SkeletalMeshComponent);
 		if (SequencerInst)
@@ -974,11 +988,7 @@ private:
 		{
 			return;
 		}
-		TScriptInterface<ISequencerAnimationOverride> SequencerAnimOverride = ISequencerAnimationOverride::GetSequencerAnimOverride(Params.SkeletalMeshComponent);
-		if (AnimParams.bForceCustomMode || (SequencerAnimOverride.GetObject() && ISequencerAnimationOverride::Execute_AllowsCinematicOverride(SequencerAnimOverride.GetObject())))
-		{
-			Params.SkeletalMeshComponent->SetAnimationMode(EAnimationMode::AnimationCustomMode);
-		}
+
 		UAnimSequencerInstance* SequencerInst = GetAnimSequencerInstance(Params.SkeletalMeshComponent);
 		if (SequencerInst)
 		{
