MovieSceneTracksCustomAccessors: Add custom accessor for USkyLightComponent::bRealtimeCapture

diff --git a/Engine/Source/Runtime/MovieSceneTracks/Private/MovieSceneTracksCustomAccessors.cpp b/Engine/Source/Runtime/MovieSceneTracks/Private/MovieSceneTracksCustomAccessors.cpp
--- a/Engine/Source/Runtime/MovieSceneTracks/Private/MovieSceneTracksCustomAccessors.cpp
+++ b/Engine/Source/Runtime/MovieSceneTracks/Private/MovieSceneTracksCustomAccessors.cpp
@@ -167,6 +167,18 @@ void SetSkyLightComponentIntensity(UObject* Object, float InIntensity)
 	SkyLightComponent->SetIntensity(InIntensity);
 }
 
+bool GetSkyLightComponentRealtimeCapture(const UObject* Object, FBoolPropertyTraits::FBoolMetaData InMetadata)
+{
+	const USkyLightComponent* SkyLightComponent = CastChecked<const USkyLightComponent>(Object);
+	return SkyLightComponent->bRealTimeCapture;
+}
+
+void SetSkyLightComponentRealtimeCapture(UObject* Object, FBoolPropertyTraits::FBoolMetaData InMetadata, bool bInRealtimeCapture)
+{
+	USkyLightComponent* SkyLightComponent = CastChecked<USkyLightComponent>(Object);
+	SkyLightComponent->SetRealTimeCapture(bInRealtimeCapture);
+}
+
 float GetSkyAtmosphereComponentMieScatteringScale(const UObject* Object)
 {
 	const USkyAtmosphereComponent* SkyAtmosphereComponent = CastChecked<const USkyAtmosphereComponent>(Object);
@@ -584,6 +596,9 @@ void InitializeMovieSceneTracksAccessors(FMovieSceneTracksComponentTypes* Tracks
 	TracksComponents->Accessors.Float.Add(
 			USkyLightComponent::StaticClass(), GET_MEMBER_NAME_CHECKED(USkyLightComponent, Intensity),
 			GetSkyLightComponentIntensity, SetSkyLightComponentIntensity);
+	TracksComponents->Accessors.Bool.Add(
+			USkyLightComponent::StaticClass(), GET_MEMBER_NAME_CHECKED(USkyLightComponent, bRealTimeCapture),
+			GetSkyLightComponentRealtimeCapture, SetSkyLightComponentRealtimeCapture);
 
 	// SkyAtmosphereComponent
 	TracksComponents->Accessors.Float.Add(
