Sequencer: Add missing modify when editing an existing camera cut section

diff --git a/Engine/Source/Runtime/MovieSceneTracks/Private/Tracks/MovieSceneCameraCutTrack.cpp b/Engine/Source/Runtime/MovieSceneTracks/Private/Tracks/MovieSceneCameraCutTrack.cpp
--- a/Engine/Source/Runtime/MovieSceneTracks/Private/Tracks/MovieSceneCameraCutTrack.cpp
+++ b/Engine/Source/Runtime/MovieSceneTracks/Private/Tracks/MovieSceneCameraCutTrack.cpp
@@ -51,6 +51,7 @@ UMovieSceneCameraCutSection* UMovieSceneCameraCutTrack::AddNewCameraCut(const FM
 	UMovieSceneCameraCutSection* NewSection = ExistingSection;
 	if (ExistingSection != nullptr)
 	{
+		ExistingSection->Modify();
 		ExistingSection->SetCameraBindingID(CameraBindingID);
 	}
 	else
