Take Recorder: Fix crash when saving take recorded actors that reference other objects in the map.
Any illegal references are null'ed out.

diff --git a/Engine/Source/Runtime/MovieSceneTracks/Private/Bindings/MovieSceneSpawnableActorBinding.cpp b/Engine/Source/Runtime/MovieSceneTracks/Private/Bindings/MovieSceneSpawnableActorBinding.cpp
--- a/Engine/Source/Runtime/MovieSceneTracks/Private/Bindings/MovieSceneSpawnableActorBinding.cpp
+++ b/Engine/Source/Runtime/MovieSceneTracks/Private/Bindings/MovieSceneSpawnableActorBinding.cpp
@@ -21,6 +21,7 @@
 #include "TransformData.h"
 #include "MovieSceneBindingReferences.h"
 #include "MovieSceneCommonHelpers.h"
+#include "Serialization/ArchiveReplaceOrClearExternalReferences.h"
 #include "Systems/MovieSceneDeferredComponentMovementSystem.h"
 
 #include UE_INLINE_GENERATED_CPP_BY_NAME(MovieSceneSpawnableActorBinding)
@@ -29,6 +30,54 @@ static const FName SequencerActorTag(TEXT("SequencerActor"));
 static const FName SequencerPreviewActorTag(TEXT("SequencerPreviewActor"));
 #define LOCTEXT_NAMESPACE "MovieScene"
 
+namespace UE::MovieSceneTracks
+{
+class FArchiveClearExternalReferences : public FArchiveReplaceOrClearExternalReferences<UObject>
+{
+	using Super = FArchiveReplaceOrClearExternalReferences<UObject>;
+
+	TArray<UObject*> RemovedObjects;
+public:
+	
+	FArchiveClearExternalReferences(UObject* InSearchObject, const TMap<UObject*, UObject*>& InReplacementMap, UPackage* InDestPackage)
+		: Super(InSearchObject, InReplacementMap, InDestPackage, EArchiveReplaceObjectFlags::NullPrivateRefs | EArchiveReplaceObjectFlags::DelayStart )
+	{
+		SerializeSearchObject();
+	}
+
+	const TArray<UObject*>& GetRemovedObjects() const { return RemovedObjects; }
+
+	virtual FArchive& operator<<(UObject*& Obj) override
+	{
+		UObject* Original = Obj;
+		Super::operator<<(Obj);
+
+		if (Original && !Obj)
+		{
+			RemovedObjects.AddUnique(Original);
+		}
+		
+		return *this;
+	}
+};
+
+static void LogRemovedExternalObjects(UObject* ReferencingObject, const TArray<UObject*>& InRemovedObjects)
+{
+	if (InRemovedObjects.IsEmpty())
+	{
+		return;
+	}
+
+	const FString List = FString::JoinBy(InRemovedObjects, TEXT(", "), [](UObject* Object)
+	{
+		return ensure(Object) ? Object->GetPathName() : TEXT("null");
+	});
+	UE_LOG(LogMovieScene, Warning, TEXT("While saving %s, the following references were cleared because they were private and external: %s"),
+		*ReferencingObject->GetPathName(), *List
+		);
+}
+}
+
 UObject* UMovieSceneSpawnableActorBindingBase::SpawnObjectInternal(UWorld* WorldContext, FName SpawnName, const FGuid& BindingId, int32 BindingIndex, UMovieScene& MovieScene, FMovieSceneSequenceIDRef TemplateID, TSharedRef<const UE::MovieScene::FSharedPlaybackState> SharedPlaybackState)
 {
 	const EObjectFlags SpawnObjectFlags = RF_Transient | RF_Transactional;
@@ -477,6 +526,17 @@ UMovieSceneCustomBinding* UMovieSceneSpawnableActorBinding::CreateNewCustomBindi
 			{
 				Actor->ClearFlags(RF_Transactional);
 			}
+
+			// This achieves that the template actor keeps no reference to any external private objects.
+			// Such references would cause a crash when saving the owning UMovieScene / sequence.
+			// Example: An actor that has a reference to another actor in the level.
+			// SpawnedActor keep referencing the actor in the original map, which is a private external object.
+			// FArchiveReplaceOrClearExternalReferences handles recursively serializing subobjects, such as components.
+			UPackage* NewPackage = SpawnedActor->GetOutermost();
+			const UE::MovieSceneTracks::FArchiveClearExternalReferences ReplaceActorInvalidReferences(SpawnedActor, {}, NewPackage);
+			// User should be warned in case they notice weird behaviour (e.g. if the spawnable actor is selected in details panel,
+			// saving nulls out the references which the user may notice).
+			UE::MovieSceneTracks::LogRemovedExternalObjects(this, ReplaceActorInvalidReferences.GetRemovedObjects());
 		}
 	}
 	// If it's a blueprint, we need some special handling
