Sequencer: Refactor material tracks to use DefaultInterpolation so that they follow other keyed curve track behaviors

diff --git a/Engine/Source/Editor/MovieSceneTools/Private/TrackEditors/CustomPrimitiveDataTrackEditor.cpp b/Engine/Source/Editor/MovieSceneTools/Private/TrackEditors/CustomPrimitiveDataTrackEditor.cpp
--- a/Engine/Source/Editor/MovieSceneTools/Private/TrackEditors/CustomPrimitiveDataTrackEditor.cpp
+++ b/Engine/Source/Editor/MovieSceneTools/Private/TrackEditors/CustomPrimitiveDataTrackEditor.cpp
@@ -334,7 +334,7 @@ void FCustomPrimitiveDataTrackEditor::AddScalarParameter(UPrimitiveComponent* Co
 		ParameterValue = CustomPrimitiveData[PrimitiveDataIndex];
 	}
 	CPDTrack->Modify();
-	CPDTrack->AddScalarParameterKey(PrimitiveDataIndex, KeyTime, ParameterValue);
+	CPDTrack->AddScalarParameterKey(PrimitiveDataIndex, KeyTime, ParameterValue, GetSequencer()->GetKeyInterpolation());
 
 	GetSequencer()->NotifyMovieSceneDataChanged(EMovieSceneDataChangeType::MovieSceneStructureItemAdded);
 }
@@ -353,7 +353,7 @@ void FCustomPrimitiveDataTrackEditor::AddVector2DParameter(UPrimitiveComponent*
 		ParameterValue.Y = CustomPrimitiveData[PrimitiveDataStartIndex + 1];
 	}
 	CPDTrack->Modify();
-	CPDTrack->AddVector2DParameterKey(PrimitiveDataStartIndex, KeyTime, ParameterValue);
+	CPDTrack->AddVector2DParameterKey(PrimitiveDataStartIndex, KeyTime, ParameterValue, GetSequencer()->GetKeyInterpolation());
 
 	GetSequencer()->NotifyMovieSceneDataChanged(EMovieSceneDataChangeType::MovieSceneStructureItemAdded);
 }
@@ -373,7 +373,7 @@ void FCustomPrimitiveDataTrackEditor::AddVectorParameter(UPrimitiveComponent* Co
 		ParameterValue.Z = CustomPrimitiveData[PrimitiveDataStartIndex + 2];
 	}
 	CPDTrack->Modify();
-	CPDTrack->AddVectorParameterKey(PrimitiveDataStartIndex, KeyTime, ParameterValue);
+	CPDTrack->AddVectorParameterKey(PrimitiveDataStartIndex, KeyTime, ParameterValue, GetSequencer()->GetKeyInterpolation());
 
 	GetSequencer()->NotifyMovieSceneDataChanged(EMovieSceneDataChangeType::MovieSceneStructureItemAdded);
 }
@@ -394,7 +394,7 @@ void FCustomPrimitiveDataTrackEditor::AddColorParameter(UPrimitiveComponent* Com
 		ParameterValue.A = CustomPrimitiveData[PrimitiveDataStartIndex + 3];
 	}
 	CPDTrack->Modify();
-	CPDTrack->AddColorParameterKey(PrimitiveDataStartIndex, KeyTime, ParameterValue);
+	CPDTrack->AddColorParameterKey(PrimitiveDataStartIndex, KeyTime, ParameterValue, GetSequencer()->GetKeyInterpolation());
 
 	GetSequencer()->NotifyMovieSceneDataChanged(EMovieSceneDataChangeType::MovieSceneStructureItemAdded);
 }
diff --git a/Engine/Source/Editor/MovieSceneTools/Private/TrackEditors/MaterialParameterCollectionTrackEditor.cpp b/Engine/Source/Editor/MovieSceneTools/Private/TrackEditors/MaterialParameterCollectionTrackEditor.cpp
--- a/Engine/Source/Editor/MovieSceneTools/Private/TrackEditors/MaterialParameterCollectionTrackEditor.cpp
+++ b/Engine/Source/Editor/MovieSceneTools/Private/TrackEditors/MaterialParameterCollectionTrackEditor.cpp
@@ -327,7 +327,7 @@ void FMaterialParameterCollectionTrackEditor::AddScalarParameter(UMovieSceneMate
 
 	const FScopedTransaction Transaction(LOCTEXT("AddScalarParameter", "Add scalar parameter"));
 	Track->Modify();
-	Track->AddScalarParameterKey(Parameter.ParameterName, KeyTime, RowIndex, Value);
+	Track->AddScalarParameterKey(Parameter.ParameterName, KeyTime, RowIndex, Value, GetSequencer()->GetKeyInterpolation());
 	GetSequencer()->NotifyMovieSceneDataChanged(EMovieSceneDataChangeType::MovieSceneStructureItemAdded);
 }
 
@@ -352,7 +352,7 @@ void FMaterialParameterCollectionTrackEditor::AddVectorParameter(UMovieSceneMate
 
 	const FScopedTransaction Transaction(LOCTEXT("AddVectorParameter", "Add vector parameter"));
 	Track->Modify();
-	Track->AddColorParameterKey(Parameter.ParameterName, KeyTime, RowIndex, Value);
+	Track->AddColorParameterKey(Parameter.ParameterName, KeyTime, RowIndex, Value, GetSequencer()->GetKeyInterpolation());
 	GetSequencer()->NotifyMovieSceneDataChanged(EMovieSceneDataChangeType::MovieSceneStructureItemAdded);
 }
 
diff --git a/Engine/Source/Editor/MovieSceneTools/Private/TrackEditors/MaterialTrackEditor.cpp b/Engine/Source/Editor/MovieSceneTools/Private/TrackEditors/MaterialTrackEditor.cpp
--- a/Engine/Source/Editor/MovieSceneTools/Private/TrackEditors/MaterialTrackEditor.cpp
+++ b/Engine/Source/Editor/MovieSceneTools/Private/TrackEditors/MaterialTrackEditor.cpp
@@ -258,7 +258,7 @@ void FMaterialTrackEditor::AddScalarParameter( FGuid ObjectBinding, UMovieSceneM
 		float ParameterValue;
 		Material->GetScalarParameterValue(ParameterInfo, ParameterValue);
 		MaterialTrack->Modify();
-		MaterialTrack->AddScalarParameterKey(ParameterInfo, KeyTime, ParameterValue, InLayerName, InAssetName);
+		MaterialTrack->AddScalarParameterKey(ParameterInfo, KeyTime, ParameterValue, InLayerName, InAssetName, GetSequencer()->GetKeyInterpolation());
 	}
 	GetSequencer()->NotifyMovieSceneDataChanged( EMovieSceneDataChangeType::MovieSceneStructureItemAdded );
 }
@@ -277,7 +277,7 @@ void FMaterialTrackEditor::AddColorParameter( FGuid ObjectBinding, UMovieSceneMa
 		MaterialTrack->Modify();
 		FParameterChannelNames ChannelNames;
 		Material->GetVectorParameterChannelNames(ParameterInfo, ChannelNames);
-		MaterialTrack->AddColorParameterKey(ParameterInfo, KeyTime, INDEX_NONE, ParameterValue, InLayerName, InAssetName, ChannelNames);
+		MaterialTrack->AddColorParameterKey(ParameterInfo, KeyTime, INDEX_NONE, ParameterValue, InLayerName, InAssetName, ChannelNames, GetSequencer()->GetKeyInterpolation());
 	}
 	GetSequencer()->NotifyMovieSceneDataChanged( EMovieSceneDataChangeType::MovieSceneStructureItemAdded );
 }
diff --git a/Engine/Source/Editor/MovieSceneTools/Private/TrackEditors/ParticleParameterTrackEditor.cpp b/Engine/Source/Editor/MovieSceneTools/Private/TrackEditors/ParticleParameterTrackEditor.cpp
--- a/Engine/Source/Editor/MovieSceneTools/Private/TrackEditors/ParticleParameterTrackEditor.cpp
+++ b/Engine/Source/Editor/MovieSceneTools/Private/TrackEditors/ParticleParameterTrackEditor.cpp
@@ -156,7 +156,7 @@ void FParticleParameterTrackEditor::AddScalarParameter( FGuid ObjectBinding, UMo
 		ParticleSystemComponent->GetFloatParameter(ParameterName, Value);
 
 		ParticleParameterTrack->Modify();
-		ParticleParameterTrack->AddScalarParameterKey(ParameterName, KeyTime, Value);
+		ParticleParameterTrack->AddScalarParameterKey(ParameterName, KeyTime, Value, GetSequencer()->GetKeyInterpolation());
 
 		GetSequencer()->NotifyMovieSceneDataChanged(EMovieSceneDataChangeType::MovieSceneStructureItemAdded);
 	}
@@ -175,7 +175,7 @@ void FParticleParameterTrackEditor::AddVectorParameter( FGuid ObjectBinding, UMo
 		ParticleSystemComponent->GetVectorParameter(ParameterName, Value);
 
 		ParticleParameterTrack->Modify();
-		ParticleParameterTrack->AddVectorParameterKey(ParameterName, KeyTime, Value);
+		ParticleParameterTrack->AddVectorParameterKey(ParameterName, KeyTime, Value, GetSequencer()->GetKeyInterpolation());
 
 		GetSequencer()->NotifyMovieSceneDataChanged(EMovieSceneDataChangeType::MovieSceneStructureItemAdded);
 	}
@@ -194,7 +194,7 @@ void FParticleParameterTrackEditor::AddColorParameter( FGuid ObjectBinding, UMov
 		ParticleSystemComponent->GetColorParameter(ParameterName, Value);
 
 		ParticleParameterTrack->Modify();
-		ParticleParameterTrack->AddColorParameterKey(ParameterName, KeyTime, Value);
+		ParticleParameterTrack->AddColorParameterKey(ParameterName, KeyTime, Value, GetSequencer()->GetKeyInterpolation());
 
 		GetSequencer()->NotifyMovieSceneDataChanged(EMovieSceneDataChangeType::MovieSceneStructureItemAdded);
 	}
diff --git a/Engine/Source/Runtime/MovieSceneTracks/Private/Sections/MovieSceneComponentMaterialParameterSection.cpp b/Engine/Source/Runtime/MovieSceneTracks/Private/Sections/MovieSceneComponentMaterialParameterSection.cpp
--- a/Engine/Source/Runtime/MovieSceneTracks/Private/Sections/MovieSceneComponentMaterialParameterSection.cpp
+++ b/Engine/Source/Runtime/MovieSceneTracks/Private/Sections/MovieSceneComponentMaterialParameterSection.cpp
@@ -168,7 +168,7 @@ void UMovieSceneComponentMaterialParameterSection::PostEditUndo()
 	CacheChannelProxy();
 }
 #endif
-void UMovieSceneComponentMaterialParameterSection::AddScalarParameterKey(const FMaterialParameterInfo& InParameterInfo, FFrameNumber InTime, float InValue, const FString& InLayerName, const FString& InAssetName)
+void UMovieSceneComponentMaterialParameterSection::AddScalarParameterKey(const FMaterialParameterInfo& InParameterInfo, FFrameNumber InTime, float InValue, const FString& InLayerName, const FString& InAssetName, EMovieSceneKeyInterpolation DefaultInterpolation)
 {
 	FMovieSceneFloatChannel* ExistingChannel = nullptr;
 	for (FScalarMaterialParameterInfoAndCurve& ScalarParameterInfoAndCurve : ScalarParameterInfosAndCurves)
@@ -190,7 +190,8 @@ void UMovieSceneComponentMaterialParameterSection::AddScalarParameterKey(const F
 		CacheChannelProxy();
 	}
 
-	ExistingChannel->GetData().UpdateOrAddKey(InTime, FMovieSceneFloatValue(InValue));
+	EMovieSceneKeyInterpolation Interpolation = GetInterpolationMode(ExistingChannel, InTime, DefaultInterpolation);
+	AddKeyToChannel(ExistingChannel, InTime, InValue, Interpolation);
 
 	if (TryModify())
 	{
@@ -198,7 +199,7 @@ void UMovieSceneComponentMaterialParameterSection::AddScalarParameterKey(const F
 	}
 }
 
-void UMovieSceneComponentMaterialParameterSection::AddColorParameterKey(const FMaterialParameterInfo& InParameterInfo, FFrameNumber InTime, FLinearColor InValue, const FString& InLayerName, const FString& InAssetName, const FParameterChannelNames& InChannelNames)
+void UMovieSceneComponentMaterialParameterSection::AddColorParameterKey(const FMaterialParameterInfo& InParameterInfo, FFrameNumber InTime, FLinearColor InValue, const FString& InLayerName, const FString& InAssetName, const FParameterChannelNames& InChannelNames, EMovieSceneKeyInterpolation DefaultInterpolation)
 {
 	FColorMaterialParameterInfoAndCurves* ExistingCurves = nullptr;
 	for (FColorMaterialParameterInfoAndCurves& ColorParameterInfoAndCurve : ColorParameterInfosAndCurves)
@@ -221,10 +222,19 @@ void UMovieSceneComponentMaterialParameterSection::AddColorParameterKey(const FM
 		CacheChannelProxy();
 	}
 
-	ExistingCurves->RedCurve.GetData().UpdateOrAddKey(InTime, FMovieSceneFloatValue(InValue.R));
-	ExistingCurves->GreenCurve.GetData().UpdateOrAddKey(InTime, FMovieSceneFloatValue(InValue.G));
-	ExistingCurves->BlueCurve.GetData().UpdateOrAddKey(InTime, FMovieSceneFloatValue(InValue.B));
-	ExistingCurves->AlphaCurve.GetData().UpdateOrAddKey(InTime, FMovieSceneFloatValue(InValue.A));
+	EMovieSceneKeyInterpolation RedInterpolation = GetInterpolationMode(&ExistingCurves->RedCurve, InTime, DefaultInterpolation);
+	AddKeyToChannel(&ExistingCurves->RedCurve, InTime, InValue.R, RedInterpolation);
+
+	EMovieSceneKeyInterpolation GreenInterpolation = GetInterpolationMode(&ExistingCurves->GreenCurve, InTime, DefaultInterpolation);
+	AddKeyToChannel(&ExistingCurves->GreenCurve, InTime, InValue.G, GreenInterpolation);
+
+	EMovieSceneKeyInterpolation BlueInterpolation = GetInterpolationMode(&ExistingCurves->BlueCurve, InTime, DefaultInterpolation);
+	AddKeyToChannel(&ExistingCurves->BlueCurve, InTime, InValue.B, BlueInterpolation);
+
+	EMovieSceneKeyInterpolation AlphaInterpolation = GetInterpolationMode(&ExistingCurves->AlphaCurve, InTime, DefaultInterpolation);
+	AddKeyToChannel(&ExistingCurves->AlphaCurve, InTime, InValue.A, AlphaInterpolation);
+
+
 
 	if (TryModify())
 	{
diff --git a/Engine/Source/Runtime/MovieSceneTracks/Private/Sections/MovieSceneParameterSection.cpp b/Engine/Source/Runtime/MovieSceneTracks/Private/Sections/MovieSceneParameterSection.cpp
--- a/Engine/Source/Runtime/MovieSceneTracks/Private/Sections/MovieSceneParameterSection.cpp
+++ b/Engine/Source/Runtime/MovieSceneTracks/Private/Sections/MovieSceneParameterSection.cpp
@@ -225,7 +225,7 @@ void UMovieSceneParameterSection::ReconstructChannelProxy()
 	ChannelProxy = MakeShared<FMovieSceneChannelProxy>(MoveTemp(Channels));
 }
 
-void UMovieSceneParameterSection::AddScalarParameterKey( FName InParameterName, FFrameNumber InTime, float InValue )
+void UMovieSceneParameterSection::AddScalarParameterKey( FName InParameterName, FFrameNumber InTime, float InValue, EMovieSceneKeyInterpolation DefaultInterpolation)
 {
 	FMovieSceneFloatChannel* ExistingChannel = nullptr;
 	for ( FScalarParameterNameAndCurve& ScalarParameterNameAndCurve : ScalarParameterNamesAndCurves )
@@ -244,7 +244,8 @@ void UMovieSceneParameterSection::AddScalarParameterKey( FName InParameterName,
 		ReconstructChannelProxy();
 	}
 
-	ExistingChannel->GetData().UpdateOrAddKey(InTime, FMovieSceneFloatValue(InValue));
+	EMovieSceneKeyInterpolation Interpolation = GetInterpolationMode(ExistingChannel, InTime, DefaultInterpolation);
+	AddKeyToChannel(ExistingChannel, InTime, InValue, Interpolation);
 
 	if (TryModify())
 	{
@@ -280,7 +281,7 @@ void UMovieSceneParameterSection::AddBoolParameterKey(FName InParameterName, FFr
 }
 
 
-void UMovieSceneParameterSection::AddVector2DParameterKey(FName InParameterName, FFrameNumber InTime, FVector2D InValue)
+void UMovieSceneParameterSection::AddVector2DParameterKey(FName InParameterName, FFrameNumber InTime, FVector2D InValue, EMovieSceneKeyInterpolation DefaultInterpolation)
 {
 	FVector2DParameterNameAndCurves* ExistingCurves = nullptr;
 	for (FVector2DParameterNameAndCurves& VectorParameterNameAndCurve : Vector2DParameterNamesAndCurves)
@@ -299,8 +300,10 @@ void UMovieSceneParameterSection::AddVector2DParameterKey(FName InParameterName,
 		ReconstructChannelProxy();
 	}
 
-	ExistingCurves->XCurve.GetData().UpdateOrAddKey(InTime, FMovieSceneFloatValue(InValue.X));
-	ExistingCurves->YCurve.GetData().UpdateOrAddKey(InTime, FMovieSceneFloatValue(InValue.Y));
+	EMovieSceneKeyInterpolation XInterpolation = GetInterpolationMode(&ExistingCurves->XCurve, InTime, DefaultInterpolation);
+	AddKeyToChannel(&ExistingCurves->XCurve, InTime, InValue.X, XInterpolation);
+	EMovieSceneKeyInterpolation YInterpolation = GetInterpolationMode(&ExistingCurves->YCurve, InTime, DefaultInterpolation);
+	AddKeyToChannel(&ExistingCurves->YCurve, InTime, InValue.Y, YInterpolation);
 
 	if (TryModify())
 	{
@@ -308,7 +311,7 @@ void UMovieSceneParameterSection::AddVector2DParameterKey(FName InParameterName,
 	}
 }
 
-void UMovieSceneParameterSection::AddVectorParameterKey( FName InParameterName, FFrameNumber InTime, FVector InValue )
+void UMovieSceneParameterSection::AddVectorParameterKey( FName InParameterName, FFrameNumber InTime, FVector InValue, EMovieSceneKeyInterpolation DefaultInterpolation)
 {
 	FVectorParameterNameAndCurves* ExistingCurves = nullptr;
 	for ( FVectorParameterNameAndCurves& VectorParameterNameAndCurve : VectorParameterNamesAndCurves )
@@ -327,9 +330,12 @@ void UMovieSceneParameterSection::AddVectorParameterKey( FName InParameterName,
 		ReconstructChannelProxy();
 	}
 
-	ExistingCurves->XCurve.GetData().UpdateOrAddKey(InTime, FMovieSceneFloatValue(InValue.X));
-	ExistingCurves->YCurve.GetData().UpdateOrAddKey(InTime, FMovieSceneFloatValue(InValue.Y));
-	ExistingCurves->ZCurve.GetData().UpdateOrAddKey(InTime, FMovieSceneFloatValue(InValue.Z));
+	EMovieSceneKeyInterpolation XInterpolation = GetInterpolationMode(&ExistingCurves->XCurve, InTime, DefaultInterpolation);
+	AddKeyToChannel(&ExistingCurves->XCurve, InTime, InValue.X, XInterpolation);
+	EMovieSceneKeyInterpolation YInterpolation = GetInterpolationMode(&ExistingCurves->YCurve, InTime, DefaultInterpolation);
+	AddKeyToChannel(&ExistingCurves->YCurve, InTime, InValue.Y, YInterpolation);
+	EMovieSceneKeyInterpolation ZInterpolation = GetInterpolationMode(&ExistingCurves->ZCurve, InTime, DefaultInterpolation);
+	AddKeyToChannel(&ExistingCurves->ZCurve, InTime, InValue.Y, ZInterpolation);
 
 	if (TryModify())
 	{
@@ -337,7 +343,7 @@ void UMovieSceneParameterSection::AddVectorParameterKey( FName InParameterName,
 	}
 }
 
-void UMovieSceneParameterSection::AddColorParameterKey( FName InParameterName, FFrameNumber InTime, FLinearColor InValue )
+void UMovieSceneParameterSection::AddColorParameterKey( FName InParameterName, FFrameNumber InTime, FLinearColor InValue, EMovieSceneKeyInterpolation DefaultInterpolation)
 {
 	FColorParameterNameAndCurves* ExistingCurves = nullptr;
 	for ( FColorParameterNameAndCurves& ColorParameterNameAndCurve : ColorParameterNamesAndCurves )
@@ -356,10 +362,14 @@ void UMovieSceneParameterSection::AddColorParameterKey( FName InParameterName, F
 		ReconstructChannelProxy();
 	}
 	
-	ExistingCurves->RedCurve.GetData().UpdateOrAddKey(InTime, FMovieSceneFloatValue(InValue.R));
-	ExistingCurves->GreenCurve.GetData().UpdateOrAddKey(InTime, FMovieSceneFloatValue(InValue.G));
-	ExistingCurves->BlueCurve.GetData().UpdateOrAddKey(InTime, FMovieSceneFloatValue(InValue.B));
-	ExistingCurves->AlphaCurve.GetData().UpdateOrAddKey(InTime, FMovieSceneFloatValue(InValue.A));
+	EMovieSceneKeyInterpolation RedInterpolation = GetInterpolationMode(&ExistingCurves->RedCurve, InTime, DefaultInterpolation);
+	AddKeyToChannel(&ExistingCurves->RedCurve, InTime, InValue.R, RedInterpolation);
+	EMovieSceneKeyInterpolation GreenInterpolation = GetInterpolationMode(&ExistingCurves->GreenCurve, InTime, DefaultInterpolation);
+	AddKeyToChannel(&ExistingCurves->GreenCurve, InTime, InValue.G, RedInterpolation);
+	EMovieSceneKeyInterpolation BlueInterpolation = GetInterpolationMode(&ExistingCurves->BlueCurve, InTime, DefaultInterpolation);
+	AddKeyToChannel(&ExistingCurves->BlueCurve, InTime, InValue.B, RedInterpolation);
+	EMovieSceneKeyInterpolation AlphaInterpolation = GetInterpolationMode(&ExistingCurves->AlphaCurve, InTime, DefaultInterpolation);
+	AddKeyToChannel(&ExistingCurves->AlphaCurve, InTime, InValue.A, RedInterpolation);
 
 	if (TryModify())
 	{
@@ -367,7 +377,7 @@ void UMovieSceneParameterSection::AddColorParameterKey( FName InParameterName, F
 	}
 }
 
-void UMovieSceneParameterSection::AddTransformParameterKey(FName InParameterName, FFrameNumber InTime, const FTransform& InValue)
+void UMovieSceneParameterSection::AddTransformParameterKey(FName InParameterName, FFrameNumber InTime, const FTransform& InValue, EMovieSceneKeyInterpolation DefaultInterpolation)
 {
 	FTransformParameterNameAndCurves* ExistingCurves = nullptr;
 	for (FTransformParameterNameAndCurves& TransformParameterNamesAndCurve : TransformParameterNamesAndCurves)
@@ -389,6 +399,8 @@ void UMovieSceneParameterSection::AddTransformParameterKey(FName InParameterName
 	FRotator Rotator = InValue.GetRotation().Rotator();
 	FVector Scale = InValue.GetScale3D();
 
+
+	////todo
 	ExistingCurves->Translation[0].GetData().UpdateOrAddKey(InTime, FMovieSceneFloatValue(Translation[0]));
 	ExistingCurves->Translation[1].GetData().UpdateOrAddKey(InTime, FMovieSceneFloatValue(Translation[1]));
 	ExistingCurves->Translation[2].GetData().UpdateOrAddKey(InTime, FMovieSceneFloatValue(Translation[2]));
diff --git a/Engine/Source/Runtime/MovieSceneTracks/Private/Tracks/MovieSceneCustomPrimitiveDataTrack.cpp b/Engine/Source/Runtime/MovieSceneTracks/Private/Tracks/MovieSceneCustomPrimitiveDataTrack.cpp
--- a/Engine/Source/Runtime/MovieSceneTracks/Private/Tracks/MovieSceneCustomPrimitiveDataTrack.cpp
+++ b/Engine/Source/Runtime/MovieSceneTracks/Private/Tracks/MovieSceneCustomPrimitiveDataTrack.cpp
@@ -128,12 +128,12 @@ const TArray<UMovieSceneSection*>& UMovieSceneCustomPrimitiveDataTrack::GetAllSe
 }
 
 
-void UMovieSceneCustomPrimitiveDataTrack::AddScalarParameterKey(uint8 CustomPrimitiveDataStartIndex, FFrameNumber Time, float Value)
+void UMovieSceneCustomPrimitiveDataTrack::AddScalarParameterKey(uint8 CustomPrimitiveDataStartIndex, FFrameNumber Time, float Value, EMovieSceneKeyInterpolation DefaultInterpolation)
 {
-	AddScalarParameterKey(CustomPrimitiveDataStartIndex, Time, INDEX_NONE, Value);
+	AddScalarParameterKey(CustomPrimitiveDataStartIndex, Time, INDEX_NONE, Value, DefaultInterpolation);
 }
 
-void UMovieSceneCustomPrimitiveDataTrack::AddScalarParameterKey(uint8 CustomPrimitiveDataStartIndex, FFrameNumber Time, int32 RowIndex, float Value)
+void UMovieSceneCustomPrimitiveDataTrack::AddScalarParameterKey(uint8 CustomPrimitiveDataStartIndex, FFrameNumber Time, int32 RowIndex, float Value, EMovieSceneKeyInterpolation DefaultInterpolation)
 {
 	UMovieSceneSection* NearestSection = SectionToKey;
 	if (NearestSection == nullptr || (RowIndex != INDEX_NONE && NearestSection->GetRowIndex() != RowIndex))
@@ -154,17 +154,17 @@ void UMovieSceneCustomPrimitiveDataTrack::AddScalarParameterKey(uint8 CustomPrim
 	{
 		if (UMovieSceneCustomPrimitiveDataSection* NearestCustomPrimitiveDataSection = Cast<UMovieSceneCustomPrimitiveDataSection>(NearestSection))
 		{
-			NearestCustomPrimitiveDataSection->AddScalarParameterKey(*FString::FromInt(CustomPrimitiveDataStartIndex), Time, Value);
+			NearestCustomPrimitiveDataSection->AddScalarParameterKey(*FString::FromInt(CustomPrimitiveDataStartIndex), Time, Value, DefaultInterpolation);
 		}
 	}
 }
 
-void UMovieSceneCustomPrimitiveDataTrack::AddVector2DParameterKey(uint8 CustomPrimitiveDataStartIndex, FFrameNumber Time, FVector2D Value)
+void UMovieSceneCustomPrimitiveDataTrack::AddVector2DParameterKey(uint8 CustomPrimitiveDataStartIndex, FFrameNumber Time, FVector2D Value, EMovieSceneKeyInterpolation DefaultInterpolation)
 {
-	AddVector2DParameterKey(CustomPrimitiveDataStartIndex, Time, INDEX_NONE, Value);
+	AddVector2DParameterKey(CustomPrimitiveDataStartIndex, Time, INDEX_NONE, Value, DefaultInterpolation);
 }
 
-void UMovieSceneCustomPrimitiveDataTrack::AddVector2DParameterKey(uint8 CustomPrimitiveDataStartIndex, FFrameNumber Time, int32 RowIndex, FVector2D Value)
+void UMovieSceneCustomPrimitiveDataTrack::AddVector2DParameterKey(uint8 CustomPrimitiveDataStartIndex, FFrameNumber Time, int32 RowIndex, FVector2D Value, EMovieSceneKeyInterpolation DefaultInterpolation)
 {
 	UMovieSceneSection* NearestSection = SectionToKey;
 	if (NearestSection == nullptr || (RowIndex != INDEX_NONE && NearestSection->GetRowIndex() != RowIndex))
@@ -185,17 +185,17 @@ void UMovieSceneCustomPrimitiveDataTrack::AddVector2DParameterKey(uint8 CustomPr
 	{
 		if (UMovieSceneCustomPrimitiveDataSection* NearestCustomPrimitiveDataSection = Cast<UMovieSceneCustomPrimitiveDataSection>(NearestSection))
 		{
-			NearestCustomPrimitiveDataSection->AddVector2DParameterKey(*FString::FromInt(CustomPrimitiveDataStartIndex), Time, Value);
+			NearestCustomPrimitiveDataSection->AddVector2DParameterKey(*FString::FromInt(CustomPrimitiveDataStartIndex), Time, Value, DefaultInterpolation);
 		}
 	}
 }
 
-void UMovieSceneCustomPrimitiveDataTrack::AddVectorParameterKey(uint8 CustomPrimitiveDataStartIndex, FFrameNumber Time, FVector Value)
+void UMovieSceneCustomPrimitiveDataTrack::AddVectorParameterKey(uint8 CustomPrimitiveDataStartIndex, FFrameNumber Time, FVector Value, EMovieSceneKeyInterpolation DefaultInterpolation)
 {
-	AddVectorParameterKey(CustomPrimitiveDataStartIndex, Time, INDEX_NONE, Value);
+	AddVectorParameterKey(CustomPrimitiveDataStartIndex, Time, INDEX_NONE, Value, DefaultInterpolation);
 }
 
-void UMovieSceneCustomPrimitiveDataTrack::AddVectorParameterKey(uint8 CustomPrimitiveDataStartIndex, FFrameNumber Time, int32 RowIndex, FVector Value)
+void UMovieSceneCustomPrimitiveDataTrack::AddVectorParameterKey(uint8 CustomPrimitiveDataStartIndex, FFrameNumber Time, int32 RowIndex, FVector Value, EMovieSceneKeyInterpolation DefaultInterpolation)
 {
 	UMovieSceneSection* NearestSection = SectionToKey;
 	if (NearestSection == nullptr || (RowIndex != INDEX_NONE && NearestSection->GetRowIndex() != RowIndex))
@@ -216,17 +216,17 @@ void UMovieSceneCustomPrimitiveDataTrack::AddVectorParameterKey(uint8 CustomPrim
 	{
 		if (UMovieSceneCustomPrimitiveDataSection* NearestCustomPrimitiveDataSection = Cast<UMovieSceneCustomPrimitiveDataSection>(NearestSection))
 		{
-			NearestCustomPrimitiveDataSection->AddVectorParameterKey(*FString::FromInt(CustomPrimitiveDataStartIndex), Time, Value);
+			NearestCustomPrimitiveDataSection->AddVectorParameterKey(*FString::FromInt(CustomPrimitiveDataStartIndex), Time, Value, DefaultInterpolation);
 		}
 	}
 }
 
-void UMovieSceneCustomPrimitiveDataTrack::AddColorParameterKey(uint8 CustomPrimitiveDataStartIndex, FFrameNumber Time, FLinearColor Value)
+void UMovieSceneCustomPrimitiveDataTrack::AddColorParameterKey(uint8 CustomPrimitiveDataStartIndex, FFrameNumber Time, FLinearColor Value, EMovieSceneKeyInterpolation DefaultInterpolation)
 {
-	AddColorParameterKey(CustomPrimitiveDataStartIndex, Time, INDEX_NONE, Value);
+	AddColorParameterKey(CustomPrimitiveDataStartIndex, Time, INDEX_NONE, Value, DefaultInterpolation);
 }
 
-void UMovieSceneCustomPrimitiveDataTrack::AddColorParameterKey(uint8 CustomPrimitiveDataStartIndex, FFrameNumber Time, int32 RowIndex, FLinearColor Value)
+void UMovieSceneCustomPrimitiveDataTrack::AddColorParameterKey(uint8 CustomPrimitiveDataStartIndex, FFrameNumber Time, int32 RowIndex, FLinearColor Value, EMovieSceneKeyInterpolation DefaultInterpolation)
 {
 	UMovieSceneSection* NearestSection = SectionToKey;
 	if (NearestSection == nullptr || (RowIndex != INDEX_NONE && NearestSection->GetRowIndex() != RowIndex))
@@ -247,7 +247,7 @@ void UMovieSceneCustomPrimitiveDataTrack::AddColorParameterKey(uint8 CustomPrimi
 	{
 		if (UMovieSceneCustomPrimitiveDataSection* NearestCustomPrimitiveDataSection = Cast<UMovieSceneCustomPrimitiveDataSection>(NearestSection))
 		{
-			NearestCustomPrimitiveDataSection->AddColorParameterKey(*FString::FromInt(CustomPrimitiveDataStartIndex), Time, Value);
+			NearestCustomPrimitiveDataSection->AddColorParameterKey(*FString::FromInt(CustomPrimitiveDataStartIndex), Time, Value, DefaultInterpolation);
 		}
 	}
 }
diff --git a/Engine/Source/Runtime/MovieSceneTracks/Private/Tracks/MovieSceneMaterialTrack.cpp b/Engine/Source/Runtime/MovieSceneTracks/Private/Tracks/MovieSceneMaterialTrack.cpp
--- a/Engine/Source/Runtime/MovieSceneTracks/Private/Tracks/MovieSceneMaterialTrack.cpp
+++ b/Engine/Source/Runtime/MovieSceneTracks/Private/Tracks/MovieSceneMaterialTrack.cpp
@@ -127,22 +127,22 @@ const TArray<UMovieSceneSection*>& UMovieSceneMaterialTrack::GetAllSections() co
 }
 
 
-void UMovieSceneMaterialTrack::AddScalarParameterKey(FName ParameterName, FFrameNumber Time, float Value)
+void UMovieSceneMaterialTrack::AddScalarParameterKey(FName ParameterName, FFrameNumber Time, float Value, EMovieSceneKeyInterpolation DefaultInterpolation)
 {
-	AddScalarParameterKey(FMaterialParameterInfo(ParameterName), Time, INDEX_NONE, Value, FString(), FString());
+	AddScalarParameterKey(FMaterialParameterInfo(ParameterName), Time, INDEX_NONE, Value, FString(), FString(), DefaultInterpolation);
 }
 
-void UMovieSceneMaterialTrack::AddScalarParameterKey(FName ParameterName, FFrameNumber Time, int32 RowIndex, float Value)
+void UMovieSceneMaterialTrack::AddScalarParameterKey(FName ParameterName, FFrameNumber Time, int32 RowIndex, float Value, EMovieSceneKeyInterpolation DefaultInterpolation)
 {
-	AddScalarParameterKey(FMaterialParameterInfo(ParameterName), Time, RowIndex, Value, FString(), FString());
+	AddScalarParameterKey(FMaterialParameterInfo(ParameterName), Time, RowIndex, Value, FString(), FString(), DefaultInterpolation);
 }
 
-void UMovieSceneMaterialTrack::AddScalarParameterKey(const FMaterialParameterInfo& ParameterInfo, FFrameNumber Time, float Value, const FString& InLayerName, const FString& InAssetName)
+void UMovieSceneMaterialTrack::AddScalarParameterKey(const FMaterialParameterInfo& ParameterInfo, FFrameNumber Time, float Value, const FString& InLayerName, const FString& InAssetName, EMovieSceneKeyInterpolation DefaultInterpolation)
 {
-	AddScalarParameterKey(ParameterInfo, Time, INDEX_NONE, Value, InLayerName, InAssetName);
+	AddScalarParameterKey(ParameterInfo, Time, INDEX_NONE, Value, InLayerName, InAssetName, DefaultInterpolation);
 }
 
-void UMovieSceneMaterialTrack::AddScalarParameterKey(const FMaterialParameterInfo& ParameterInfo, FFrameNumber Time, int32 RowIndex, float Value, const FString& InLayerName, const FString& InAssetName)
+void UMovieSceneMaterialTrack::AddScalarParameterKey(const FMaterialParameterInfo& ParameterInfo, FFrameNumber Time, int32 RowIndex, float Value, const FString& InLayerName, const FString& InAssetName, EMovieSceneKeyInterpolation DefaultInterpolation)
 {
 	UMovieSceneSection* NearestSection = SectionToKey;
 	if (NearestSection == nullptr || (RowIndex != INDEX_NONE && NearestSection->GetRowIndex() != RowIndex))
@@ -164,39 +164,39 @@ void UMovieSceneMaterialTrack::AddScalarParameterKey(const FMaterialParameterInf
 		// If we have an old parameter section already, fall back to the old style section
 		if (UMovieSceneParameterSection* NearestParameterSection = Cast<UMovieSceneParameterSection>(NearestSection))
 		{
-			NearestParameterSection->AddScalarParameterKey(ParameterInfo.Name, Time, Value);
+			NearestParameterSection->AddScalarParameterKey(ParameterInfo.Name, Time, Value, DefaultInterpolation);
 		}
 		else if (UMovieSceneComponentMaterialParameterSection* NearestComponentMaterialParameterSection = Cast<UMovieSceneComponentMaterialParameterSection>(NearestSection))
 		{
-			NearestComponentMaterialParameterSection->AddScalarParameterKey(ParameterInfo, Time, Value, InLayerName, InAssetName);
+			NearestComponentMaterialParameterSection->AddScalarParameterKey(ParameterInfo, Time, Value, InLayerName, InAssetName, DefaultInterpolation);
 		}
 	}
 }
 
 
-void UMovieSceneMaterialTrack::AddColorParameterKey(FName ParameterName, FFrameNumber Time, FLinearColor Value)
+void UMovieSceneMaterialTrack::AddColorParameterKey(FName ParameterName, FFrameNumber Time, FLinearColor Value, EMovieSceneKeyInterpolation DefaultInterpolation)
 {
-	AddColorParameterKey(FMaterialParameterInfo(ParameterName), Time, INDEX_NONE, Value, FString(), FString());
+	AddColorParameterKey(FMaterialParameterInfo(ParameterName), Time, INDEX_NONE, Value, FString(), FString(), DefaultInterpolation);
 }
 
 
-void UMovieSceneMaterialTrack::AddColorParameterKey(FName ParameterName, FFrameNumber Time, int32 RowIndex, FLinearColor Value)
+void UMovieSceneMaterialTrack::AddColorParameterKey(FName ParameterName, FFrameNumber Time, int32 RowIndex, FLinearColor Value, EMovieSceneKeyInterpolation DefaultInterpolation)
 {
-	AddColorParameterKey(FMaterialParameterInfo(ParameterName), Time, RowIndex, Value, FString(), FString());
+	AddColorParameterKey(FMaterialParameterInfo(ParameterName), Time, RowIndex, Value, FString(), FString(), DefaultInterpolation);
 }
 
-void UMovieSceneMaterialTrack::AddColorParameterKey(const FMaterialParameterInfo& ParameterInfo, FFrameNumber Time, FLinearColor Value, const FString& InLayerName, const FString& InAssetName)
+void UMovieSceneMaterialTrack::AddColorParameterKey(const FMaterialParameterInfo& ParameterInfo, FFrameNumber Time, FLinearColor Value, const FString& InLayerName, const FString& InAssetName, EMovieSceneKeyInterpolation DefaultInterpolation)
 {
-	AddColorParameterKey(ParameterInfo, Time, INDEX_NONE, Value, InLayerName, InAssetName);
+	AddColorParameterKey(ParameterInfo, Time, INDEX_NONE, Value, InLayerName, InAssetName, DefaultInterpolation);
 }
 
 
-void UMovieSceneMaterialTrack::AddColorParameterKey(const FMaterialParameterInfo& ParameterInfo, FFrameNumber Time, int32 RowIndex, FLinearColor Value, const FString& InLayerName, const FString& InAssetName)
+void UMovieSceneMaterialTrack::AddColorParameterKey(const FMaterialParameterInfo& ParameterInfo, FFrameNumber Time, int32 RowIndex, FLinearColor Value, const FString& InLayerName, const FString& InAssetName, EMovieSceneKeyInterpolation DefaultInterpolation)
 {
-	AddColorParameterKey(ParameterInfo, Time, RowIndex, Value, InLayerName, InAssetName, FParameterChannelNames());
+	AddColorParameterKey(ParameterInfo, Time, RowIndex, Value, InLayerName, InAssetName, FParameterChannelNames(), DefaultInterpolation);
 }
 
-void UMovieSceneMaterialTrack::AddColorParameterKey(const FMaterialParameterInfo& ParameterInfo, FFrameNumber Time, int32 RowIndex, FLinearColor Value, const FString& InLayerName, const FString& InAssetName, const FParameterChannelNames& InChannelNames)
+void UMovieSceneMaterialTrack::AddColorParameterKey(const FMaterialParameterInfo& ParameterInfo, FFrameNumber Time, int32 RowIndex, FLinearColor Value, const FString& InLayerName, const FString& InAssetName, const FParameterChannelNames& InChannelNames, EMovieSceneKeyInterpolation DefaultInterpolation)
 {
 	UMovieSceneSection* NearestSection = SectionToKey;
 	if (NearestSection == nullptr || (RowIndex != INDEX_NONE && NearestSection->GetRowIndex() != RowIndex))
@@ -218,11 +218,11 @@ void UMovieSceneMaterialTrack::AddColorParameterKey(const FMaterialParameterInfo
 		// If we have an old parameter section already, fall back to the old style section
 		if (UMovieSceneParameterSection* NearestParameterSection = Cast<UMovieSceneParameterSection>(NearestSection))
 		{
-			NearestParameterSection->AddColorParameterKey(ParameterInfo.Name, Time, Value);
+			NearestParameterSection->AddColorParameterKey(ParameterInfo.Name, Time, Value, DefaultInterpolation);
 		}
 		else if (UMovieSceneComponentMaterialParameterSection* NearestComponentMaterialParameterSection = Cast<UMovieSceneComponentMaterialParameterSection>(NearestSection))
 		{
-			NearestComponentMaterialParameterSection->AddColorParameterKey(ParameterInfo, Time, Value, InLayerName, InAssetName, InChannelNames);
+			NearestComponentMaterialParameterSection->AddColorParameterKey(ParameterInfo, Time, Value, InLayerName, InAssetName, InChannelNames, DefaultInterpolation);
 		}
 	}
 }
diff --git a/Engine/Source/Runtime/MovieSceneTracks/Private/Tracks/MovieSceneParticleParameterTrack.cpp b/Engine/Source/Runtime/MovieSceneTracks/Private/Tracks/MovieSceneParticleParameterTrack.cpp
--- a/Engine/Source/Runtime/MovieSceneTracks/Private/Tracks/MovieSceneParticleParameterTrack.cpp
+++ b/Engine/Source/Runtime/MovieSceneTracks/Private/Tracks/MovieSceneParticleParameterTrack.cpp
@@ -76,7 +76,7 @@ FText UMovieSceneParticleParameterTrack::GetDefaultDisplayName() const
 #endif
 
 
-void UMovieSceneParticleParameterTrack::AddScalarParameterKey( FName ParameterName, FFrameNumber Time, float Value )
+void UMovieSceneParticleParameterTrack::AddScalarParameterKey( FName ParameterName, FFrameNumber Time, float Value, EMovieSceneKeyInterpolation DefaultInterpolation)
 {
 	UMovieSceneParameterSection* NearestSection = Cast<UMovieSceneParameterSection>(MovieSceneHelpers::FindNearestSectionAtTime(Sections, Time));
 	if ( NearestSection == nullptr )
@@ -90,10 +90,10 @@ void UMovieSceneParticleParameterTrack::AddScalarParameterKey( FName ParameterNa
 
 		Sections.Add( NearestSection );
 	}
-	NearestSection->AddScalarParameterKey(ParameterName, Time, Value);
+	NearestSection->AddScalarParameterKey(ParameterName, Time, Value, DefaultInterpolation);
 }
 
-void UMovieSceneParticleParameterTrack::AddVectorParameterKey( FName ParameterName, FFrameNumber Time, FVector Value )
+void UMovieSceneParticleParameterTrack::AddVectorParameterKey( FName ParameterName, FFrameNumber Time, FVector Value, EMovieSceneKeyInterpolation DefaultInterpolation)
 {
 	UMovieSceneParameterSection* NearestSection = Cast<UMovieSceneParameterSection>( MovieSceneHelpers::FindNearestSectionAtTime( Sections, Time ) );
 	if ( NearestSection == nullptr )
@@ -107,10 +107,10 @@ void UMovieSceneParticleParameterTrack::AddVectorParameterKey( FName ParameterNa
 
 		Sections.Add( NearestSection );
 	}
-	NearestSection->AddVectorParameterKey( ParameterName, Time, Value );
+	NearestSection->AddVectorParameterKey( ParameterName, Time, Value, DefaultInterpolation);
 }
 
-void UMovieSceneParticleParameterTrack::AddColorParameterKey( FName ParameterName, FFrameNumber Time, FLinearColor Value )
+void UMovieSceneParticleParameterTrack::AddColorParameterKey( FName ParameterName, FFrameNumber Time, FLinearColor Value, EMovieSceneKeyInterpolation DefaultInterpolation)
 {
 	UMovieSceneParameterSection* NearestSection = Cast<UMovieSceneParameterSection>( MovieSceneHelpers::FindNearestSectionAtTime( Sections, Time ) );
 	if ( NearestSection == nullptr )
@@ -124,7 +124,7 @@ void UMovieSceneParticleParameterTrack::AddColorParameterKey( FName ParameterNam
 
 		Sections.Add( NearestSection );
 	}
-	NearestSection->AddColorParameterKey( ParameterName, Time, Value );
+	NearestSection->AddColorParameterKey( ParameterName, Time, Value, DefaultInterpolation);
 }
 #undef LOCTEXT_NAMESPACE
 
diff --git a/Engine/Source/Runtime/MovieSceneTracks/Public/Sections/MovieSceneComponentMaterialParameterSection.h b/Engine/Source/Runtime/MovieSceneTracks/Public/Sections/MovieSceneComponentMaterialParameterSection.h
--- a/Engine/Source/Runtime/MovieSceneTracks/Public/Sections/MovieSceneComponentMaterialParameterSection.h
+++ b/Engine/Source/Runtime/MovieSceneTracks/Public/Sections/MovieSceneComponentMaterialParameterSection.h
@@ -107,11 +107,11 @@ class UMovieSceneComponentMaterialParameterSection
 public:
 	/** Adds a a key for a specific scalar parameter at the specified time with the specified value. */
 	UFUNCTION(BlueprintCallable, Category = "Sequencer|Section")
-	MOVIESCENETRACKS_API void AddScalarParameterKey(const FMaterialParameterInfo& InParameterInfo, FFrameNumber InTime, float InValue, const FString& InLayerName, const FString& InAssetName);
+	MOVIESCENETRACKS_API void AddScalarParameterKey(const FMaterialParameterInfo& InParameterInfo, FFrameNumber InTime, float InValue, const FString& InLayerName, const FString& InAssetName, EMovieSceneKeyInterpolation DefaultInterpolation = EMovieSceneKeyInterpolation::Auto);
 
 	/** Adds a a key for a specific color parameter at the specified time with the specified value. */
 	UFUNCTION(BlueprintCallable, Category = "Sequencer|Section")
-	MOVIESCENETRACKS_API void AddColorParameterKey(const FMaterialParameterInfo& InParameterInfo, FFrameNumber InTime, FLinearColor InValue, const FString& InLayerName, const FString& InAssetName, const FParameterChannelNames& InChannelNames=FParameterChannelNames());
+	MOVIESCENETRACKS_API void AddColorParameterKey(const FMaterialParameterInfo& InParameterInfo, FFrameNumber InTime, FLinearColor InValue, const FString& InLayerName, const FString& InAssetName, const FParameterChannelNames& InChannelNames=FParameterChannelNames(), EMovieSceneKeyInterpolation DefaultInterpolation = EMovieSceneKeyInterpolation::Auto);
 
 	/** 
 	 * Removes a scalar parameter from this section. 
diff --git a/Engine/Source/Runtime/MovieSceneTracks/Public/Sections/MovieSceneParameterSection.h b/Engine/Source/Runtime/MovieSceneTracks/Public/Sections/MovieSceneParameterSection.h
--- a/Engine/Source/Runtime/MovieSceneTracks/Public/Sections/MovieSceneParameterSection.h
+++ b/Engine/Source/Runtime/MovieSceneTracks/Public/Sections/MovieSceneParameterSection.h
@@ -343,7 +343,7 @@ class UMovieSceneParameterSection
 public:
 	/** Adds a a key for a specific scalar parameter at the specified time with the specified value. */
 	UFUNCTION(BlueprintCallable, Category = "Sequencer|Section")
-	MOVIESCENETRACKS_API void AddScalarParameterKey(FName InParameterName, FFrameNumber InTime, float InValue);
+	MOVIESCENETRACKS_API void AddScalarParameterKey(FName InParameterName, FFrameNumber InTime, float InValue, EMovieSceneKeyInterpolation DefaultInterpolation = EMovieSceneKeyInterpolation::Auto);
 
 	/** Adds a a key for a specific bool parameter at the specified time with the specified value. */
 	UFUNCTION(BlueprintCallable, Category = "Sequencer|Section")
@@ -351,19 +351,19 @@ public:
 
 	/** Adds a a key for a specific vector2D parameter at the specified time with the specified value. */
 	UFUNCTION(BlueprintCallable, Category = "Sequencer|Section")
-	MOVIESCENETRACKS_API void AddVector2DParameterKey(FName InParameterName, FFrameNumber InTime, FVector2D InValue);
+	MOVIESCENETRACKS_API void AddVector2DParameterKey(FName InParameterName, FFrameNumber InTime, FVector2D InValue, EMovieSceneKeyInterpolation DefaultInterpolation = EMovieSceneKeyInterpolation::Auto);
 
 	/** Adds a a key for a specific vector parameter at the specified time with the specified value. */
 	UFUNCTION(BlueprintCallable, Category = "Sequencer|Section")
-	MOVIESCENETRACKS_API void AddVectorParameterKey(FName InParameterName, FFrameNumber InTime, FVector InValue);
+	MOVIESCENETRACKS_API void AddVectorParameterKey(FName InParameterName, FFrameNumber InTime, FVector InValue, EMovieSceneKeyInterpolation DefaultInterpolation = EMovieSceneKeyInterpolation::Auto);
 
 	/** Adds a a key for a specific color parameter at the specified time with the specified value. */
 	UFUNCTION(BlueprintCallable, Category = "Sequencer|Section")
-	MOVIESCENETRACKS_API void AddColorParameterKey(FName InParameterName, FFrameNumber InTime, FLinearColor InValue);
+	MOVIESCENETRACKS_API void AddColorParameterKey(FName InParameterName, FFrameNumber InTime, FLinearColor InValue, EMovieSceneKeyInterpolation DefaultInterpolation = EMovieSceneKeyInterpolation::Auto);
 
 	/** Adds a a key for a specific color parameter at the specified time with the specified value. */
 	UFUNCTION(BlueprintCallable, Category = "Sequencer|Section")
-	MOVIESCENETRACKS_API void AddTransformParameterKey(FName InParameterName, FFrameNumber InTime, const FTransform& InValue);
+	MOVIESCENETRACKS_API void AddTransformParameterKey(FName InParameterName, FFrameNumber InTime, const FTransform& InValue, EMovieSceneKeyInterpolation DefaultInterpolation = EMovieSceneKeyInterpolation::Auto);
 
 	/** 
 	 * Removes a scalar parameter from this section. 
diff --git a/Engine/Source/Runtime/MovieSceneTracks/Public/Tracks/MovieSceneCustomPrimitiveDataTrack.h b/Engine/Source/Runtime/MovieSceneTracks/Public/Tracks/MovieSceneCustomPrimitiveDataTrack.h
--- a/Engine/Source/Runtime/MovieSceneTracks/Public/Tracks/MovieSceneCustomPrimitiveDataTrack.h
+++ b/Engine/Source/Runtime/MovieSceneTracks/Public/Tracks/MovieSceneCustomPrimitiveDataTrack.h
@@ -65,7 +65,7 @@ public:
 	* @param Time The time to add the new key.
 	* @param The value for the new key.
 	*/
-	MOVIESCENETRACKS_API void AddScalarParameterKey(uint8 CustomPrimitiveDataStartIndex, FFrameNumber Position, float Value);
+	MOVIESCENETRACKS_API void AddScalarParameterKey(uint8 CustomPrimitiveDataStartIndex, FFrameNumber Position, float Value, EMovieSceneKeyInterpolation DefaultInterpolation = EMovieSceneKeyInterpolation::Auto);
 
 	/**
 	 * Adds a scalar parameter key to the track. 
@@ -74,7 +74,7 @@ public:
 	 * @param RowIndex The preferred row index on which to look for sections.
 	 * @param The value for the new key.
 	 */
-	MOVIESCENETRACKS_API void AddScalarParameterKey(uint8 CustomPrimitiveDataStartIndex, FFrameNumber Position, int32 RowIndex, float Value);
+	MOVIESCENETRACKS_API void AddScalarParameterKey(uint8 CustomPrimitiveDataStartIndex, FFrameNumber Position, int32 RowIndex, float Value, EMovieSceneKeyInterpolation DefaultInterpolation = EMovieSceneKeyInterpolation::Auto);
 	
 	/**
 	 * Adds a Vector2D parameter key to the track.
@@ -82,7 +82,7 @@ public:
 	 * @param Time The time to add the new key.
 	 * @param The value for the new key.
 	 */
-	MOVIESCENETRACKS_API void AddVector2DParameterKey(uint8 CustomPrimitiveDataStartIndex, FFrameNumber Position, FVector2D Value);
+	MOVIESCENETRACKS_API void AddVector2DParameterKey(uint8 CustomPrimitiveDataStartIndex, FFrameNumber Position, FVector2D Value, EMovieSceneKeyInterpolation DefaultInterpolation = EMovieSceneKeyInterpolation::Auto);
 
 	/**
 	 * Adds a Vector2D parameter key to the track.
@@ -91,7 +91,7 @@ public:
 	 * @param RowIndex The preferred row index on which to look for sections.
 	 * @param The value for the new key.
 	 */
-	MOVIESCENETRACKS_API void AddVector2DParameterKey(uint8 CustomPrimitiveDataStartIndex, FFrameNumber Position, int32 RowIndex, FVector2D Value);
+	MOVIESCENETRACKS_API void AddVector2DParameterKey(uint8 CustomPrimitiveDataStartIndex, FFrameNumber Position, int32 RowIndex, FVector2D Value, EMovieSceneKeyInterpolation DefaultInterpolation = EMovieSceneKeyInterpolation::Auto);
 
 	/**
 	 * Adds a Vector parameter key to the track.
@@ -99,7 +99,7 @@ public:
 	 * @param Time The time to add the new key.
 	 * @param The value for the new key.
 	 */
-	MOVIESCENETRACKS_API void AddVectorParameterKey(uint8 CustomPrimitiveDataStartIndex, FFrameNumber Position, FVector Value);
+	MOVIESCENETRACKS_API void AddVectorParameterKey(uint8 CustomPrimitiveDataStartIndex, FFrameNumber Position, FVector Value, EMovieSceneKeyInterpolation DefaultInterpolation = EMovieSceneKeyInterpolation::Auto);
 
 	/**
 	 * Adds a Vector parameter key to the track.
@@ -108,7 +108,7 @@ public:
 	 * @param RowIndex The preferred row index on which to look for sections.
 	 * @param The value for the new key.
 	 */
-	MOVIESCENETRACKS_API void AddVectorParameterKey(uint8 CustomPrimitiveDataStartIndex, FFrameNumber Position, int32 RowIndex, FVector Value);
+	MOVIESCENETRACKS_API void AddVectorParameterKey(uint8 CustomPrimitiveDataStartIndex, FFrameNumber Position, int32 RowIndex, FVector Value, EMovieSceneKeyInterpolation DefaultInterpolation = EMovieSceneKeyInterpolation::Auto);
 
 	/**
 	* Adds a Vector2D parameter key to the track.
@@ -116,7 +116,7 @@ public:
 	* @param Time The time to add the new key.
 	* @param The value for the new key.
 	*/
-	MOVIESCENETRACKS_API void AddColorParameterKey(uint8 CustomPrimitiveDataStartIndex, FFrameNumber Position, FLinearColor Value);
+	MOVIESCENETRACKS_API void AddColorParameterKey(uint8 CustomPrimitiveDataStartIndex, FFrameNumber Position, FLinearColor Value, EMovieSceneKeyInterpolation DefaultInterpolation = EMovieSceneKeyInterpolation::Auto);
 
 	/**
 	* Adds a Vector2D parameter key to the track.
@@ -125,7 +125,7 @@ public:
 	* @param RowIndex The preferred row index on which to look for sections.
 	* @param The value for the new key.
 	*/
-	MOVIESCENETRACKS_API void AddColorParameterKey(uint8 CustomPrimitiveDataStartIndex, FFrameNumber Position, int32 RowIndex, FLinearColor Value);
+	MOVIESCENETRACKS_API void AddColorParameterKey(uint8 CustomPrimitiveDataStartIndex, FFrameNumber Position, int32 RowIndex, FLinearColor Value, EMovieSceneKeyInterpolation DefaultInterpolation = EMovieSceneKeyInterpolation::Auto);
 
 	/*~ IMovieSceneEntityProvider */
 	virtual void ImportEntityImpl(UMovieSceneEntitySystemLinker* EntityLinker, const FEntityImportParams& Params, FImportedEntity* OutImportedEntity) override;
diff --git a/Engine/Source/Runtime/MovieSceneTracks/Public/Tracks/MovieSceneMaterialTrack.h b/Engine/Source/Runtime/MovieSceneTracks/Public/Tracks/MovieSceneMaterialTrack.h
--- a/Engine/Source/Runtime/MovieSceneTracks/Public/Tracks/MovieSceneMaterialTrack.h
+++ b/Engine/Source/Runtime/MovieSceneTracks/Public/Tracks/MovieSceneMaterialTrack.h
@@ -115,7 +115,7 @@ public:
 	* @param Time The time to add the new key.
 	* @param The value for the new key.
 	*/
-	MOVIESCENETRACKS_API void AddScalarParameterKey(FName ParameterName, FFrameNumber Position, float Value);
+	MOVIESCENETRACKS_API void AddScalarParameterKey(FName ParameterName, FFrameNumber Position, float Value, EMovieSceneKeyInterpolation DefaultInterpolation = EMovieSceneKeyInterpolation::Auto);
 
 	/**
 	 * Adds a scalar parameter key to the track. 
@@ -124,7 +124,7 @@ public:
 	 * @param RowIndex The preferred row index on which to look for sections.
 	 * @param The value for the new key.
 	 */
-	MOVIESCENETRACKS_API void AddScalarParameterKey(FName ParameterName, FFrameNumber Position, int32 RowIndex, float Value);
+	MOVIESCENETRACKS_API void AddScalarParameterKey(FName ParameterName, FFrameNumber Position, int32 RowIndex, float Value, EMovieSceneKeyInterpolation DefaultInterpolation = EMovieSceneKeyInterpolation::Auto);
 
 	/**
 	* Adds a scalar parameter key to the track.
@@ -134,7 +134,7 @@ public:
 	* @param InLayerName Optional layer name for use in UI.
 	* @param InAssetName Optional asset name for use in UI.
 	*/
-	MOVIESCENETRACKS_API void AddScalarParameterKey(const FMaterialParameterInfo& ParameterInfo, FFrameNumber Position, float Value, const FString& InLayerName, const FString& InAssetName);
+	MOVIESCENETRACKS_API void AddScalarParameterKey(const FMaterialParameterInfo& ParameterInfo, FFrameNumber Position, float Value, const FString& InLayerName, const FString& InAssetName, EMovieSceneKeyInterpolation DefaultInterpolation = EMovieSceneKeyInterpolation::Auto);
 
 	/**
 	* Adds a scalar parameter key to the track.
@@ -145,7 +145,7 @@ public:
 	* @param InLayerName Optional layer name for use in UI.
 	* @param InAssetName Optional asset name for use in UI.
 	*/
-	MOVIESCENETRACKS_API void AddScalarParameterKey(const FMaterialParameterInfo& ParameterInfo, FFrameNumber Position, int32 RowIndex, float Value, const FString& InLayerName, const FString& InAssetName);
+	MOVIESCENETRACKS_API void AddScalarParameterKey(const FMaterialParameterInfo& ParameterInfo, FFrameNumber Position, int32 RowIndex, float Value, const FString& InLayerName, const FString& InAssetName, EMovieSceneKeyInterpolation DefaultInterpolation = EMovieSceneKeyInterpolation::Auto);
 
 	/**
 	* Adds a color parameter key to the track.
@@ -153,7 +153,7 @@ public:
 	* @param Time The time to add the new key.
 	* @param The value for the new key.
 	*/
-	MOVIESCENETRACKS_API void AddColorParameterKey(FName ParameterName, FFrameNumber Position, FLinearColor Value);
+	MOVIESCENETRACKS_API void AddColorParameterKey(FName ParameterName, FFrameNumber Position, FLinearColor Value, EMovieSceneKeyInterpolation DefaultInterpolation = EMovieSceneKeyInterpolation::Auto);
 
 	/**
 	* Adds a color parameter key to the track.
@@ -162,7 +162,7 @@ public:
 	* @param RowIndex The preferred row index on which to look for sections.
 	* @param The value for the new key.
 	*/
-	MOVIESCENETRACKS_API void AddColorParameterKey(FName ParameterName, FFrameNumber Position, int32 RowIndex, FLinearColor Value);
+	MOVIESCENETRACKS_API void AddColorParameterKey(FName ParameterName, FFrameNumber Position, int32 RowIndex, FLinearColor Value, EMovieSceneKeyInterpolation DefaultInterpolation = EMovieSceneKeyInterpolation::Auto);
 
 	/**
 	* Adds a color parameter key to the track.
@@ -172,7 +172,7 @@ public:
 	* @param InLayerName Optional layer name for use in UI.
 	* @param InAssetName Optional asset name for use in UI.
 	*/
-	MOVIESCENETRACKS_API void AddColorParameterKey(const FMaterialParameterInfo& ParameterInfo, FFrameNumber Position, FLinearColor Value, const FString& InLayerName, const FString& InAssetName);
+	MOVIESCENETRACKS_API void AddColorParameterKey(const FMaterialParameterInfo& ParameterInfo, FFrameNumber Position, FLinearColor Value, const FString& InLayerName, const FString& InAssetName, EMovieSceneKeyInterpolation DefaultInterpolation = EMovieSceneKeyInterpolation::Auto);
 
 	/**
 	* Adds a color parameter key to the track.
@@ -183,7 +183,7 @@ public:
 	* @param InLayerName Optional layer name for use in UI.
 	* @param InAssetName Optional asset name for use in UI.
 	*/
-	MOVIESCENETRACKS_API void AddColorParameterKey(const FMaterialParameterInfo& ParameterInfo, FFrameNumber Position, int32 RowIndex, FLinearColor Value, const FString& InLayerName, const FString& InAssetName);
+	MOVIESCENETRACKS_API void AddColorParameterKey(const FMaterialParameterInfo& ParameterInfo, FFrameNumber Position, int32 RowIndex, FLinearColor Value, const FString& InLayerName, const FString& InAssetName, EMovieSceneKeyInterpolation DefaultInterpolation = EMovieSceneKeyInterpolation::Auto);
 
 	/**
 	* Adds a color parameter key to the track.
@@ -195,7 +195,7 @@ public:
 	* @param InAssetName Optional asset name for use in UI.
 	* @param InChannelNames Optional channel names for the 4 vector channels for use in the UI.
 	*/
-	MOVIESCENETRACKS_API void AddColorParameterKey(const FMaterialParameterInfo& ParameterInfo, FFrameNumber Position, int32 RowIndex, FLinearColor Value, const FString& InLayerName, const FString& InAssetName, const FParameterChannelNames& InChannelNames);
+	MOVIESCENETRACKS_API void AddColorParameterKey(const FMaterialParameterInfo& ParameterInfo, FFrameNumber Position, int32 RowIndex, FLinearColor Value, const FString& InLayerName, const FString& InAssetName, const FParameterChannelNames& InChannelNames, EMovieSceneKeyInterpolation DefaultInterpolation = EMovieSceneKeyInterpolation::Auto);
 
 private:
 
diff --git a/Engine/Source/Runtime/MovieSceneTracks/Public/Tracks/MovieSceneParticleParameterTrack.h b/Engine/Source/Runtime/MovieSceneTracks/Public/Tracks/MovieSceneParticleParameterTrack.h
--- a/Engine/Source/Runtime/MovieSceneTracks/Public/Tracks/MovieSceneParticleParameterTrack.h
+++ b/Engine/Source/Runtime/MovieSceneTracks/Public/Tracks/MovieSceneParticleParameterTrack.h
@@ -61,7 +61,7 @@ public:
 	 * @param Time The time to add the new key.
 	 * @param The value for the new key.
 	 */
-	void MOVIESCENETRACKS_API AddScalarParameterKey( FName ParameterName, FFrameNumber Position, float Value );
+	void MOVIESCENETRACKS_API AddScalarParameterKey( FName ParameterName, FFrameNumber Position, float Value, EMovieSceneKeyInterpolation DefaultInterpolation = EMovieSceneKeyInterpolation::Auto);
 
 	/**
 	* Adds a Vector parameter key to the track.
@@ -69,7 +69,7 @@ public:
 	* @param Time The time to add the new key.
 	* @param The value for the new key.
 	*/
-	void MOVIESCENETRACKS_API AddVectorParameterKey( FName ParameterName, FFrameNumber Position, FVector Value );
+	void MOVIESCENETRACKS_API AddVectorParameterKey( FName ParameterName, FFrameNumber Position, FVector Value, EMovieSceneKeyInterpolation DefaultInterpolation = EMovieSceneKeyInterpolation::Auto);
 
 	/**
 	* Adds a Vector parameter key to the track.
@@ -77,7 +77,7 @@ public:
 	* @param Time The time to add the new key.
 	* @param The value for the new key.
 	*/
-	void MOVIESCENETRACKS_API AddColorParameterKey( FName ParameterName, FFrameNumber Position, FLinearColor Value );
+	void MOVIESCENETRACKS_API AddColorParameterKey( FName ParameterName, FFrameNumber Position, FLinearColor Value, EMovieSceneKeyInterpolation DefaultInterpolation = EMovieSceneKeyInterpolation::Auto);
 
 private:
 
