Sequencer: Follow up fixes - docs updates

#rnx

diff --git a/Engine/Source/Runtime/MovieSceneTracks/Private/Sections/MovieSceneParameterSection.cpp b/Engine/Source/Runtime/MovieSceneTracks/Private/Sections/MovieSceneParameterSection.cpp
--- a/Engine/Source/Runtime/MovieSceneTracks/Private/Sections/MovieSceneParameterSection.cpp
+++ b/Engine/Source/Runtime/MovieSceneTracks/Private/Sections/MovieSceneParameterSection.cpp
@@ -335,7 +335,7 @@ void UMovieSceneParameterSection::AddVectorParameterKey( FName InParameterName,
 	EMovieSceneKeyInterpolation YInterpolation = GetInterpolationMode(&ExistingCurves->YCurve, InTime, DefaultInterpolation);
 	AddKeyToChannel(&ExistingCurves->YCurve, InTime, InValue.Y, YInterpolation);
 	EMovieSceneKeyInterpolation ZInterpolation = GetInterpolationMode(&ExistingCurves->ZCurve, InTime, DefaultInterpolation);
-	AddKeyToChannel(&ExistingCurves->ZCurve, InTime, InValue.Y, ZInterpolation);
+	AddKeyToChannel(&ExistingCurves->ZCurve, InTime, InValue.Z, ZInterpolation);
 
 	if (TryModify())
 	{
@@ -365,11 +365,11 @@ void UMovieSceneParameterSection::AddColorParameterKey( FName InParameterName, F
 	EMovieSceneKeyInterpolation RedInterpolation = GetInterpolationMode(&ExistingCurves->RedCurve, InTime, DefaultInterpolation);
 	AddKeyToChannel(&ExistingCurves->RedCurve, InTime, InValue.R, RedInterpolation);
 	EMovieSceneKeyInterpolation GreenInterpolation = GetInterpolationMode(&ExistingCurves->GreenCurve, InTime, DefaultInterpolation);
-	AddKeyToChannel(&ExistingCurves->GreenCurve, InTime, InValue.G, RedInterpolation);
+	AddKeyToChannel(&ExistingCurves->GreenCurve, InTime, InValue.G, GreenInterpolation);
 	EMovieSceneKeyInterpolation BlueInterpolation = GetInterpolationMode(&ExistingCurves->BlueCurve, InTime, DefaultInterpolation);
-	AddKeyToChannel(&ExistingCurves->BlueCurve, InTime, InValue.B, RedInterpolation);
+	AddKeyToChannel(&ExistingCurves->BlueCurve, InTime, InValue.B, BlueInterpolation);
 	EMovieSceneKeyInterpolation AlphaInterpolation = GetInterpolationMode(&ExistingCurves->AlphaCurve, InTime, DefaultInterpolation);
-	AddKeyToChannel(&ExistingCurves->AlphaCurve, InTime, InValue.A, RedInterpolation);
+	AddKeyToChannel(&ExistingCurves->AlphaCurve, InTime, InValue.A, AlphaInterpolation);
 
 	if (TryModify())
 	{
diff --git a/Engine/Source/Runtime/MovieSceneTracks/Public/Tracks/MovieSceneCustomPrimitiveDataTrack.h b/Engine/Source/Runtime/MovieSceneTracks/Public/Tracks/MovieSceneCustomPrimitiveDataTrack.h
--- a/Engine/Source/Runtime/MovieSceneTracks/Public/Tracks/MovieSceneCustomPrimitiveDataTrack.h
+++ b/Engine/Source/Runtime/MovieSceneTracks/Public/Tracks/MovieSceneCustomPrimitiveDataTrack.h
@@ -64,25 +64,28 @@ public:
 	* @param CustomPrimitiveDataStartIndex The start index for the custom primitive data
 	* @param Time The time to add the new key.
 	* @param The value for the new key.
+	* @param DefaultInterpolation The default interpolation to use for thew new key.
 	*/
-	MOVIESCENETRACKS_API void AddScalarParameterKey(uint8 CustomPrimitiveDataStartIndex, FFrameNumber Position, float Value, EMovieSceneKeyInterpolation DefaultInterpolation = EMovieSceneKeyInterpolation::Auto);
+	MOVIESCENETRACKS_API void AddScalarParameterKey(uint8 CustomPrimitiveDataStartIndex, FFrameNumber Time, float Value, EMovieSceneKeyInterpolation DefaultInterpolation = EMovieSceneKeyInterpolation::Auto);
 
 	/**
 	 * Adds a scalar parameter key to the track. 
-	* @param CustomPrimitiveDataStartIndex The start index for the custom primitive data
+	 * @param CustomPrimitiveDataStartIndex The start index for the custom primitive data
 	 * @param Time The time to add the new key.
 	 * @param RowIndex The preferred row index on which to look for sections.
 	 * @param The value for the new key.
+	 * @param DefaultInterpolation The default interpolation to use for thew new key.
 	 */
-	MOVIESCENETRACKS_API void AddScalarParameterKey(uint8 CustomPrimitiveDataStartIndex, FFrameNumber Position, int32 RowIndex, float Value, EMovieSceneKeyInterpolation DefaultInterpolation = EMovieSceneKeyInterpolation::Auto);
+	MOVIESCENETRACKS_API void AddScalarParameterKey(uint8 CustomPrimitiveDataStartIndex, FFrameNumber Time, int32 RowIndex, float Value, EMovieSceneKeyInterpolation DefaultInterpolation = EMovieSceneKeyInterpolation::Auto);
 	
 	/**
 	 * Adds a Vector2D parameter key to the track.
 	 * @param CustomPrimitiveDataStartIndex The start index for the custom primitive data
 	 * @param Time The time to add the new key.
 	 * @param The value for the new key.
+	 * @param DefaultInterpolation The default interpolation to use for thew new key.
 	 */
-	MOVIESCENETRACKS_API void AddVector2DParameterKey(uint8 CustomPrimitiveDataStartIndex, FFrameNumber Position, FVector2D Value, EMovieSceneKeyInterpolation DefaultInterpolation = EMovieSceneKeyInterpolation::Auto);
+	MOVIESCENETRACKS_API void AddVector2DParameterKey(uint8 CustomPrimitiveDataStartIndex, FFrameNumber Time, FVector2D Value, EMovieSceneKeyInterpolation DefaultInterpolation = EMovieSceneKeyInterpolation::Auto);
 
 	/**
 	 * Adds a Vector2D parameter key to the track.
@@ -90,16 +93,18 @@ public:
 	 * @param Time The time to add the new key.
 	 * @param RowIndex The preferred row index on which to look for sections.
 	 * @param The value for the new key.
+	 * @param DefaultInterpolation The default interpolation to use for thew new key.
 	 */
-	MOVIESCENETRACKS_API void AddVector2DParameterKey(uint8 CustomPrimitiveDataStartIndex, FFrameNumber Position, int32 RowIndex, FVector2D Value, EMovieSceneKeyInterpolation DefaultInterpolation = EMovieSceneKeyInterpolation::Auto);
+	MOVIESCENETRACKS_API void AddVector2DParameterKey(uint8 CustomPrimitiveDataStartIndex, FFrameNumber Time, int32 RowIndex, FVector2D Value, EMovieSceneKeyInterpolation DefaultInterpolation = EMovieSceneKeyInterpolation::Auto);
 
 	/**
 	 * Adds a Vector parameter key to the track.
 	 * @param CustomPrimitiveDataStartIndex The start index for the custom primitive data
 	 * @param Time The time to add the new key.
 	 * @param The value for the new key.
+	 * @param DefaultInterpolation The default interpolation to use for thew new key.
 	 */
-	MOVIESCENETRACKS_API void AddVectorParameterKey(uint8 CustomPrimitiveDataStartIndex, FFrameNumber Position, FVector Value, EMovieSceneKeyInterpolation DefaultInterpolation = EMovieSceneKeyInterpolation::Auto);
+	MOVIESCENETRACKS_API void AddVectorParameterKey(uint8 CustomPrimitiveDataStartIndex, FFrameNumber Time, FVector Value, EMovieSceneKeyInterpolation DefaultInterpolation = EMovieSceneKeyInterpolation::Auto);
 
 	/**
 	 * Adds a Vector parameter key to the track.
@@ -107,16 +112,18 @@ public:
 	 * @param Time The time to add the new key.
 	 * @param RowIndex The preferred row index on which to look for sections.
 	 * @param The value for the new key.
+	 * @param DefaultInterpolation The default interpolation to use for thew new key.
 	 */
-	MOVIESCENETRACKS_API void AddVectorParameterKey(uint8 CustomPrimitiveDataStartIndex, FFrameNumber Position, int32 RowIndex, FVector Value, EMovieSceneKeyInterpolation DefaultInterpolation = EMovieSceneKeyInterpolation::Auto);
+	MOVIESCENETRACKS_API void AddVectorParameterKey(uint8 CustomPrimitiveDataStartIndex, FFrameNumber Time, int32 RowIndex, FVector Value, EMovieSceneKeyInterpolation DefaultInterpolation = EMovieSceneKeyInterpolation::Auto);
 
 	/**
 	* Adds a Vector2D parameter key to the track.
 	* @param CustomPrimitiveDataStartIndex The start index for the custom primitive data
 	* @param Time The time to add the new key.
 	* @param The value for the new key.
+	* @param DefaultInterpolation The default interpolation to use for thew new key.
 	*/
-	MOVIESCENETRACKS_API void AddColorParameterKey(uint8 CustomPrimitiveDataStartIndex, FFrameNumber Position, FLinearColor Value, EMovieSceneKeyInterpolation DefaultInterpolation = EMovieSceneKeyInterpolation::Auto);
+	MOVIESCENETRACKS_API void AddColorParameterKey(uint8 CustomPrimitiveDataStartIndex, FFrameNumber Time, FLinearColor Value, EMovieSceneKeyInterpolation DefaultInterpolation = EMovieSceneKeyInterpolation::Auto);
 
 	/**
 	* Adds a Vector2D parameter key to the track.
@@ -124,8 +131,9 @@ public:
 	* @param Time The time to add the new key.
 	* @param RowIndex The preferred row index on which to look for sections.
 	* @param The value for the new key.
+	* @param DefaultInterpolation The default interpolation to use for thew new key.
 	*/
-	MOVIESCENETRACKS_API void AddColorParameterKey(uint8 CustomPrimitiveDataStartIndex, FFrameNumber Position, int32 RowIndex, FLinearColor Value, EMovieSceneKeyInterpolation DefaultInterpolation = EMovieSceneKeyInterpolation::Auto);
+	MOVIESCENETRACKS_API void AddColorParameterKey(uint8 CustomPrimitiveDataStartIndex, FFrameNumber Time, int32 RowIndex, FLinearColor Value, EMovieSceneKeyInterpolation DefaultInterpolation = EMovieSceneKeyInterpolation::Auto);
 
 	/*~ IMovieSceneEntityProvider */
 	virtual void ImportEntityImpl(UMovieSceneEntitySystemLinker* EntityLinker, const FEntityImportParams& Params, FImportedEntity* OutImportedEntity) override;
diff --git a/Engine/Source/Runtime/MovieSceneTracks/Public/Tracks/MovieSceneMaterialTrack.h b/Engine/Source/Runtime/MovieSceneTracks/Public/Tracks/MovieSceneMaterialTrack.h
--- a/Engine/Source/Runtime/MovieSceneTracks/Public/Tracks/MovieSceneMaterialTrack.h
+++ b/Engine/Source/Runtime/MovieSceneTracks/Public/Tracks/MovieSceneMaterialTrack.h
@@ -114,8 +114,9 @@ public:
 	* @param ParameterName The name of the parameter to add a key for.
 	* @param Time The time to add the new key.
 	* @param The value for the new key.
+	* @param DefaultInterpolation The default interpolation to use for thew new key.
 	*/
-	MOVIESCENETRACKS_API void AddScalarParameterKey(FName ParameterName, FFrameNumber Position, float Value, EMovieSceneKeyInterpolation DefaultInterpolation = EMovieSceneKeyInterpolation::Auto);
+	MOVIESCENETRACKS_API void AddScalarParameterKey(FName ParameterName, FFrameNumber Time, float Value, EMovieSceneKeyInterpolation DefaultInterpolation = EMovieSceneKeyInterpolation::Auto);
 
 	/**
 	 * Adds a scalar parameter key to the track. 
@@ -123,8 +124,9 @@ public:
 	 * @param Time The time to add the new key.
 	 * @param RowIndex The preferred row index on which to look for sections.
 	 * @param The value for the new key.
+	 * @param DefaultInterpolation The default interpolation to use for thew new key.
 	 */
-	MOVIESCENETRACKS_API void AddScalarParameterKey(FName ParameterName, FFrameNumber Position, int32 RowIndex, float Value, EMovieSceneKeyInterpolation DefaultInterpolation = EMovieSceneKeyInterpolation::Auto);
+	MOVIESCENETRACKS_API void AddScalarParameterKey(FName ParameterName, FFrameNumber Time, int32 RowIndex, float Value, EMovieSceneKeyInterpolation DefaultInterpolation = EMovieSceneKeyInterpolation::Auto);
 
 	/**
 	* Adds a scalar parameter key to the track.
@@ -133,8 +135,9 @@ public:
 	* @param The value for the new key.
 	* @param InLayerName Optional layer name for use in UI.
 	* @param InAssetName Optional asset name for use in UI.
+	* @param DefaultInterpolation The default interpolation to use for thew new key.
 	*/
-	MOVIESCENETRACKS_API void AddScalarParameterKey(const FMaterialParameterInfo& ParameterInfo, FFrameNumber Position, float Value, const FString& InLayerName, const FString& InAssetName, EMovieSceneKeyInterpolation DefaultInterpolation = EMovieSceneKeyInterpolation::Auto);
+	MOVIESCENETRACKS_API void AddScalarParameterKey(const FMaterialParameterInfo& ParameterInfo, FFrameNumber Time, float Value, const FString& InLayerName, const FString& InAssetName, EMovieSceneKeyInterpolation DefaultInterpolation = EMovieSceneKeyInterpolation::Auto);
 
 	/**
 	* Adds a scalar parameter key to the track.
@@ -144,16 +147,18 @@ public:
 	* @param The value for the new key.
 	* @param InLayerName Optional layer name for use in UI.
 	* @param InAssetName Optional asset name for use in UI.
+	* @param DefaultInterpolation The default interpolation to use for thew new key.
 	*/
-	MOVIESCENETRACKS_API void AddScalarParameterKey(const FMaterialParameterInfo& ParameterInfo, FFrameNumber Position, int32 RowIndex, float Value, const FString& InLayerName, const FString& InAssetName, EMovieSceneKeyInterpolation DefaultInterpolation = EMovieSceneKeyInterpolation::Auto);
+	MOVIESCENETRACKS_API void AddScalarParameterKey(const FMaterialParameterInfo& ParameterInfo, FFrameNumber Time, int32 RowIndex, float Value, const FString& InLayerName, const FString& InAssetName, EMovieSceneKeyInterpolation DefaultInterpolation = EMovieSceneKeyInterpolation::Auto);
 
 	/**
 	* Adds a color parameter key to the track.
 	* @param ParameterName The name of the parameter to add a key for.
 	* @param Time The time to add the new key.
 	* @param The value for the new key.
+	* @param DefaultInterpolation The default interpolation to use for thew new key.
 	*/
-	MOVIESCENETRACKS_API void AddColorParameterKey(FName ParameterName, FFrameNumber Position, FLinearColor Value, EMovieSceneKeyInterpolation DefaultInterpolation = EMovieSceneKeyInterpolation::Auto);
+	MOVIESCENETRACKS_API void AddColorParameterKey(FName ParameterName, FFrameNumber Time, FLinearColor Value, EMovieSceneKeyInterpolation DefaultInterpolation = EMovieSceneKeyInterpolation::Auto);
 
 	/**
 	* Adds a color parameter key to the track.
@@ -161,8 +166,9 @@ public:
 	* @param Time The time to add the new key.
 	* @param RowIndex The preferred row index on which to look for sections.
 	* @param The value for the new key.
+	* @param DefaultInterpolation The default interpolation to use for thew new key.
 	*/
-	MOVIESCENETRACKS_API void AddColorParameterKey(FName ParameterName, FFrameNumber Position, int32 RowIndex, FLinearColor Value, EMovieSceneKeyInterpolation DefaultInterpolation = EMovieSceneKeyInterpolation::Auto);
+	MOVIESCENETRACKS_API void AddColorParameterKey(FName ParameterName, FFrameNumber Time, int32 RowIndex, FLinearColor Value, EMovieSceneKeyInterpolation DefaultInterpolation = EMovieSceneKeyInterpolation::Auto);
 
 	/**
 	* Adds a color parameter key to the track.
@@ -171,8 +177,9 @@ public:
 	* @param The value for the new key.
 	* @param InLayerName Optional layer name for use in UI.
 	* @param InAssetName Optional asset name for use in UI.
+	* @param DefaultInterpolation The default interpolation to use for thew new key.
 	*/
-	MOVIESCENETRACKS_API void AddColorParameterKey(const FMaterialParameterInfo& ParameterInfo, FFrameNumber Position, FLinearColor Value, const FString& InLayerName, const FString& InAssetName, EMovieSceneKeyInterpolation DefaultInterpolation = EMovieSceneKeyInterpolation::Auto);
+	MOVIESCENETRACKS_API void AddColorParameterKey(const FMaterialParameterInfo& ParameterInfo, FFrameNumber Time, FLinearColor Value, const FString& InLayerName, const FString& InAssetName, EMovieSceneKeyInterpolation DefaultInterpolation = EMovieSceneKeyInterpolation::Auto);
 
 	/**
 	* Adds a color parameter key to the track.
@@ -182,8 +189,9 @@ public:
 	* @param The value for the new key.
 	* @param InLayerName Optional layer name for use in UI.
 	* @param InAssetName Optional asset name for use in UI.
+	* @param DefaultInterpolation The default interpolation to use for thew new key.
 	*/
-	MOVIESCENETRACKS_API void AddColorParameterKey(const FMaterialParameterInfo& ParameterInfo, FFrameNumber Position, int32 RowIndex, FLinearColor Value, const FString& InLayerName, const FString& InAssetName, EMovieSceneKeyInterpolation DefaultInterpolation = EMovieSceneKeyInterpolation::Auto);
+	MOVIESCENETRACKS_API void AddColorParameterKey(const FMaterialParameterInfo& ParameterInfo, FFrameNumber Time, int32 RowIndex, FLinearColor Value, const FString& InLayerName, const FString& InAssetName, EMovieSceneKeyInterpolation DefaultInterpolation = EMovieSceneKeyInterpolation::Auto);
 
 	/**
 	* Adds a color parameter key to the track.
@@ -194,8 +202,9 @@ public:
 	* @param InLayerName Optional layer name for use in UI.
 	* @param InAssetName Optional asset name for use in UI.
 	* @param InChannelNames Optional channel names for the 4 vector channels for use in the UI.
+	* @param DefaultInterpolation The default interpolation to use for thew new key.
 	*/
-	MOVIESCENETRACKS_API void AddColorParameterKey(const FMaterialParameterInfo& ParameterInfo, FFrameNumber Position, int32 RowIndex, FLinearColor Value, const FString& InLayerName, const FString& InAssetName, const FParameterChannelNames& InChannelNames, EMovieSceneKeyInterpolation DefaultInterpolation = EMovieSceneKeyInterpolation::Auto);
+	MOVIESCENETRACKS_API void AddColorParameterKey(const FMaterialParameterInfo& ParameterInfo, FFrameNumber Time, int32 RowIndex, FLinearColor Value, const FString& InLayerName, const FString& InAssetName, const FParameterChannelNames& InChannelNames, EMovieSceneKeyInterpolation DefaultInterpolation = EMovieSceneKeyInterpolation::Auto);
 
 private:
 
diff --git a/Engine/Source/Runtime/MovieSceneTracks/Public/Tracks/MovieSceneParticleParameterTrack.h b/Engine/Source/Runtime/MovieSceneTracks/Public/Tracks/MovieSceneParticleParameterTrack.h
--- a/Engine/Source/Runtime/MovieSceneTracks/Public/Tracks/MovieSceneParticleParameterTrack.h
+++ b/Engine/Source/Runtime/MovieSceneTracks/Public/Tracks/MovieSceneParticleParameterTrack.h
@@ -60,24 +60,27 @@ public:
 	 * @param ParameterName The name of the parameter to add a key for.
 	 * @param Time The time to add the new key.
 	 * @param The value for the new key.
+	 * @param DefaultInterpolation The default interpolation to use for thew new key.
 	 */
-	void MOVIESCENETRACKS_API AddScalarParameterKey( FName ParameterName, FFrameNumber Position, float Value, EMovieSceneKeyInterpolation DefaultInterpolation = EMovieSceneKeyInterpolation::Auto);
+	void MOVIESCENETRACKS_API AddScalarParameterKey( FName ParameterName, FFrameNumber Time, float Value, EMovieSceneKeyInterpolation DefaultInterpolation = EMovieSceneKeyInterpolation::Auto);
 
 	/**
 	* Adds a Vector parameter key to the track.
 	* @param ParameterName The name of the parameter to add a key for.
 	* @param Time The time to add the new key.
 	* @param The value for the new key.
+	* @param DefaultInterpolation The default interpolation to use for thew new key.
 	*/
-	void MOVIESCENETRACKS_API AddVectorParameterKey( FName ParameterName, FFrameNumber Position, FVector Value, EMovieSceneKeyInterpolation DefaultInterpolation = EMovieSceneKeyInterpolation::Auto);
+	void MOVIESCENETRACKS_API AddVectorParameterKey( FName ParameterName, FFrameNumber Time, FVector Value, EMovieSceneKeyInterpolation DefaultInterpolation = EMovieSceneKeyInterpolation::Auto);
 
 	/**
-	* Adds a Vector parameter key to the track.
+	* Adds a color parameter key to the track.
 	* @param ParameterName The name of the parameter to add a key for.
 	* @param Time The time to add the new key.
 	* @param The value for the new key.
+	* @param DefaultInterpolation The default interpolation to use for thew new key.
 	*/
-	void MOVIESCENETRACKS_API AddColorParameterKey( FName ParameterName, FFrameNumber Position, FLinearColor Value, EMovieSceneKeyInterpolation DefaultInterpolation = EMovieSceneKeyInterpolation::Auto);
+	void MOVIESCENETRACKS_API AddColorParameterKey( FName ParameterName, FFrameNumber Time, FLinearColor Value, EMovieSceneKeyInterpolation DefaultInterpolation = EMovieSceneKeyInterpolation::Auto);
 
 private:
 
