UE-235104
Fix Movie Render Queue - Audio - Fix Popping/Gaps/Repeats between Camera Cuts
Solution: Pause/resume the AudioComponent between shots, maintaining exact play state (e.g. position) between shots, rather than stopping and attempting to restart with a non-sample-accurate start time.

diff --git a/Engine/Plugins/MovieScene/MovieRenderPipeline/Source/MovieRenderPipelineCore/Private/MoviePipelineAudioRendering.cpp b/Engine/Plugins/MovieScene/MovieRenderPipeline/Source/MovieRenderPipelineCore/Private/MoviePipelineAudioRendering.cpp
--- a/Engine/Plugins/MovieScene/MovieRenderPipeline/Source/MovieRenderPipelineCore/Private/MoviePipelineAudioRendering.cpp
+++ b/Engine/Plugins/MovieScene/MovieRenderPipeline/Source/MovieRenderPipelineCore/Private/MoviePipelineAudioRendering.cpp
@@ -141,16 +141,8 @@ void UMoviePipeline::ProcessAudioTick()
 	if (CurrentShot->ShotInfo.State == EMovieRenderShotState::Rendering && !AudioState.bIsRecordingAudio)
 	{
 		// There can be time in the Unreal world that passes when MRQ is not yet active. When MRQ isn't active,
-		// then the NRT audio platform doesn't process anything and it builds up for when we do start to actually
-		// process it for a shot. This causes our actual audio to become offset from the sequence once lined up
-		// afterwards. So for now we'll just flush out anything that needed to be processed before we start capturing
-			
-		MixerDevice->Update(true);
-		
-		// We don't have a great way of determining how much outstanding audio there is, so we're just going to 
-		// process an arbitrary amount of time.
-		NRTPlatform->RenderAudio(30.f);
-		
+		// then the NRT audio platform doesn't process anything as all the AudioComponents are paused.
+		// On capturing the next camera, The AudioComponents resume playback, ensuring sample-accurate continuity across camera cuts.
 		StartAudioRecording();
 	}
 
diff --git a/Engine/Source/Runtime/MovieSceneTracks/Private/Systems/MovieSceneAudioSystem.cpp b/Engine/Source/Runtime/MovieSceneTracks/Private/Systems/MovieSceneAudioSystem.cpp
--- a/Engine/Source/Runtime/MovieSceneTracks/Private/Systems/MovieSceneAudioSystem.cpp
+++ b/Engine/Source/Runtime/MovieSceneTracks/Private/Systems/MovieSceneAudioSystem.cpp
@@ -110,15 +110,13 @@ struct FPreAnimatedAudioStateTraits : BaseTraits
 	{
 		if (UAudioComponent* AudioComponent = Cast<UAudioComponent>(InKey.ResolveObjectPtr()))
 		{
-			switch (InStateType)
-			{
-				case EPreAnimatedAudioStateType::AudioPlaying:
-					AudioComponent->Stop();
-					break;
-				case EPreAnimatedAudioStateType::AudioComponentLifespan:
-					AudioComponent->DestroyComponent();
-					break;
-			}
+			// MRQ will process UE time outside of recording to let secondary animations settle etc.
+			// To ensure audio capture has no gaps between shots, we pause to maintain audio state between shots.
+			// On capturing the next camera, we'll resume audio playback.
+			check(InStateType == EPreAnimatedAudioStateType::AudioPlaying || InStateType == EPreAnimatedAudioStateType::AudioComponentLifespan);
+			AudioComponent->SetPaused(true);
+			// For EPreAnimatedAudioStateType::AudioComponentLifespan, rather than explicitly calling AudioComponent.DestroyComponent(),
+			// we rely on it being destroyed during ULevel::ClearLevelComponents.
 		}
 	}
 };
@@ -624,7 +622,7 @@ private:
 					{
 						// If we're not a looping section and the AC is done, return
 						// otherwise a looping section will restart the sound if it's not playing
-						if (!AudioComponent.IsPlaying())
+						if (!AudioComponent.IsPlaying()) // !IsPlaying() == not playing, fading, or paused
 						{
 							UE_LOG(LogMovieScene, Verbose, TEXT("Procedural sound Audio Component reached end of playback. Component: %s Sound: %s"), *AudioComponent.GetName(), *GetNameSafe(AudioComponent.Sound));
 							return;
@@ -642,7 +640,7 @@ private:
 				if (!AudioSection.GetLooping() && AudioTime > Duration && Duration != 0.f)
 				{
 					// If this is Non-ProcedurallyGenerated audio, and it's not looping then check to see if it needs to be stopped
-					if (AudioComponent.IsPlaying())
+					if (AudioComponent.IsPlaying()) // IsPlaying() == playing, fading, or paused
 					{
 						UE_LOG(LogMovieScene, Verbose, TEXT("Audio Component reached end of playback. Component: %s Sound: %s"), *AudioComponent.GetName(), *GetNameSafe(AudioComponent.Sound));
 						AudioComponent.Stop();
@@ -661,7 +659,8 @@ private:
 		// If the audio component is not playing we (may) need a state change. If the audio component is playing
 		// the wrong sound then we need a state change. If the audio playback time is significantly out of sync 
 		// with the desired time then we need a state change.
-		const bool bSoundsNeedPlaying = !AudioComponent.IsPlaying();
+		const bool bAudioComponentIsPaused = AudioComponent.GetPlayState() == EAudioComponentPlayState::Paused;
+		const bool bSoundsNeedPlaying = !AudioComponent.IsPlaying() || bAudioComponentIsPaused;
 		const bool bSoundNeedsStateChange =  AudioComponent.Sound != Sound;
 		bool bSoundNeedsTimeSync = false;
 
@@ -734,7 +733,7 @@ private:
 			{
 				// Only call stop on the sound if it is actually playing. This prevents spamming
 				// stop calls when a sound cue with a duration of zero is played.
-				if (AudioComponent.IsPlaying() || bSoundNeedsTimeSync)
+				if ((AudioComponent.IsPlaying() && !bAudioComponentIsPaused) || bSoundNeedsTimeSync)
 				{
 					UE_LOG(LogMovieScene, Verbose, TEXT("Audio Component stopped due to needing a state change bIsPlaying: %d bNeedsTimeSync: %d. Component: %s Sound: %s"), AudioComponent.IsPlaying(), bSoundNeedsTimeSync, *AudioComponent.GetName(), *GetNameSafe(AudioComponent.Sound));
 					AudioComponent.Stop();
@@ -776,7 +775,7 @@ private:
 				if (bPerformGranularScrubbing)
 				{
 					// If we're not playing the audio component, then set up the audio component to use the scrubbed sound wave and play it
-					if (!AudioComponent.IsPlaying())
+					if (!AudioComponent.IsPlaying()) // !IsPlaying() == not playing, fading, or paused
 					{
 						ScrubbedSound->SetSoundWave(SoundWave);
 						ScrubbedSound->SetPlayheadTime(AudioTime);
@@ -813,7 +812,14 @@ private:
 						AudioComponent.SetSound(Sound);
 					}
 					
-					AudioComponent.Play(AudioTime);
+					if (!bAudioComponentIsPaused)
+					{
+						AudioComponent.Play(AudioTime);
+					}
+					else
+					{
+						AudioComponent.SetPaused(false); // no-op if already unpaused
+					}
 					
 					if (Context.GetStatus() == EMovieScenePlayerStatus::Playing)
 					{
@@ -868,7 +874,7 @@ private:
 			AudioComponent.StopDelayed(ScrubDuration);
 		}
 
-		if (AudioComponent.IsPlaying() && AudioInputs)
+		if ((AudioComponent.IsPlaying() && !bAudioComponentIsPaused) && AudioInputs)
 		{
 			SetAudioInputTriggers(AudioInputs->Inputs_Trigger, AudioComponent);
 		}
@@ -1295,7 +1301,12 @@ void UMovieSceneAudioSystem::StopSound(FInstanceHandle InstanceHandle, FObjectKe
 		{
 			if (UAudioComponent* AudioComponent = Data->AudioComponent.Get())
 			{
-				AudioComponent->Stop();
+				// Camera cuts cannot be differentiated from explicit stops as the pipeline is currently set up.
+				// Pause should result in equivalent behaviour for the vast majority of cases.
+				// In the rare case of a release tail to be played after stop, users can release explicitly
+				// which will be naturally handled by the engine (e.g. MetaSounds) and its
+				// existing support for pause/resume while in the tail.
+				AudioComponent->SetPaused(true);
 			}
 
 			Data->bAudioComponentHasBeenPlayed = false;
