Sequencer: fix missing support of InterpNotify for property track that don't support channel overrides

diff --git a/Engine/Source/Runtime/MovieSceneTracks/Public/Tracks/MovieScenePropertyTrack.h b/Engine/Source/Runtime/MovieSceneTracks/Public/Tracks/MovieScenePropertyTrack.h
--- a/Engine/Source/Runtime/MovieSceneTracks/Public/Tracks/MovieScenePropertyTrack.h
+++ b/Engine/Source/Runtime/MovieSceneTracks/Public/Tracks/MovieScenePropertyTrack.h
@@ -208,6 +208,7 @@ struct TPropertyTrackEntityImportHelperImpl
 	void Commit(const UMovieSceneSection* InSection, const UE::MovieScene::FEntityImportParams& Params, UE::MovieScene::FImportedEntity* OutImportedEntity)
 	{
 		const FGuid ObjectBindingID = Params.GetObjectBindingID();
+		const FMovieSceneEvaluationFieldEntityMetaData* MetaData = Params.EntityMetaData;
 		const FBuiltInComponentTypes* BuiltInComponents = FBuiltInComponentTypes::Get();
 		const FMovieSceneTracksComponentTypes* TracksComponents = FMovieSceneTracksComponentTypes::Get();
 
@@ -215,10 +216,13 @@ struct TPropertyTrackEntityImportHelperImpl
 		{
 			if (FMovieScenePropertyTrackEntityImportHelper::IsPropertyValueID(Params))
 			{
+				const FName NotifyFunctionName = (MetaData ? MetaData->NotifyFunctionName : NAME_None);
+
 				OutImportedEntity->AddBuilder(
 					Builder
 					.Add(BuiltInComponents->PropertyBinding, PropertyTrack->GetPropertyBinding())
-					.AddConditional(BuiltInComponents->GenericObjectBinding, ObjectBindingID, ObjectBindingID.IsValid()));
+					.AddConditional(BuiltInComponents->GenericObjectBinding, ObjectBindingID, ObjectBindingID.IsValid())
+					.AddConditional(TracksComponents->PropertyNotify, FPropertyNotifyComponentData{ NotifyFunctionName }, !NotifyFunctionName.IsNone()));
 			}
 			else if (ensure(FMovieScenePropertyTrackEntityImportHelper::IsEditConditionToggleID(Params)))
 			{
