Sequencer: fix calls to bool property custom accessors, which were using the wrong function signature

The code was still using the old signature without the FBoolMetadata parameter, so we were basically calling 3-parameter functions using only 2 parameters.

diff --git a/Engine/Source/Runtime/MovieSceneTracks/Public/MovieSceneTracksComponentTypes.h b/Engine/Source/Runtime/MovieSceneTracks/Public/MovieSceneTracksComponentTypes.h
--- a/Engine/Source/Runtime/MovieSceneTracks/Public/MovieSceneTracksComponentTypes.h
+++ b/Engine/Source/Runtime/MovieSceneTracks/Public/MovieSceneTracksComponentTypes.h
@@ -352,8 +352,8 @@ struct FBoolPropertyTraits : IPropertyTraits
 	/** Property Value Getters  */
 	static void GetObjectPropertyValue(const UObject* InObject, FBoolMetaData MetaData, const FCustomPropertyAccessor& BaseCustomAccessor, bool& OutValue)
 	{
-		const TCustomPropertyAccessor<TDirectPropertyTraits<bool>>& CustomAccessor = static_cast<const TCustomPropertyAccessor<TDirectPropertyTraits<bool>>&>(BaseCustomAccessor);
-		OutValue = (*CustomAccessor.Functions.Getter)(InObject);
+		const TCustomPropertyAccessor<FBoolPropertyTraits>& CustomAccessor = static_cast<const TCustomPropertyAccessor<FBoolPropertyTraits>&>(BaseCustomAccessor);
+		OutValue = (*CustomAccessor.Functions.Getter)(InObject, MetaData);
 	}
 	static void GetObjectPropertyValue(const UObject* InObject, FBoolMetaData MetaData, uint16 PropertyOffset, bool& OutValue)
 	{
@@ -383,8 +383,8 @@ struct FBoolPropertyTraits : IPropertyTraits
 	/** Property Value Setters  */
 	static void SetObjectPropertyValue(UObject* InObject, FBoolMetaData MetaData, const FCustomPropertyAccessor& BaseCustomAccessor, bool InValue)
 	{
-		const TCustomPropertyAccessor<TDirectPropertyTraits<bool>>& CustomAccessor = static_cast<const TCustomPropertyAccessor<TDirectPropertyTraits<bool>>&>(BaseCustomAccessor);
-		(*CustomAccessor.Functions.Setter)(InObject, InValue);
+		const TCustomPropertyAccessor<FBoolPropertyTraits>& CustomAccessor = static_cast<const TCustomPropertyAccessor<FBoolPropertyTraits>&>(BaseCustomAccessor);
+		(*CustomAccessor.Functions.Setter)(InObject, MetaData, InValue);
 	}
 	static void SetObjectPropertyValue(UObject* InObject, FBoolMetaData MetaData, uint16 PropertyOffset, bool InValue)
 	{
