Sequencer: take the first loop start offset into account when determining if a skeletal animation section needs to loop

Without taking it into account, we could have a short section that is shorter than the animation length, but the first loop start offset was big enough to start the animation near the end and require looping still.

diff --git a/Engine/Source/Runtime/MovieSceneTracks/Private/Sections/MovieSceneSkeletalAnimationSection.cpp b/Engine/Source/Runtime/MovieSceneTracks/Private/Sections/MovieSceneSkeletalAnimationSection.cpp
--- a/Engine/Source/Runtime/MovieSceneTracks/Private/Sections/MovieSceneSkeletalAnimationSection.cpp
+++ b/Engine/Source/Runtime/MovieSceneTracks/Private/Sections/MovieSceneSkeletalAnimationSection.cpp
@@ -69,7 +69,7 @@ FMovieSceneSequenceTransform FMovieSceneSkeletalAnimationParams::MakeTransform(c
 	const FFrameTime AnimationLength = SequenceLength * OuterFrameRate;
 	const int32 LengthInFrames = AnimationLength.FrameNumber.Value + (int)(AnimationLength.GetSubFrame() + 0.5f) + 1;
 
-	const bool bLooping = (UE::MovieScene::DiscreteSize(OuterRange) + StartFrameOffset + EndFrameOffset) > LengthInFrames;
+	const bool bLooping = (UE::MovieScene::DiscreteSize(OuterRange) + FirstLoopStartFrameOffset + StartFrameOffset + EndFrameOffset) > LengthInFrames;
 
 	FMovieSceneSectionTimingParametersSeconds TimingParams = {
 		PlayRate.ShallowCopy(),
