Sequencer: fix regression in how skeletal animation sections behave when resized from the left edge

diff --git a/Engine/Source/Editor/MovieSceneTools/Private/TrackEditors/CommonAnimationTrackEditor.cpp b/Engine/Source/Editor/MovieSceneTools/Private/TrackEditors/CommonAnimationTrackEditor.cpp
--- a/Engine/Source/Editor/MovieSceneTools/Private/TrackEditors/CommonAnimationTrackEditor.cpp
+++ b/Engine/Source/Editor/MovieSceneTools/Private/TrackEditors/CommonAnimationTrackEditor.cpp
@@ -870,8 +870,10 @@ void FCommonAnimationSection::BeginResizeSection()
 		return;
 	}
 
+	const bool bClampToOuterRange = false; // When resizing, we want to be able to go outside of the section range.
+	const bool bForceLoop = true; // Similarly, when resizing, we want to be able to start looping even if the section isn't looping right now.
 	const FFrameRate FrameRate = Section->GetTypedOuter<UMovieScene>()->GetTickResolution();
-	InitialDragTransform = MakeUnique<FMovieSceneSequenceTransform>(Section->Params.MakeTransform(FrameRate, Section->GetRange()));
+	InitialDragTransform = MakeUnique<FMovieSceneSequenceTransform>(Section->Params.MakeTransform(FrameRate, Section->GetRange(), nullptr, bClampToOuterRange, bForceLoop));
 }
 
 void FCommonAnimationSection::ResizeSection(ESequencerSectionResizeMode ResizeMode, FFrameNumber ResizeTime)
diff --git a/Engine/Source/Runtime/MovieSceneTracks/Private/Sections/MovieSceneSkeletalAnimationSection.cpp b/Engine/Source/Runtime/MovieSceneTracks/Private/Sections/MovieSceneSkeletalAnimationSection.cpp
--- a/Engine/Source/Runtime/MovieSceneTracks/Private/Sections/MovieSceneSkeletalAnimationSection.cpp
+++ b/Engine/Source/Runtime/MovieSceneTracks/Private/Sections/MovieSceneSkeletalAnimationSection.cpp
@@ -53,7 +53,7 @@ FMovieSceneSkeletalAnimationParams::FMovieSceneSkeletalAnimationParams()
 	bLinearPlaybackWhenScaled = false;
 }
 
-FMovieSceneSequenceTransform FMovieSceneSkeletalAnimationParams::MakeTransform(const FFrameRate& OuterFrameRate, const TRange<FFrameNumber>& OuterRange, UAnimSequenceBase* OverrideSequence) const
+FMovieSceneSequenceTransform FMovieSceneSkeletalAnimationParams::MakeTransform(const FFrameRate& OuterFrameRate, const TRange<FFrameNumber>& OuterRange, UAnimSequenceBase* OverrideSequence, bool bClampToOuterRange, bool bForceLoop) const
 {
 	if (!OverrideSequence)
 	{
@@ -69,7 +69,7 @@ FMovieSceneSequenceTransform FMovieSceneSkeletalAnimationParams::MakeTransform(c
 	const FFrameTime AnimationLength = SequenceLength * OuterFrameRate;
 	const int32 LengthInFrames = AnimationLength.FrameNumber.Value + (int)(AnimationLength.GetSubFrame() + 0.5f) + 1;
 
-	const bool bLooping = (UE::MovieScene::DiscreteSize(OuterRange) + FirstLoopStartFrameOffset + StartFrameOffset + EndFrameOffset) > LengthInFrames;
+	const bool bLooping = bForceLoop || ((UE::MovieScene::DiscreteSize(OuterRange) + FirstLoopStartFrameOffset + StartFrameOffset + EndFrameOffset) > LengthInFrames);
 
 	FMovieSceneSectionTimingParametersSeconds TimingParams = {
 		PlayRate.ShallowCopy(),
@@ -79,7 +79,7 @@ FMovieSceneSequenceTransform FMovieSceneSkeletalAnimationParams::MakeTransform(c
 		bLooping,
 		!bLooping,
 		bReverse,
-		true  // bClampToOuterRange
+		bClampToOuterRange
 	};
 
 	return TimingParams.MakeTransform(OuterFrameRate, OuterRange, SequenceLength, OverrideSequence->RateScale);
diff --git a/Engine/Source/Runtime/MovieSceneTracks/Public/Sections/MovieSceneSkeletalAnimationSection.h b/Engine/Source/Runtime/MovieSceneTracks/Public/Sections/MovieSceneSkeletalAnimationSection.h
--- a/Engine/Source/Runtime/MovieSceneTracks/Public/Sections/MovieSceneSkeletalAnimationSection.h
+++ b/Engine/Source/Runtime/MovieSceneTracks/Public/Sections/MovieSceneSkeletalAnimationSection.h
@@ -45,7 +45,7 @@ struct FMovieSceneSkeletalAnimationParams
 	/**
 	 * Make a transform structure from these animation parameters
 	 */
-	MOVIESCENETRACKS_API FMovieSceneSequenceTransform MakeTransform(const FFrameRate& OuterFrameRate, const TRange<FFrameNumber>& OuterRange, UAnimSequenceBase* OverrideSequence = nullptr) const;
+	MOVIESCENETRACKS_API FMovieSceneSequenceTransform MakeTransform(const FFrameRate& OuterFrameRate, const TRange<FFrameNumber>& OuterRange, UAnimSequenceBase* OverrideSequence = nullptr, bool bClampToOuterRange = true, bool bForceLoop = false) const;
 
 	/** The animation this section plays */
 	UPROPERTY(BlueprintReadWrite, EditAnywhere, Category="Animation", meta=(AllowedClasses = "/Script/Engine.AnimSequence,/Script/Engine.AnimComposite,/Script/Engine.AnimStreamable"))
