Sequencer: disable ensure before cleaning up camera cut track viewport modifiers

When we close the editor while the game is running with a cinematic playing, Sequencer doesn't get any chance to update: everything gets destroyed immediately. So we can indeed hit the situation where view modifiers are still registered legitimately even though we are hitting destructors without having been torn down gracefully.

diff --git a/Engine/Source/Runtime/MovieSceneTracks/Private/TrackInstances/MovieSceneCameraCutViewportPreviewer.cpp b/Engine/Source/Runtime/MovieSceneTracks/Private/TrackInstances/MovieSceneCameraCutViewportPreviewer.cpp
--- a/Engine/Source/Runtime/MovieSceneTracks/Private/TrackInstances/MovieSceneCameraCutViewportPreviewer.cpp
+++ b/Engine/Source/Runtime/MovieSceneTracks/Private/TrackInstances/MovieSceneCameraCutViewportPreviewer.cpp
@@ -65,7 +65,7 @@ FCameraCutViewportPreviewer::FCameraCutViewportPreviewer()
 
 FCameraCutViewportPreviewer::~FCameraCutViewportPreviewer()
 {
-	if (!ensure(!bViewportModifiersRegistered))
+	if (bViewportModifiersRegistered)
 	{
 		ToggleViewportPreviewModifiers(false);
 	}
