Sequencer: fully support all camera and post-process properties when previewing blending camera cuts

Ask the camera component (if any) for its final view, and use that for blending, otherwise we don't get any generated post-process settings.

diff --git a/Engine/Source/Runtime/MovieSceneTracks/Private/TrackInstances/MovieSceneCameraCutViewportPreviewer.cpp b/Engine/Source/Runtime/MovieSceneTracks/Private/TrackInstances/MovieSceneCameraCutViewportPreviewer.cpp
--- a/Engine/Source/Runtime/MovieSceneTracks/Private/TrackInstances/MovieSceneCameraCutViewportPreviewer.cpp
+++ b/Engine/Source/Runtime/MovieSceneTracks/Private/TrackInstances/MovieSceneCameraCutViewportPreviewer.cpp
@@ -13,15 +13,18 @@
 namespace UE::MovieScene
 {
 
-void FCameraCutViewportPreviewerTarget::Get(FLevelEditorViewportClient* InClient, FVector& OutLocation, FRotator& OutRotation, float& OutFOV, const FPostProcessSettings*& OutPP, float& OutPPWeight) const
+void FCameraCutViewportPreviewerTarget::Get(FLevelEditorViewportClient* InClient, FVector& OutLocation, FRotator& OutRotation, float& OutFOV, FPostProcessSettings& OutPP, float& OutPPWeight) const
 {
 	if (CameraComponent)
 	{
-		OutLocation = CameraComponent->GetComponentLocation();
-		OutRotation = CameraComponent->GetComponentRotation();
-		OutFOV = CameraComponent->FieldOfView;
-		OutPP = &CameraComponent->PostProcessSettings;
-		OutPPWeight = CameraComponent->PostProcessBlendWeight;
+		FMinimalViewInfo CameraView;
+		CameraComponent->GetCameraView(0.f, CameraView);
+
+		OutLocation = CameraView.Location;
+		OutRotation = CameraView.Rotation;
+		OutFOV = CameraView.FOV;
+		OutPP = CameraView.PostProcessSettings;
+		OutPPWeight = CameraView.PostProcessBlendWeight;
 		return;
 	}
 
@@ -30,7 +33,6 @@ void FCameraCutViewportPreviewerTarget::Get(FLevelEditorViewportClient* InClient
 		OutLocation = CameraActor->GetActorLocation();
 		OutRotation = CameraActor->GetActorRotation();
 		OutFOV = -1.f;
-		OutPP = nullptr;
 		OutPPWeight = 0.f;
 		return;
 	}
@@ -44,7 +46,6 @@ void FCameraCutViewportPreviewerTarget::Get(FLevelEditorViewportClient* InClient
 			OutLocation = CachedValue.ViewportLocation;
 			OutRotation = CachedValue.ViewportRotation;
 			OutFOV = InClient->FOVAngle;
-			OutPP = nullptr;
 			OutPPWeight = 0.f;
 			return;
 		}
@@ -55,7 +56,6 @@ void FCameraCutViewportPreviewerTarget::Get(FLevelEditorViewportClient* InClient
 	OutLocation = InClient->GetViewLocation();
 	OutRotation = InClient->GetViewRotation();
 	OutFOV = InClient->ViewFOV;
-	OutPP = nullptr;
 	OutPPWeight = 0.f;
 }
 
@@ -134,7 +134,7 @@ void FCameraCutViewportPreviewer::ModifyViewportClientView(FEditorViewportViewMo
 	FVector FromViewLocation, ToViewLocation;
 	FRotator FromViewRotation, ToViewRotation;
 	float FromViewFOV, ToViewFOV;
-	const FPostProcessSettings* FromPP; const FPostProcessSettings* ToPP;
+	FPostProcessSettings FromPP; FPostProcessSettings ToPP;
 	float FromPPWeight, ToPPWeight;
 	FLevelEditorViewportClient* ViewportClient = static_cast<FLevelEditorViewportClient*>(Params.ViewportClient);
 
@@ -149,24 +149,27 @@ void FCameraCutViewportPreviewer::ModifyViewportClientView(FEditorViewportViewMo
 	Params.ViewInfo.Rotation = BlendedRotation;
 	Params.ViewInfo.FOV = BlendedFOV;
 
-	if (FromPP)
+	const bool bHasFromPP = (FromPPWeight > 0.f);
+	const bool bHasToPP = (ToPPWeight > 0.f);
+
+	if (bHasFromPP)
 	{
-		FPostProcessUtils::OverridePostProcessSettings(Params.ViewInfo.PostProcessSettings, *FromPP);
+		FPostProcessUtils::OverridePostProcessSettings(Params.ViewInfo.PostProcessSettings, FromPP);
 	}
-	if (ToPP)
+	if (bHasToPP)
 	{
-		FPostProcessUtils::BlendPostProcessSettings(Params.ViewInfo.PostProcessSettings, *ToPP, BlendFactor);
+		FPostProcessUtils::BlendPostProcessSettings(Params.ViewInfo.PostProcessSettings, ToPP, BlendFactor);
 	}
 
-	if (FromPP && !ToPP)
+	if (bHasFromPP && !bHasToPP)
 	{
 		Params.ViewInfo.PostProcessBlendWeight = (1.f - BlendFactor);
 	}
-	else if (!FromPP && ToPP)
+	else if (!bHasFromPP && bHasToPP)
 	{
 		Params.ViewInfo.PostProcessBlendWeight = BlendFactor;
 	}
-	else if (FromPP && ToPP)
+	else if (bHasFromPP && bHasToPP)
 	{
 		Params.ViewInfo.PostProcessBlendWeight = 1.f;
 	}
diff --git a/Engine/Source/Runtime/MovieSceneTracks/Private/TrackInstances/MovieSceneCameraCutViewportPreviewer.h b/Engine/Source/Runtime/MovieSceneTracks/Private/TrackInstances/MovieSceneCameraCutViewportPreviewer.h
--- a/Engine/Source/Runtime/MovieSceneTracks/Private/TrackInstances/MovieSceneCameraCutViewportPreviewer.h
+++ b/Engine/Source/Runtime/MovieSceneTracks/Private/TrackInstances/MovieSceneCameraCutViewportPreviewer.h
@@ -32,7 +32,7 @@ namespace UE::MovieScene
 		TSharedPtr<FPreAnimatedCameraCutEditorStorage> PreAnimatedStorage;
 
 		/** Returns the target's location, rotation, and FOV */
-		void Get(FLevelEditorViewportClient* InClient, FVector& OutLocation, FRotator& OutRotation, float& OutFOV, const FPostProcessSettings*& OutPP, float& OutPPWeight) const;
+		void Get(FLevelEditorViewportClient* InClient, FVector& OutLocation, FRotator& OutRotation, float& OutFOV, FPostProcessSettings& OutPP, float& OutPPWeight) const;
 	};
 
 	/**
