Sequencer: Check valid transform and output a warning if not valid

diff --git a/Engine/Source/Runtime/MovieSceneTracks/Private/MovieSceneTracksCustomAccessors.cpp b/Engine/Source/Runtime/MovieSceneTracks/Private/MovieSceneTracksCustomAccessors.cpp
--- a/Engine/Source/Runtime/MovieSceneTracks/Private/MovieSceneTracksCustomAccessors.cpp
+++ b/Engine/Source/Runtime/MovieSceneTracks/Private/MovieSceneTracksCustomAccessors.cpp
@@ -52,6 +52,17 @@ void SetComponentTransform(USceneComponent* SceneComponent, const FIntermediate3
 	FVector Translation = Transform.GetTranslation();
 	FRotator Rotation = Transform.GetRotation();
 	FTransform NewTransform(Rotation, Translation, Transform.GetScale());
+	if (!NewTransform.IsValid())
+	{
+#if !NO_LOGGING
+		if (UE_LOG_ACTIVE(LogMovieScene, Warning))
+		{
+			UE_LOG(LogMovieScene, Warning, TEXT("SetComponentTransform: Trying to set new transform with bad data: %s"), *NewTransform.ToString());
+		}
+#endif
+		return;
+	}
+
 	const bool bFlipped = NewTransform.GetDeterminant() * SceneComponent->GetComponentTransform().GetDeterminant() <= 0;
 
 	SceneComponent->SetRelativeTransform(NewTransform, false, nullptr, bIsSimulatingPhysics ? ETeleportType::ResetPhysics : ETeleportType::None);
