[Sequencer] Audio should not default to looping when added to a sequence

diff --git a/Engine/Source/Runtime/MovieSceneTracks/Private/Sections/MovieSceneAudioSection.cpp b/Engine/Source/Runtime/MovieSceneTracks/Private/Sections/MovieSceneAudioSection.cpp
--- a/Engine/Source/Runtime/MovieSceneTracks/Private/Sections/MovieSceneAudioSection.cpp
+++ b/Engine/Source/Runtime/MovieSceneTracks/Private/Sections/MovieSceneAudioSection.cpp
@@ -77,7 +77,7 @@ UMovieSceneAudioSection::UMovieSceneAudioSection( const FObjectInitializer& Obje
 	AudioStartTime_DEPRECATED = AudioDeprecatedMagicNumber;
 	AudioDilationFactor_DEPRECATED = AudioDeprecatedMagicNumber;
 	AudioVolume_DEPRECATED = AudioDeprecatedMagicNumber;
-	bLooping = true;
+	bLooping = false;
 	bSuppressSubtitles = false;
 	bOverrideAttenuation = false;
 	BlendType = EMovieSceneBlendType::Absolute;
@@ -351,6 +351,25 @@ void UMovieSceneAudioSection::PostLoad()
 		StartFrameOffset = ConvertFrameTime(FFrameTime::FromDecimal(DisplayRate.AsDecimal() * StartOffsetToUpgrade.GetValue()), DisplayRate, TickResolution).FrameNumber;
 	}
 }
+
+#if WITH_EDITOR
+void UMovieSceneAudioSection::PostEditChangeProperty(FPropertyChangedEvent& PropertyChangedEvent)
+{
+	Super::PostEditChangeProperty(PropertyChangedEvent);
+
+	if (PropertyChangedEvent.Property == nullptr)
+	{
+		return;
+	}
+
+	FName PropertyName = PropertyChangedEvent.Property->GetFName();
+
+	if (PropertyName == GET_MEMBER_NAME_CHECKED(UMovieSceneAudioSection, bLooping))
+	{
+		SetRange(SectionRange.Value);
+	}
+}
+#endif // WITH_EDITOR
 	
 TOptional<TRange<FFrameNumber> > UMovieSceneAudioSection::GetAutoSizeRange() const
 {
@@ -382,6 +401,46 @@ TOptional<TRange<FFrameNumber> > UMovieSceneAudioSection::GetAutoSizeRange() con
 	return TRange<FFrameNumber>(GetInclusiveStartFrame(), GetInclusiveStartFrame() + IFrameNumber);
 }
 
+void UMovieSceneAudioSection::SetRange(const TRange<FFrameNumber>& NewRange)
+{
+	TRange<FFrameNumber> ClampedRange = NewRange;
+
+	if (Sound && !bLooping)
+	{
+		check(GetTypedOuter<UMovieScene>());
+
+		const float SoundDuration = MovieSceneHelpers::GetSoundDuration(Sound);
+
+		// Metasound sources will return INDEFINITELY_LOOPING_DURATION whether they are a oneshot or not
+		if (SoundDuration == INDEFINITELY_LOOPING_DURATION || !(NewRange.HasUpperBound() && NewRange.HasLowerBound()))
+		{
+			Super::SetRange(NewRange);
+
+			return;
+		}
+
+		const FFrameRate FrameRate = GetTypedOuter<UMovieScene>()->GetTickResolution();
+		const FFrameNumber NumFrames = (SoundDuration * FrameRate).FloorToFrame();
+		const bool bValidOffset = ensure(NumFrames > StartFrameOffset);
+
+		const FFrameNumber MaxFrame = NumFrames + NewRange.GetLowerBoundValue() - StartFrameOffset;
+		const FFrameNumber NewRangeUpperBound = NewRange.GetUpperBoundValue();
+
+		if (bValidOffset)
+		{
+			if (NewRangeUpperBound > MaxFrame)
+			{
+				ClampedRange = TRange<FFrameNumber>(NewRange.GetLowerBound(), TRangeBound<FFrameNumber>::Exclusive(MaxFrame));
+			}
+		}
+		else
+		{
+			return;
+		}
+	}
+
+	Super::SetRange(ClampedRange);
+}
 	
 void UMovieSceneAudioSection::TrimSection(FQualifiedFrameTime TrimTime, bool bTrimLeft, bool bDeleteKeys)
 {
@@ -421,10 +480,7 @@ void UMovieSceneAudioSection::SetSound(USoundBase* InSound)
 {
 	Sound = InSound;
 
-	// For a procedurally generated sound, only have the section loop if it's specifically a looping sound
-	// this is because there's no way to determine a duration for the sound so the user has no way of
-	// now when to end the region to avoid retriggering the one-shot procedural asset.
-	if (Sound->IsProcedurallyGenerated())
+	if (Sound)
 	{
 		bLooping = !Sound->IsOneShot();
 	}
diff --git a/Engine/Source/Runtime/MovieSceneTracks/Public/Sections/MovieSceneAudioSection.h b/Engine/Source/Runtime/MovieSceneTracks/Public/Sections/MovieSceneAudioSection.h
--- a/Engine/Source/Runtime/MovieSceneTracks/Public/Sections/MovieSceneAudioSection.h
+++ b/Engine/Source/Runtime/MovieSceneTracks/Public/Sections/MovieSceneAudioSection.h
@@ -145,6 +145,9 @@ public:
 
 	/** ~UObject interface */
 	MOVIESCENETRACKS_API virtual void PostLoad() override;
+#if WITH_EDITOR
+	MOVIESCENETRACKS_API virtual void PostEditChangeProperty(FPropertyChangedEvent& PropertyChangedEvent) override;
+#endif //WITH_EDITOR
 
 	/** Called when subtitles are sent to the SubtitleManager.  Set this delegate if you want to hijack the subtitles for other purposes */
 	void SetOnQueueSubtitles(const FOnQueueSubtitles& InOnQueueSubtitles)
@@ -191,6 +194,7 @@ public:
 
 	//~ UMovieSceneSection interface
 	MOVIESCENETRACKS_API virtual TOptional<TRange<FFrameNumber> > GetAutoSizeRange() const override;
+	MOVIESCENETRACKS_API virtual void SetRange(const TRange<FFrameNumber>& NewRange) override;
 	MOVIESCENETRACKS_API virtual void TrimSection(FQualifiedFrameTime TrimTime, bool bTrimLeft, bool bDeleteKeys) override;
 	MOVIESCENETRACKS_API virtual UMovieSceneSection* SplitSection(FQualifiedFrameTime SplitTime, bool bDeleteKeys) override;
 	MOVIESCENETRACKS_API virtual TOptional<FFrameTime> GetOffsetTime() const override;
