Sequencer: Remove DisplayName="Anim Sequence" for MovieSceneSkeletalAnimationSection because it conflicts with Anim Sequence

diff --git a/Engine/Source/Runtime/MovieSceneTracks/Public/Sections/MovieSceneSkeletalAnimationSection.h b/Engine/Source/Runtime/MovieSceneTracks/Public/Sections/MovieSceneSkeletalAnimationSection.h
--- a/Engine/Source/Runtime/MovieSceneTracks/Public/Sections/MovieSceneSkeletalAnimationSection.h
+++ b/Engine/Source/Runtime/MovieSceneTracks/Public/Sections/MovieSceneSkeletalAnimationSection.h
@@ -108,7 +108,7 @@ struct FMovieSceneSkeletalAnimationParams
 /**
  * Movie scene section that control skeletal animation
  */
-UCLASS( MinimalAPI, DisplayName="Anim Sequence" )
+UCLASS( MinimalAPI )
 class UMovieSceneSkeletalAnimationSection
 	: public UMovieSceneSection
 	, public IMovieSceneEntityProvider
