Suppress creation of destruction info for net addressable spawned actors

diff --git a/Engine/Plugins/Experimental/DaySequence/Source/DaySequence/Private/DaySequenceSpawnRegister.cpp b/Engine/Plugins/Experimental/DaySequence/Source/DaySequence/Private/DaySequenceSpawnRegister.cpp
--- a/Engine/Plugins/Experimental/DaySequence/Source/DaySequence/Private/DaySequenceSpawnRegister.cpp
+++ b/Engine/Plugins/Experimental/DaySequence/Source/DaySequence/Private/DaySequenceSpawnRegister.cpp
@@ -6,6 +6,7 @@
 #include "IMovieSceneObjectSpawner.h"
 #include "Modules/ModuleManager.h"
 #include "Bindings/MovieSceneSpawnableBinding.h"
+#include "Engine/NetDriver.h"
 
 FDaySequenceSpawnRegister::FDaySequenceSpawnRegister()
 {
@@ -43,6 +44,7 @@ void FDaySequenceSpawnRegister::DestroySpawnedObject(UObject& Object, UMovieScen
 		{
 			if (Object.IsA(MovieSceneObjectSpawner->GetSupportedTemplateType()))
 			{
+				UE::Net::FScopedIgnoreStaticActorDestruction ScopedIgnoreDestruction; // net addressable actors normally have a destruction info created for them, we don't want that as these actors were spawned at runtime
 				MovieSceneObjectSpawner->DestroySpawnedObject(Object);
 				return;
 			}
diff --git a/Engine/Source/Runtime/Engine/Classes/Engine/NetDriver.h b/Engine/Source/Runtime/Engine/Classes/Engine/NetDriver.h
--- a/Engine/Source/Runtime/Engine/Classes/Engine/NetDriver.h
+++ b/Engine/Source/Runtime/Engine/Classes/Engine/NetDriver.h
@@ -363,8 +363,8 @@ namespace UE::Net
 	class FScopedIgnoreStaticActorDestruction
 	{
 	public:
-		FScopedIgnoreStaticActorDestruction();
-		~FScopedIgnoreStaticActorDestruction();
+		ENGINE_API FScopedIgnoreStaticActorDestruction();
+		ENGINE_API ~FScopedIgnoreStaticActorDestruction();
 
 		UE_NONCOPYABLE(FScopedIgnoreStaticActorDestruction);
 
diff --git a/Engine/Source/Runtime/LevelSequence/Private/LevelSequenceSpawnRegister.cpp b/Engine/Source/Runtime/LevelSequence/Private/LevelSequenceSpawnRegister.cpp
--- a/Engine/Source/Runtime/LevelSequence/Private/LevelSequenceSpawnRegister.cpp
+++ b/Engine/Source/Runtime/LevelSequence/Private/LevelSequenceSpawnRegister.cpp
@@ -2,6 +2,7 @@
 
 #include "LevelSequenceSpawnRegister.h"
 #include "Engine/EngineTypes.h"
+#include "Engine/NetDriver.h"
 #include "MovieScene.h"
 #include "MovieSceneSequence.h"
 #include "LevelSequenceModule.h"
@@ -47,6 +48,7 @@ void FLevelSequenceSpawnRegister::DestroySpawnedObject(UObject& Object, UMovieSc
 		{
 			if (Object.IsA(MovieSceneObjectSpawner->GetSupportedTemplateType()))
 			{
+				UE::Net::FScopedIgnoreStaticActorDestruction ScopedIgnoreDestruction; // net addressable actors normally have a destruction info created for them, we don't want that as these actors were spawned at runtime
 				MovieSceneObjectSpawner->DestroySpawnedObject(Object);
 				return;
 			}
diff --git a/Engine/Source/Runtime/MovieSceneTracks/Private/Bindings/MovieSceneSpawnableActorBinding.cpp b/Engine/Source/Runtime/MovieSceneTracks/Private/Bindings/MovieSceneSpawnableActorBinding.cpp
--- a/Engine/Source/Runtime/MovieSceneTracks/Private/Bindings/MovieSceneSpawnableActorBinding.cpp
+++ b/Engine/Source/Runtime/MovieSceneTracks/Private/Bindings/MovieSceneSpawnableActorBinding.cpp
@@ -6,6 +6,7 @@
 #include "GameFramework/Actor.h"
 #include "Engine/World.h"
 #include "Engine/Level.h"
+#include "Engine/NetDriver.h"
 #include "IMovieSceneObjectSpawner.h"
 #include "Particles/ParticleSystemComponent.h"
 #include "EntitySystem/MovieSceneEntitySystemLinker.h"
@@ -211,6 +212,7 @@ void UMovieSceneSpawnableActorBindingBase::DestroySpawnedObjectInternal(UObject*
 	UWorld* World = Actor->GetWorld();
 	if (World)
 	{
+		UE::Net::FScopedIgnoreStaticActorDestruction ScopedIgnoreDestruction; // net addressable actors normally have a destruction info created for them, we don't want that as these actors were spawned at runtime
 		const bool bNetForce = false;
 		const bool bShouldModifyLevel = false;
 		World->DestroyActor(Actor, bNetForce, bShouldModifyLevel);
