Sequencer: fixes to audio evaluation system

- Don't "stop and play" just because the context has jumped... first check if we even need to play at all, and then decide between "play" and "stop and play" based on what's going on.
- Don't force a time-sync if we were paused or stopped, otherwise it can force us to restart playing just based on the time passed between the last time we moved the playhead, and the time we move it again.

diff --git a/Engine/Source/Runtime/MovieSceneTracks/Private/Systems/MovieSceneAudioSystem.cpp b/Engine/Source/Runtime/MovieSceneTracks/Private/Systems/MovieSceneAudioSystem.cpp
--- a/Engine/Source/Runtime/MovieSceneTracks/Private/Systems/MovieSceneAudioSystem.cpp
+++ b/Engine/Source/Runtime/MovieSceneTracks/Private/Systems/MovieSceneAudioSystem.cpp
@@ -265,13 +265,6 @@ struct FEvaluateAudio
 			return EAudioEvaluationType::Skip;
 		}
 
-		if (Context.HasJumped())
-		{
-			// If the status says we jumped, we always stop all sounds, then allow them to be played again 
-			// naturally if status == Playing (for example)
-			return EAudioEvaluationType::StopAndPlay;
-		}
-
 		if (
 			(Context.GetStatus() != EMovieScenePlayerStatus::Playing && 
 			 Context.GetStatus() != EMovieScenePlayerStatus::Scrubbing && 
@@ -283,6 +276,13 @@ struct FEvaluateAudio
 			return EAudioEvaluationType::Stop;
 		}
 
+		if (Context.HasJumped())
+		{
+			// If the status says we jumped, we always stop all sounds, then allow them to be played again 
+			// naturally if status == Playing (for example)
+			return EAudioEvaluationType::StopAndPlay;
+		}
+
 		return EAudioEvaluationType::Play;
 	}
 
@@ -659,14 +659,22 @@ private:
 		// If the audio component is not playing we (may) need a state change. If the audio component is playing
 		// the wrong sound then we need a state change. If the audio playback time is significantly out of sync 
 		// with the desired time then we need a state change.
-		const bool bAudioComponentIsPaused = AudioComponent.GetPlayState() == EAudioComponentPlayState::Paused;
+		const EAudioComponentPlayState AudioComponentPlayState = AudioComponent.GetPlayState();
+		const bool bAudioComponentIsPaused = (AudioComponentPlayState == EAudioComponentPlayState::Paused);
 		const bool bSoundsNeedPlaying = !AudioComponent.IsPlaying() || bAudioComponentIsPaused;
 		const bool bSoundNeedsStateChange =  AudioComponent.Sound != Sound;
 		bool bSoundNeedsTimeSync = false;
 
-		// Sync only if there is no time dilation because otherwise the system will constantly resync because audio 
+		// Sync only if:
+		// - We are playing, and
+		// - There is no time dilation because otherwise the system will constantly resync because audio 
 		// playback is not dilated and will never match the expected playback time.
 		const bool bDoTimeSync = 
+			(
+				AudioComponentPlayState == EAudioComponentPlayState::Playing ||
+				AudioComponentPlayState == EAudioComponentPlayState::FadingIn ||
+				AudioComponentPlayState == EAudioComponentPlayState::FadingOut
+			) &&
 			World && WorldSettings &&
 			(FMath::IsNearlyEqual(WorldSettings->GetEffectiveTimeDilation(), 1.f) ||
 			 !bIgnoreAudioSyncDuringWorldTimeDilationCVar);
