[Sequencer] double audio playing on certain audio clips

diff --git a/Engine/Source/Runtime/MovieSceneTracks/Private/Systems/MovieSceneAudioSystem.cpp b/Engine/Source/Runtime/MovieSceneTracks/Private/Systems/MovieSceneAudioSystem.cpp
--- a/Engine/Source/Runtime/MovieSceneTracks/Private/Systems/MovieSceneAudioSystem.cpp
+++ b/Engine/Source/Runtime/MovieSceneTracks/Private/Systems/MovieSceneAudioSystem.cpp
@@ -637,7 +637,9 @@ private:
 			else
 			{
 				const float Duration = MovieSceneHelpers::GetSoundDuration(Sound);
-				if (!AudioSection.GetLooping() && AudioTime > Duration && Duration != 0.f)
+				check(Duration >= 0.f);
+
+				if (!AudioSection.GetLooping() && AudioTime >= Duration && Duration != 0.f)
 				{
 					// If this is Non-ProcedurallyGenerated audio, and it's not looping then check to see if it needs to be stopped
 					if (AudioComponent.IsPlaying()) // IsPlaying() == playing, fading, or paused
