Sequencer: Use GetTransientPackage if the World doesn't exist (ie. Actor Sequence)

#Synced-CL 48071716

diff --git a/Engine/Source/Runtime/MovieSceneTracks/Private/Systems/MovieSceneAudioSystem.cpp b/Engine/Source/Runtime/MovieSceneTracks/Private/Systems/MovieSceneAudioSystem.cpp
--- a/Engine/Source/Runtime/MovieSceneTracks/Private/Systems/MovieSceneAudioSystem.cpp
+++ b/Engine/Source/Runtime/MovieSceneTracks/Private/Systems/MovieSceneAudioSystem.cpp
@@ -1279,7 +1279,12 @@ UMovieSceneAudioSystem::FAudioComponentEvaluationData* UMovieSceneAudioSystem::A
 
 #if WITH_EDITOR
 		static int32 ScrubSoundCounter = 0;
-		UScrubbedSound* ScrubbedSound = NewObject<UScrubbedSound>(World->GetCurrentLevel(), FName(*FString::Printf(TEXT("ScrubbedSound_Root_%i"), ScrubSoundCounter++)));
+		UObject* OuterObject = GetTransientPackage();
+		if (World && World->GetCurrentLevel())
+		{
+			OuterObject = World->GetCurrentLevel();
+		}
+		UScrubbedSound* ScrubbedSound = NewObject<UScrubbedSound>(OuterObject, FName(*FString::Printf(TEXT("ScrubbedSound_Root_%i"), ScrubSoundCounter++)));
 		if (!ScrubbedSound)
 		{
 			const int32 RowIndex = Section->GetRowIndex();
