Sequencer: Add the ability to animate lighting channels on primitive components

Added custom accessors for lighting channels

#Synced-CL 48303336

diff --git a/Engine/Source/Runtime/Engine/Classes/Engine/EngineTypes.h b/Engine/Source/Runtime/Engine/Classes/Engine/EngineTypes.h
--- a/Engine/Source/Runtime/Engine/Classes/Engine/EngineTypes.h
+++ b/Engine/Source/Runtime/Engine/Classes/Engine/EngineTypes.h
@@ -569,15 +569,15 @@ struct FLightingChannels
 	{}
 
 	/** Default channel for all primitives and lights. */
-	UPROPERTY(EditAnywhere, BlueprintReadOnly, Category=Channels)
+	UPROPERTY(Interp, EditAnywhere, BlueprintReadOnly, Category=Channels)
 	uint8 bChannel0:1;
 
 	/** First custom channel */
-	UPROPERTY(EditAnywhere, BlueprintReadOnly, Category=Channels)
+	UPROPERTY(Interp, EditAnywhere, BlueprintReadOnly, Category=Channels)
 	uint8 bChannel1:1;
 
 	/** Second custom channel */
-	UPROPERTY(EditAnywhere, BlueprintReadOnly, Category=Channels)
+	UPROPERTY(Interp, EditAnywhere, BlueprintReadOnly, Category=Channels)
 	uint8 bChannel2:1;
 };
 
diff --git a/Engine/Source/Runtime/MovieSceneTracks/Private/MovieSceneTracksCustomAccessors.cpp b/Engine/Source/Runtime/MovieSceneTracks/Private/MovieSceneTracksCustomAccessors.cpp
--- a/Engine/Source/Runtime/MovieSceneTracks/Private/MovieSceneTracksCustomAccessors.cpp
+++ b/Engine/Source/Runtime/MovieSceneTracks/Private/MovieSceneTracksCustomAccessors.cpp
@@ -10,6 +10,7 @@
 #include "Components/SkyAtmosphereComponent.h"
 #include "Components/SkyLightComponent.h"
 #include "Components/SkeletalMeshComponent.h"
+#include "Engine/EngineTypes.h"
 #include "Engine/SkeletalMesh.h"
 
 namespace UE::MovieScene
@@ -350,6 +351,42 @@ void SetMassScale(UObject* Object, float InMassScale)
 	PrimitiveComponent->SetMassScale(NAME_None, InMassScale);
 }
 
+bool GetLightingChannel0(const UObject* Object, FBoolPropertyTraits::FBoolMetaData InMetadata)
+{
+	const UPrimitiveComponent* PrimitiveComponent = CastChecked<const UPrimitiveComponent>(Object);
+	return PrimitiveComponent->LightingChannels.bChannel0;
+}
+
+void SetLightingChannel0(UObject* Object, FBoolPropertyTraits::FBoolMetaData InMetadata, bool bInLightingChannel0)
+{
+	UPrimitiveComponent* PrimitiveComponent = CastChecked<UPrimitiveComponent>(Object);
+	PrimitiveComponent->SetLightingChannels(bInLightingChannel0, PrimitiveComponent->LightingChannels.bChannel1, PrimitiveComponent->LightingChannels.bChannel2);
+}
+
+bool GetLightingChannel1(const UObject* Object, FBoolPropertyTraits::FBoolMetaData InMetadata)
+{
+	const UPrimitiveComponent* PrimitiveComponent = CastChecked<const UPrimitiveComponent>(Object);
+	return PrimitiveComponent->LightingChannels.bChannel1;
+}
+
+void SetLightingChannel1(UObject* Object, FBoolPropertyTraits::FBoolMetaData InMetadata, bool bInLightingChannel1)
+{
+	UPrimitiveComponent* PrimitiveComponent = CastChecked<UPrimitiveComponent>(Object);
+	PrimitiveComponent->SetLightingChannels(PrimitiveComponent->LightingChannels.bChannel0, bInLightingChannel1, PrimitiveComponent->LightingChannels.bChannel2);
+}
+
+bool GetLightingChannel2(const UObject* Object, FBoolPropertyTraits::FBoolMetaData InMetadata)
+{
+	const UPrimitiveComponent* PrimitiveComponent = CastChecked<const UPrimitiveComponent>(Object);
+	return PrimitiveComponent->LightingChannels.bChannel2;
+}
+
+void SetLightingChannel2(UObject* Object, FBoolPropertyTraits::FBoolMetaData InMetadata, bool bInLightingChannel2)
+{
+	UPrimitiveComponent* PrimitiveComponent = CastChecked<UPrimitiveComponent>(Object);
+	PrimitiveComponent->SetLightingChannels(PrimitiveComponent->LightingChannels.bChannel0, PrimitiveComponent->LightingChannels.bChannel1, bInLightingChannel2);
+}
+
 float GetExponentialHeightFogComponentFogDensity(const UObject* Object)
 {
 	const UExponentialHeightFogComponent* ExponentialHeightFogComponent = CastChecked<const UExponentialHeightFogComponent>(Object);
@@ -652,6 +689,21 @@ void InitializeMovieSceneTracksAccessors(FMovieSceneTracksComponentTypes* Tracks
 		UPrimitiveComponent::StaticClass(), *MassScalePath,
 		GetMassScale, SetMassScale);
 
+	const FString LightingChannel0Path = FString::Printf(TEXT("%s.%s"), GET_MEMBER_NAME_STRING_CHECKED(UPrimitiveComponent, LightingChannels), GET_MEMBER_NAME_STRING_CHECKED(FLightingChannels, bChannel0));
+	TracksComponents->Accessors.Bool.Add(
+		UPrimitiveComponent::StaticClass(), *LightingChannel0Path,
+		GetLightingChannel0, SetLightingChannel0);
+
+	const FString LightingChannel1Path = FString::Printf(TEXT("%s.%s"), GET_MEMBER_NAME_STRING_CHECKED(UPrimitiveComponent, LightingChannels), GET_MEMBER_NAME_STRING_CHECKED(FLightingChannels, bChannel1));
+	TracksComponents->Accessors.Bool.Add(
+		UPrimitiveComponent::StaticClass(), *LightingChannel1Path,
+		GetLightingChannel1, SetLightingChannel1);
+
+	const FString LightingChannel2Path = FString::Printf(TEXT("%s.%s"), GET_MEMBER_NAME_STRING_CHECKED(UPrimitiveComponent, LightingChannels), GET_MEMBER_NAME_STRING_CHECKED(FLightingChannels, bChannel2));
+	TracksComponents->Accessors.Bool.Add(
+		UPrimitiveComponent::StaticClass(), *LightingChannel2Path,
+		GetLightingChannel2, SetLightingChannel2);
+
 	// ExponentialHeightFogComponent
 	TracksComponents->Accessors.Float.Add(
 			UExponentialHeightFogComponent::StaticClass(), GET_MEMBER_NAME_CHECKED(UExponentialHeightFogComponent, FogDensity),
