Sequencer: Whether using anim blueprints or custom mode, previous anim instance gets force re-initialized after sequence is over, potentially resetting data stored in the AnimInstance unnecessarily.

Previously before CL 44346141 this would only occur if the skeleton changed, but that CL attempted to address an issue where if Custom Mode was used, the previous anim instance was getting uninitialized anyway, and would now be stuck in an uninitalized case. That CL was a bit overzealous however, as this force re-initialization was even occurring for cases where animation blueprints were used and we weren't switching anim mode or anim instances.

This change only re-initializes if Sequencer changed (and is now restoring) the anim instance.

#Synced-CL 48079442

diff --git a/Engine/Plugins/Experimental/MovieSceneAnimMixer/Source/MovieSceneAnimMixer/Private/Systems/MovieSceneAnimInstanceTargetSystem.cpp b/Engine/Plugins/Experimental/MovieSceneAnimMixer/Source/MovieSceneAnimMixer/Private/Systems/MovieSceneAnimInstanceTargetSystem.cpp
--- a/Engine/Plugins/Experimental/MovieSceneAnimMixer/Source/MovieSceneAnimMixer/Private/Systems/MovieSceneAnimInstanceTargetSystem.cpp
+++ b/Engine/Plugins/Experimental/MovieSceneAnimMixer/Source/MovieSceneAnimMixer/Private/Systems/MovieSceneAnimInstanceTargetSystem.cpp
@@ -342,13 +342,30 @@ namespace UE::MovieScene
 			TStrongObjectPtr<UAnimInstance> PreviousAnimInstance = InOutCachedValue.CachedAnimInstance.Pin();
 			if (PreviousAnimInstance && IsValid(PreviousAnimInstance.Get()) && PreviousAnimInstance->GetSkelMeshComponent() == Component)
 			{
-				Component->AnimScriptInstance = PreviousAnimInstance.Get();
+				bool bNeedsReInit = false;
+
+				if (Component->AnimScriptInstance != PreviousAnimInstance.Get())
+				{
+					Component->AnimScriptInstance = PreviousAnimInstance.Get();
+					bNeedsReInit = true;
+				}
+
 				InOutCachedValue.CachedAnimInstance.Reset();
-				if (Component->AnimScriptInstance && Component->GetSkeletalMeshAsset() && Component->AnimScriptInstance->CurrentSkeleton != Component->GetSkeletalMeshAsset()->GetSkeleton())
+
+				if (Component->AnimScriptInstance)
 				{
 					//the skeleton may have changed so need to recalc required bones as needed.
-					Component->AnimScriptInstance->CurrentSkeleton = Component->GetSkeletalMeshAsset()->GetSkeleton();
-					//Need at least RecalcRequiredbones and UpdateMorphTargetrs
+					if (Component->GetSkeletalMeshAsset() && Component->AnimScriptInstance->CurrentSkeleton != Component->GetSkeletalMeshAsset()->GetSkeleton())
+					{
+						Component->AnimScriptInstance->CurrentSkeleton = Component->GetSkeletalMeshAsset()->GetSkeleton();
+						bNeedsReInit = true;
+					}
+				}
+
+				if (bNeedsReInit)
+				{
+					// When Changing anim instances or skeletons, RecalcRequiredBones and UpdateMorphTargets need to be called, and InitializeAnimation needs to be called regardless, when changing script instances
+					// Playing montages does not work on uninitialized anim script instances with the EarlyOutMontageWhenUninitialized cvar enabled.
 					Component->InitializeAnimScriptInstance(true);
 				}
 			}
diff --git a/Engine/Source/Runtime/MovieSceneTracks/Private/Systems/MovieSceneSkeletalAnimationSystem.cpp b/Engine/Source/Runtime/MovieSceneTracks/Private/Systems/MovieSceneSkeletalAnimationSystem.cpp
--- a/Engine/Source/Runtime/MovieSceneTracks/Private/Systems/MovieSceneSkeletalAnimationSystem.cpp
+++ b/Engine/Source/Runtime/MovieSceneTracks/Private/Systems/MovieSceneSkeletalAnimationSystem.cpp
@@ -145,16 +145,32 @@ struct FPreAnimatedSkeletalAnimationTraits : FBoundObjectPreAnimatedStateTraits
 		TStrongObjectPtr<UAnimInstance> PreviousAnimInstance = InOutCachedValue.CachedAnimInstance.Pin();
 		if (PreviousAnimInstance && IsValid(PreviousAnimInstance.Get()) && PreviousAnimInstance->GetSkelMeshComponent() == Component)
 		{
-			Component->AnimScriptInstance = PreviousAnimInstance.Get();
+			bool bNeedsReInit = false;
+
+			if (Component->AnimScriptInstance != PreviousAnimInstance.Get())
+			{
+				Component->AnimScriptInstance = PreviousAnimInstance.Get();
+				bNeedsReInit = true;
+			}
+
 			InOutCachedValue.CachedAnimInstance.Reset();
-			if (Component->AnimScriptInstance && Component->GetSkeletalMeshAsset() && Component->AnimScriptInstance->CurrentSkeleton != Component->GetSkeletalMeshAsset()->GetSkeleton())
+
+			if (Component->AnimScriptInstance)
 			{
 				//the skeleton may have changed so need to recalc required bones as needed.
-				Component->AnimScriptInstance->CurrentSkeleton = Component->GetSkeletalMeshAsset()->GetSkeleton();
+				if (Component->GetSkeletalMeshAsset() && Component->AnimScriptInstance->CurrentSkeleton != Component->GetSkeletalMeshAsset()->GetSkeleton())
+				{
+					Component->AnimScriptInstance->CurrentSkeleton = Component->GetSkeletalMeshAsset()->GetSkeleton();
+					bNeedsReInit = true;
+				}
+			}
+
+			if (bNeedsReInit)
+			{
+				// When Changing anim instances or skeletons, RecalcRequiredBones and UpdateMorphTargets need to be called, and InitializeAnimation needs to be called regardless, when changing script instances
+				// Playing montages does not work on uninitialized anim script instances with the EarlyOutMontageWhenUninitialized cvar enabled.
+				Component->InitializeAnimScriptInstance(true);
 			}
-			// When Changing skeletons RecalcRequiredBones and UpdateMorphTargets need to be called, and InitializeAnimation needs to be called regardless, when changing script instances
-			// Playing montages does not work on uninitialized anim script instances with the EarlyOutMontageWhenUninitialized cvar enabled.
-			Component->InitializeAnimScriptInstance(true);
 		}
 
 		// Restore pose after unbinding to force the restored pose
