Sequencer: Follow up to 48402553 add the ability to animate lighting channels on light components

#Synced-CL 48303336

diff --git a/Engine/Source/Runtime/MovieSceneTracks/Private/MovieSceneTracksCustomAccessors.cpp b/Engine/Source/Runtime/MovieSceneTracks/Private/MovieSceneTracksCustomAccessors.cpp
--- a/Engine/Source/Runtime/MovieSceneTracks/Private/MovieSceneTracksCustomAccessors.cpp
+++ b/Engine/Source/Runtime/MovieSceneTracks/Private/MovieSceneTracksCustomAccessors.cpp
@@ -353,38 +353,101 @@ void SetMassScale(UObject* Object, float InMassScale)
 
 bool GetLightingChannel0(const UObject* Object, FBoolPropertyTraits::FBoolMetaData InMetadata)
 {
-	const UPrimitiveComponent* PrimitiveComponent = CastChecked<const UPrimitiveComponent>(Object);
-	return PrimitiveComponent->LightingChannels.bChannel0;
+	if (Object->IsA<ULightComponent>())
+	{
+		const ULightComponent* LightComponent = CastChecked<const ULightComponent>(Object);
+		return LightComponent->LightingChannels.bChannel0;
+	}
+	else if (Object->IsA<UPrimitiveComponent>())
+	{
+		const UPrimitiveComponent* PrimitiveComponent = CastChecked<const UPrimitiveComponent>(Object);
+		return PrimitiveComponent->LightingChannels.bChannel0;
+	}
+	ensureMsgf(false, TEXT("GetLightingChannel0 with unknown object %s"), *GetNameSafe(Object));
+	return false;
 }
 
 void SetLightingChannel0(UObject* Object, FBoolPropertyTraits::FBoolMetaData InMetadata, bool bInLightingChannel0)
 {
-	UPrimitiveComponent* PrimitiveComponent = CastChecked<UPrimitiveComponent>(Object);
-	PrimitiveComponent->SetLightingChannels(bInLightingChannel0, PrimitiveComponent->LightingChannels.bChannel1, PrimitiveComponent->LightingChannels.bChannel2);
+	if (Object->IsA<ULightComponent>())
+	{
+		ULightComponent* LightComponent = CastChecked<ULightComponent>(Object);
+		LightComponent->SetLightingChannels(bInLightingChannel0, LightComponent->LightingChannels.bChannel1, LightComponent->LightingChannels.bChannel2);
+		return;
+	}
+	else if (Object->IsA<UPrimitiveComponent>())
+	{
+		UPrimitiveComponent* PrimitiveComponent = CastChecked<UPrimitiveComponent>(Object);
+		PrimitiveComponent->SetLightingChannels(bInLightingChannel0, PrimitiveComponent->LightingChannels.bChannel1, PrimitiveComponent->LightingChannels.bChannel2);
+		return;
+	}
+	ensureMsgf(false, TEXT("SetLightingChannel0 with unknown object %s"), *GetNameSafe(Object));
 }
 
 bool GetLightingChannel1(const UObject* Object, FBoolPropertyTraits::FBoolMetaData InMetadata)
 {
-	const UPrimitiveComponent* PrimitiveComponent = CastChecked<const UPrimitiveComponent>(Object);
-	return PrimitiveComponent->LightingChannels.bChannel1;
+	if (Object->IsA<ULightComponent>())
+	{
+		const ULightComponent* LightComponent = CastChecked<const ULightComponent>(Object);
+		return LightComponent->LightingChannels.bChannel1;
+	}
+	else if (Object->IsA<UPrimitiveComponent>())
+	{
+		const UPrimitiveComponent* PrimitiveComponent = CastChecked<const UPrimitiveComponent>(Object);
+		return PrimitiveComponent->LightingChannels.bChannel1;
+	}
+	ensureMsgf(false, TEXT("GetLightingChannel1 with unknown object %s"), *GetNameSafe(Object));
+	return false;
 }
 
 void SetLightingChannel1(UObject* Object, FBoolPropertyTraits::FBoolMetaData InMetadata, bool bInLightingChannel1)
 {
-	UPrimitiveComponent* PrimitiveComponent = CastChecked<UPrimitiveComponent>(Object);
-	PrimitiveComponent->SetLightingChannels(PrimitiveComponent->LightingChannels.bChannel0, bInLightingChannel1, PrimitiveComponent->LightingChannels.bChannel2);
+	if (Object->IsA<ULightComponent>())
+	{
+		ULightComponent* LightComponent = CastChecked<ULightComponent>(Object);
+		LightComponent->SetLightingChannels(LightComponent->LightingChannels.bChannel0, bInLightingChannel1, LightComponent->LightingChannels.bChannel2);
+		return;
+	}
+	else if (Object->IsA<UPrimitiveComponent>())
+	{
+		UPrimitiveComponent* PrimitiveComponent = CastChecked<UPrimitiveComponent>(Object);
+		PrimitiveComponent->SetLightingChannels(PrimitiveComponent->LightingChannels.bChannel0, bInLightingChannel1, PrimitiveComponent->LightingChannels.bChannel2);
+		return;
+	}
+	ensureMsgf(false, TEXT("SetLightingChannel1 with unknown object %s"), *GetNameSafe(Object));
 }
 
 bool GetLightingChannel2(const UObject* Object, FBoolPropertyTraits::FBoolMetaData InMetadata)
 {
-	const UPrimitiveComponent* PrimitiveComponent = CastChecked<const UPrimitiveComponent>(Object);
-	return PrimitiveComponent->LightingChannels.bChannel2;
+	if (Object->IsA<ULightComponent>())
+	{
+		const ULightComponent* LightComponent = CastChecked<const ULightComponent>(Object);
+		return LightComponent->LightingChannels.bChannel2;
+	}
+	else if (Object->IsA<UPrimitiveComponent>())
+	{
+		const UPrimitiveComponent* PrimitiveComponent = CastChecked<const UPrimitiveComponent>(Object);
+		return PrimitiveComponent->LightingChannels.bChannel2;
+	}
+	ensureMsgf(false, TEXT("GetLightingChannel2 with unknown object %s"), *GetNameSafe(Object));
+	return false;
 }
 
 void SetLightingChannel2(UObject* Object, FBoolPropertyTraits::FBoolMetaData InMetadata, bool bInLightingChannel2)
 {
-	UPrimitiveComponent* PrimitiveComponent = CastChecked<UPrimitiveComponent>(Object);
-	PrimitiveComponent->SetLightingChannels(PrimitiveComponent->LightingChannels.bChannel0, PrimitiveComponent->LightingChannels.bChannel1, bInLightingChannel2);
+	if (Object->IsA<ULightComponent>())
+	{
+		ULightComponent* LightComponent = CastChecked<ULightComponent>(Object);
+		LightComponent->SetLightingChannels(LightComponent->LightingChannels.bChannel0, LightComponent->LightingChannels.bChannel1, bInLightingChannel2);
+		return;
+	}
+	else if (Object->IsA<UPrimitiveComponent>())
+	{
+		UPrimitiveComponent* PrimitiveComponent = CastChecked<UPrimitiveComponent>(Object);
+		PrimitiveComponent->SetLightingChannels(PrimitiveComponent->LightingChannels.bChannel0, PrimitiveComponent->LightingChannels.bChannel1, bInLightingChannel2);
+		return;
+	}
+	ensureMsgf(false, TEXT("SetLightingChannel2 with unknown object %s"), *GetNameSafe(Object));
 }
 
 float GetExponentialHeightFogComponentFogDensity(const UObject* Object)
@@ -632,6 +695,21 @@ void InitializeMovieSceneTracksAccessors(FMovieSceneTracksComponentTypes* Tracks
 			ULightComponent::StaticClass(), GET_MEMBER_NAME_CHECKED(ULightComponent, IndirectLightingIntensity),
 			GetLightComponentIndirectLightingIntensity, SetLightComponentIndirectLightingIntensity);
 
+	const FString LightingChannel0Path = FString::Printf(TEXT("%s.%s"), GET_MEMBER_NAME_STRING_CHECKED(ULightComponent, LightingChannels), GET_MEMBER_NAME_STRING_CHECKED(FLightingChannels, bChannel0));
+	TracksComponents->Accessors.Bool.Add(
+		ULightComponent::StaticClass(), *LightingChannel0Path,
+		GetLightingChannel0, SetLightingChannel0);
+
+	const FString LightingChannel1Path = FString::Printf(TEXT("%s.%s"), GET_MEMBER_NAME_STRING_CHECKED(ULightComponent, LightingChannels), GET_MEMBER_NAME_STRING_CHECKED(FLightingChannels, bChannel1));
+	TracksComponents->Accessors.Bool.Add(
+		ULightComponent::StaticClass(), *LightingChannel1Path,
+		GetLightingChannel1, SetLightingChannel1);
+
+	const FString LightingChannel2Path = FString::Printf(TEXT("%s.%s"), GET_MEMBER_NAME_STRING_CHECKED(ULightComponent, LightingChannels), GET_MEMBER_NAME_STRING_CHECKED(FLightingChannels, bChannel2));
+	TracksComponents->Accessors.Bool.Add(
+		ULightComponent::StaticClass(), *LightingChannel2Path,
+		GetLightingChannel2, SetLightingChannel2);
+
 	// DirectionalLightComponent
 	TracksComponents->Accessors.Float.Add(
 			UDirectionalLightComponent::StaticClass(), GET_MEMBER_NAME_CHECKED(UDirectionalLightComponent, LightSourceAngle),
@@ -689,19 +767,19 @@ void InitializeMovieSceneTracksAccessors(FMovieSceneTracksComponentTypes* Tracks
 		UPrimitiveComponent::StaticClass(), *MassScalePath,
 		GetMassScale, SetMassScale);
 
-	const FString LightingChannel0Path = FString::Printf(TEXT("%s.%s"), GET_MEMBER_NAME_STRING_CHECKED(UPrimitiveComponent, LightingChannels), GET_MEMBER_NAME_STRING_CHECKED(FLightingChannels, bChannel0));
+	const FString PrimitiveLightingChannel0Path = FString::Printf(TEXT("%s.%s"), GET_MEMBER_NAME_STRING_CHECKED(UPrimitiveComponent, LightingChannels), GET_MEMBER_NAME_STRING_CHECKED(FLightingChannels, bChannel0));
 	TracksComponents->Accessors.Bool.Add(
-		UPrimitiveComponent::StaticClass(), *LightingChannel0Path,
+		UPrimitiveComponent::StaticClass(), *PrimitiveLightingChannel0Path,
 		GetLightingChannel0, SetLightingChannel0);
 
-	const FString LightingChannel1Path = FString::Printf(TEXT("%s.%s"), GET_MEMBER_NAME_STRING_CHECKED(UPrimitiveComponent, LightingChannels), GET_MEMBER_NAME_STRING_CHECKED(FLightingChannels, bChannel1));
+	const FString PrimitiveLightingChannel1Path = FString::Printf(TEXT("%s.%s"), GET_MEMBER_NAME_STRING_CHECKED(UPrimitiveComponent, LightingChannels), GET_MEMBER_NAME_STRING_CHECKED(FLightingChannels, bChannel1));
 	TracksComponents->Accessors.Bool.Add(
-		UPrimitiveComponent::StaticClass(), *LightingChannel1Path,
+		UPrimitiveComponent::StaticClass(), *PrimitiveLightingChannel1Path,
 		GetLightingChannel1, SetLightingChannel1);
 
-	const FString LightingChannel2Path = FString::Printf(TEXT("%s.%s"), GET_MEMBER_NAME_STRING_CHECKED(UPrimitiveComponent, LightingChannels), GET_MEMBER_NAME_STRING_CHECKED(FLightingChannels, bChannel2));
+	const FString PrimitiveLightingChannel2Path = FString::Printf(TEXT("%s.%s"), GET_MEMBER_NAME_STRING_CHECKED(UPrimitiveComponent, LightingChannels), GET_MEMBER_NAME_STRING_CHECKED(FLightingChannels, bChannel2));
 	TracksComponents->Accessors.Bool.Add(
-		UPrimitiveComponent::StaticClass(), *LightingChannel2Path,
+		UPrimitiveComponent::StaticClass(), *PrimitiveLightingChannel2Path,
 		GetLightingChannel2, SetLightingChannel2);
 
 	// ExponentialHeightFogComponent
