Sequencer: fixes to audio evaluation system

- Don't "stop and play" just because the context has jumped... first check if we even need to play at all, and then decide between "play" and "stop and play" based on what's going on.
- Don't force a time-sync if we were paused or stopped, otherwise it can force us to restart playing just based on the time passed between the last time we moved the playhead, and the time we move it again.

#rnx
#Synced-CL 48343954

diff --git a/Engine/Source/Runtime/MovieSceneTracks/Private/Systems/MovieSceneAudioSystem.cpp b/Engine/Source/Runtime/MovieSceneTracks/Private/Systems/MovieSceneAudioSystem.cpp
--- a/Engine/Source/Runtime/MovieSceneTracks/Private/Systems/MovieSceneAudioSystem.cpp
+++ b/Engine/Source/Runtime/MovieSceneTracks/Private/Systems/MovieSceneAudioSystem.cpp
@@ -664,10 +664,12 @@ private:
 		// the wrong sound then we need a state change. If the audio playback time is significantly out of sync 
 		// with the desired time then we need a state change.
 		const EAudioComponentPlayState AudioComponentPlayState = AudioComponent.GetPlayState();
-		const bool bSoundsNeedPlaying = !AudioComponent.IsPlaying();
+		const bool bAudioComponentIsPaused = (AudioComponentPlayState == EAudioComponentPlayState::Paused);
+		const bool bSoundsNeedPlaying = !AudioComponent.IsPlaying() || bAudioComponentIsPaused;
 		const bool bSoundNeedsStateChange =  AudioComponent.Sound != Sound;
 		bool bSoundNeedsTimeSync = false;
 
+
 		// Sync only if:
 		// - We are playing, and
 		// - There is no time dilation because otherwise the system will constantly resync because audio 
