Sequencer: upgrade the Actor Reference track to use the ECS evaluation, and move/rename its channel type into a separate file.

#Synced-CL 48619321

diff --git a/Engine/Config/BaseEngine.ini b/Engine/Config/BaseEngine.ini
--- a/Engine/Config/BaseEngine.ini
+++ b/Engine/Config/BaseEngine.ini
@@ -1382,6 +1382,8 @@ ClassName=/Script/OnlineSubsystemUtils.OnlineEngineInterfaceImpl
 +ClassRedirects=(OldName="/Script/MovieSceneTextTrack.TextChannelEvaluatorSystem",NewName="/Script/MovieSceneTracks.TextChannelEvaluatorSystem")
 +ClassRedirects=(OldName="/Script/MovieSceneTextTrackEditor.MovieSceneTextKeyStruct",NewName="/Script/MovieSceneTools.MovieSceneTextKeyStruct")
 
++StructRedirects=(OldName="/Script/MovieSceneTracks.MovieSceneActorReferenceData",NewName="/Script/MovieSceneTracks.MovieSceneActorReferenceChannel")
+
 +PropertyRedirects=(OldName="/Script/MovieScene.MovieSceneSectionTimingParametersSeconds.bClamp",NewName="MovieSceneSectionTimingParametersSeconds.bClampToInnerRange")
 +PropertyRedirects=(OldName="/Script/MovieScene.FMovieSceneSectionTimingParametersFrames.bClamp",NewName="FMovieSceneSectionTimingParametersFrames.bClampToInnerRange")
 
diff --git a/Engine/Plugins/MovieScene/SequencerScripting/Source/SequencerScripting/Private/ExtensionLibraries/MovieSceneSectionExtensions.cpp b/Engine/Plugins/MovieScene/SequencerScripting/Source/SequencerScripting/Private/ExtensionLibraries/MovieSceneSectionExtensions.cpp
--- a/Engine/Plugins/MovieScene/SequencerScripting/Source/SequencerScripting/Private/ExtensionLibraries/MovieSceneSectionExtensions.cpp
+++ b/Engine/Plugins/MovieScene/SequencerScripting/Source/SequencerScripting/Private/ExtensionLibraries/MovieSceneSectionExtensions.cpp
@@ -492,7 +492,7 @@ TArray<UMovieSceneScriptingChannel*> UMovieSceneSectionExtensions::GetAllChannel
 	GetScriptingChannelsForChannel<FMovieSceneStringChannel, UMovieSceneScriptingStringChannel>(ChannelProxy, Sequence, Section, Channels);
 	GetScriptingChannelsForChannel<FMovieSceneTextChannel, UMovieSceneScriptingTextChannel>(ChannelProxy, Sequence, Section, Channels);
 	GetScriptingChannelsForChannel<FMovieSceneEventChannel, UMovieSceneScriptingEventChannel>(ChannelProxy, Sequence, Section, Channels);
-	GetScriptingChannelsForChannel<FMovieSceneActorReferenceData, UMovieSceneScriptingActorReferenceChannel>(ChannelProxy, Sequence, Section, Channels);
+	GetScriptingChannelsForChannel<FMovieSceneActorReferenceChannel, UMovieSceneScriptingActorReferenceChannel>(ChannelProxy, Sequence, Section, Channels);
 	GetScriptingChannelsForChannel<FMovieSceneObjectPathChannel, UMovieSceneScriptingObjectPathChannel>(ChannelProxy, Sequence, Section, Channels);
 	GetScriptingChannelsForChannel<FMovieSceneParticleChannel, UMovieSceneScriptingParticleChannel>(ChannelProxy, Sequence, Section, Channels);
 	
@@ -539,7 +539,7 @@ TArray<UMovieSceneScriptingChannel*> UMovieSceneSectionExtensions::GetChannelsBy
 	else if (ChannelType == UMovieSceneScriptingStringChannel::StaticClass())			{ GetScriptingChannelsForChannel<FMovieSceneStringChannel, UMovieSceneScriptingStringChannel>(ChannelProxy, Sequence, Section, Channels); }
 	else if (ChannelType == UMovieSceneScriptingTextChannel::StaticClass())				{ GetScriptingChannelsForChannel<FMovieSceneTextChannel, UMovieSceneScriptingTextChannel>(ChannelProxy, Sequence, Section, Channels); }	
 	else if (ChannelType == UMovieSceneScriptingEventChannel::StaticClass())			{ GetScriptingChannelsForChannel<FMovieSceneEventChannel, UMovieSceneScriptingEventChannel>(ChannelProxy, Sequence, Section, Channels); }
-	else if (ChannelType == UMovieSceneScriptingActorReferenceChannel::StaticClass()) { GetScriptingChannelsForChannel<FMovieSceneActorReferenceData, UMovieSceneScriptingActorReferenceChannel>(ChannelProxy, Sequence, Section, Channels); }
+	else if (ChannelType == UMovieSceneScriptingActorReferenceChannel::StaticClass()) { GetScriptingChannelsForChannel<FMovieSceneActorReferenceChannel, UMovieSceneScriptingActorReferenceChannel>(ChannelProxy, Sequence, Section, Channels); }
 	else if (ChannelType == UMovieSceneScriptingObjectPathChannel::StaticClass()) { GetScriptingChannelsForChannel<FMovieSceneObjectPathChannel, UMovieSceneScriptingObjectPathChannel>(ChannelProxy, Sequence, Section, Channels); }
 	else if (ChannelType == UMovieSceneScriptingParticleChannel::StaticClass()) { GetScriptingChannelsForChannel<FMovieSceneParticleChannel, UMovieSceneScriptingParticleChannel>(ChannelProxy, Sequence, Section, Channels); }
 	else
@@ -601,9 +601,9 @@ UMovieSceneScriptingChannel* UMovieSceneSectionExtensions::GetChannel(UMovieScen
 	{
 		return GetScriptingChannelForChannel<FMovieSceneEventChannel, UMovieSceneScriptingEventChannel>(MovieSceneChannel, ChannelProxy, ChanneHandle, Sequence, Section, ChannelName);
 	}
-	else if (ChanneHandle.GetChannelTypeName() == FMovieSceneActorReferenceData::StaticStruct()->GetFName())
+	else if (ChanneHandle.GetChannelTypeName() == FMovieSceneActorReferenceChannel::StaticStruct()->GetFName())
 	{
-		return GetScriptingChannelForChannel<FMovieSceneActorReferenceData, UMovieSceneScriptingActorReferenceChannel>(MovieSceneChannel, ChannelProxy, ChanneHandle, Sequence, Section, ChannelName);
+		return GetScriptingChannelForChannel<FMovieSceneActorReferenceChannel, UMovieSceneScriptingActorReferenceChannel>(MovieSceneChannel, ChannelProxy, ChanneHandle, Sequence, Section, ChannelName);
 	}
 	else if (ChanneHandle.GetChannelTypeName() == FMovieSceneObjectPathChannel::StaticStruct()->GetFName())
 	{
diff --git a/Engine/Plugins/MovieScene/SequencerScripting/Source/SequencerScripting/Private/KeysAndChannels/MovieSceneScriptingActorReference.h b/Engine/Plugins/MovieScene/SequencerScripting/Source/SequencerScripting/Private/KeysAndChannels/MovieSceneScriptingActorReference.h
--- a/Engine/Plugins/MovieScene/SequencerScripting/Source/SequencerScripting/Private/KeysAndChannels/MovieSceneScriptingActorReference.h
+++ b/Engine/Plugins/MovieScene/SequencerScripting/Source/SequencerScripting/Private/KeysAndChannels/MovieSceneScriptingActorReference.h
@@ -18,7 +18,7 @@ template <typename ChannelType, typename ScriptingKeyType, typename ScriptingKey
 * Stores a reference to the data so changes to this class are forwarded onto the underlying data structures.
 */
 UCLASS(BlueprintType)
-class UMovieSceneScriptingActorReferenceKey : public UMovieSceneScriptingKey, public TMovieSceneScriptingKey<FMovieSceneActorReferenceData, FMovieSceneActorReferenceKey>
+class UMovieSceneScriptingActorReferenceKey : public UMovieSceneScriptingKey, public TMovieSceneScriptingKey<FMovieSceneActorReferenceChannel, FMovieSceneActorReferenceKey>
 {
 	GENERATED_BODY()
 public:
@@ -67,7 +67,7 @@ class UMovieSceneScriptingActorReferenceChannel : public UMovieSceneScriptingCha
 {
 	GENERATED_BODY()
 
-	using Impl = TMovieSceneScriptingChannel<FMovieSceneActorReferenceData, UMovieSceneScriptingActorReferenceKey, FMovieSceneActorReferenceKey>;
+	using Impl = TMovieSceneScriptingChannel<FMovieSceneActorReferenceChannel, UMovieSceneScriptingActorReferenceKey, FMovieSceneActorReferenceKey>;
 
 public:
 	/**
@@ -135,8 +135,8 @@ public:
 	UFUNCTION(BlueprintCallable, Category = "Sequencer|Keys", meta = (DisplayName = "Get Default (Actor Reference)"))
 	FMovieSceneObjectBindingID GetDefault() const
 	{
-		// FMovieSceneActorReferenceData doesn't implement GetDefault via TOptional, so we're wrapping this function by hand as well.
-		FMovieSceneActorReferenceData* Channel = ChannelHandle.Get();
+		// FMovieSceneActorReferenceChannel doesn't implement GetDefault via TOptional, so we're wrapping this function by hand as well.
+		FMovieSceneActorReferenceChannel* Channel = ChannelHandle.Get();
 		if (Channel)
 		{
 			FMovieSceneActorReferenceKey ReferenceKey = Channel->GetDefault();
@@ -158,9 +158,9 @@ public:
 			OwningSection->Modify();
 		}
 
-		// FMovieSceneActorReferenceData doesn't implement RemoveDefault, instead it implements ClearDefault(). Wrapping this function by hand,
+		// FMovieSceneActorReferenceChannel doesn't implement RemoveDefault, instead it implements ClearDefault(). Wrapping this function by hand,
 		// and not falling back to the template as a result, but keeping the same function name so it is consistent with the other scripting channels.
-		FMovieSceneActorReferenceData* Channel = ChannelHandle.Get();
+		FMovieSceneActorReferenceChannel* Channel = ChannelHandle.Get();
 		if (Channel)
 		{
 			Channel->ClearDefault();
@@ -204,5 +204,5 @@ public:
 public:
 	TWeakObjectPtr<UMovieSceneSequence> OwningSequence;
 	TWeakObjectPtr<UMovieSceneSection> OwningSection;
-	TMovieSceneChannelHandle<FMovieSceneActorReferenceData> ChannelHandle;
+	TMovieSceneChannelHandle<FMovieSceneActorReferenceChannel> ChannelHandle;
 };
diff --git a/Engine/Source/Editor/MovieSceneTools/Private/Channels/BuiltInChannelEditors.cpp b/Engine/Source/Editor/MovieSceneTools/Private/Channels/BuiltInChannelEditors.cpp
--- a/Engine/Source/Editor/MovieSceneTools/Private/Channels/BuiltInChannelEditors.cpp
+++ b/Engine/Source/Editor/MovieSceneTools/Private/Channels/BuiltInChannelEditors.cpp
@@ -233,7 +233,7 @@ FKeyHandle AddOrUpdateKey(FMovieSceneTimeWarpChannel* Channel, UMovieSceneSectio
 	return AddOrUpdateKey(Channel, SectionToKey, ExternalValue, InTime, Sequencer, InObjectBindingID, PropertyBindings);
 }
 
-FKeyHandle AddOrUpdateKey(FMovieSceneActorReferenceData* Channel, UMovieSceneSection* SectionToKey, FFrameNumber InTime, ISequencer& Sequencer, const FGuid& InObjectBindingID, FTrackInstancePropertyBindings* PropertyBindings)
+FKeyHandle AddOrUpdateKey(FMovieSceneActorReferenceChannel* Channel, UMovieSceneSection* SectionToKey, FFrameNumber InTime, ISequencer& Sequencer, const FGuid& InObjectBindingID, FTrackInstancePropertyBindings* PropertyBindings)
 {
 	if (PropertyBindings && InObjectBindingID.IsValid())
 	{
@@ -307,7 +307,7 @@ bool CanCreateKeyEditor(const FMovieSceneObjectPathChannel* Channel)
 	return true;
 }
 
-bool CanCreateKeyEditor(const FMovieSceneActorReferenceData* Channel)
+bool CanCreateKeyEditor(const FMovieSceneActorReferenceChannel* Channel)
 {
 	return true;
 }
@@ -718,9 +718,9 @@ private:
 };
 
 
-TSharedRef<SWidget> CreateKeyEditor(const TMovieSceneChannelHandle<FMovieSceneActorReferenceData>& Channel, const UE::Sequencer::FCreateKeyEditorParams& Params)
+TSharedRef<SWidget> CreateKeyEditor(const TMovieSceneChannelHandle<FMovieSceneActorReferenceChannel>& Channel, const UE::Sequencer::FCreateKeyEditorParams& Params)
 {
-	const FMovieSceneActorReferenceData* RawChannel = Channel.Get();
+	const FMovieSceneActorReferenceChannel* RawChannel = Channel.Get();
 	if (!RawChannel)
 	{
 		return SNullWidget::NullWidget;
@@ -728,7 +728,7 @@ TSharedRef<SWidget> CreateKeyEditor(const TMovieSceneChannelHandle<FMovieSceneAc
 
 	TFunction<TOptional<FMovieSceneActorReferenceKey>(UObject&, FTrackInstancePropertyBindings*)> Func;
 
-	TSequencerKeyEditor<FMovieSceneActorReferenceData, FMovieSceneActorReferenceKey> KeyEditor(Params.ObjectBindingID, Channel, Params.OwningSection, Params.Sequencer, Params.PropertyBindings, Func);
+	TSequencerKeyEditor<FMovieSceneActorReferenceChannel, FMovieSceneActorReferenceKey> KeyEditor(Params.ObjectBindingID, Channel, Params.OwningSection, Params.Sequencer, Params.PropertyBindings, Func);
 
 	auto OnSetCurrentValueLambda = [KeyEditor](const FMovieSceneActorReferenceKey& ActorKey) mutable
 	{
diff --git a/Engine/Source/Editor/MovieSceneTools/Private/MovieSceneToolsModule.cpp b/Engine/Source/Editor/MovieSceneTools/Private/MovieSceneToolsModule.cpp
--- a/Engine/Source/Editor/MovieSceneTools/Private/MovieSceneToolsModule.cpp
+++ b/Engine/Source/Editor/MovieSceneTools/Private/MovieSceneToolsModule.cpp
@@ -249,7 +249,7 @@ void FMovieSceneToolsModule::StartupModule()
 		SequencerModule.RegisterChannelInterface<FMovieSceneStringChannel>();
 		SequencerModule.RegisterChannelInterface<FMovieSceneTextChannel>();
 		SequencerModule.RegisterChannelInterface<FMovieSceneParticleChannel>();
-		SequencerModule.RegisterChannelInterface<FMovieSceneActorReferenceData>();
+		SequencerModule.RegisterChannelInterface<FMovieSceneActorReferenceChannel>();
 		SequencerModule.RegisterChannelInterface<FMovieSceneEventSectionData>();
 		SequencerModule.RegisterChannelInterface<FMovieSceneObjectPathChannel>();
 		SequencerModule.RegisterChannelInterface<FMovieSceneTimeWarpChannel>();
diff --git a/Engine/Source/Editor/MovieSceneTools/Private/TrackEditors/AudioTrackEditor.cpp b/Engine/Source/Editor/MovieSceneTools/Private/TrackEditors/AudioTrackEditor.cpp
--- a/Engine/Source/Editor/MovieSceneTools/Private/TrackEditors/AudioTrackEditor.cpp
+++ b/Engine/Source/Editor/MovieSceneTools/Private/TrackEditors/AudioTrackEditor.cpp
@@ -845,7 +845,7 @@ FText FAudioSection::GetSectionToolTip() const
 
 	FText AttachToolTip;
 
-	const FMovieSceneActorReferenceData& AttachActorData = AudioSection->GetAttachActorData();
+	const FMovieSceneActorReferenceChannel& AttachActorData = AudioSection->GetAttachActorData();
 
 	TMovieSceneChannelData<const FMovieSceneActorReferenceKey> ChannelData = AttachActorData.GetData();
 	TArrayView<const FMovieSceneActorReferenceKey> Values = ChannelData.GetValues().Num() != 0
diff --git a/Engine/Source/Editor/MovieSceneTools/Private/TrackEditors/PropertyTrackEditors/ActorReferencePropertyTrackEditor.cpp b/Engine/Source/Editor/MovieSceneTools/Private/TrackEditors/PropertyTrackEditors/ActorReferencePropertyTrackEditor.cpp
--- a/Engine/Source/Editor/MovieSceneTools/Private/TrackEditors/PropertyTrackEditors/ActorReferencePropertyTrackEditor.cpp
+++ b/Engine/Source/Editor/MovieSceneTools/Private/TrackEditors/PropertyTrackEditors/ActorReferencePropertyTrackEditor.cpp
@@ -29,11 +29,16 @@ FText FActorReferencePropertyTrackEditor::GetDisplayName() const
 	return LOCTEXT("ActorReferencePropertyTrackEditor_DisplayName", "Actor Reference Property");
 }
 
-void FActorReferencePropertyTrackEditor::GenerateKeysFromPropertyChanged( const FPropertyChangedParams& PropertyChangedParams, UMovieSceneSection* SectionToKey, FGeneratedTrackKeys& OutGeneratedKeys )
+void FActorReferencePropertyTrackEditor::ProcessKeyOperation(FFrameNumber InKeyTime, const UE::Sequencer::FKeyOperation& Operation, ISequencer& InSequencer, TArray<UE::Sequencer::FAddKeyResult>* OutResults)
+{
+	Operation.ApplyDefault(InKeyTime, InSequencer, OutResults);
+}
+
+void FActorReferencePropertyTrackEditor::GenerateKeysFromPropertyChanged(const FPropertyChangedParams& PropertyChangedParams, UMovieSceneSection* SectionToKey, FGeneratedTrackKeys& OutGeneratedKeys)
 {
 	// Care is taken here to ensure that GetPropertyValue is templated on UObject* which causes it to use the correct instantiation of GetPropertyValueImpl
 	AActor* NewReferencedActor = Cast<AActor>(PropertyChangedParams.GetPropertyValue<UObject*>());
-	if ( NewReferencedActor != nullptr )
+	if (NewReferencedActor != nullptr)
 	{
 		TSharedPtr<ISequencer> SequencerPtr = GetSequencer();
 
@@ -55,7 +60,7 @@ void FActorReferencePropertyTrackEditor::GenerateKeysFromPropertyChanged( const
 		{
 			FMovieSceneActorReferenceKey NewKey;
 			NewKey.Object = Binding;
-			OutGeneratedKeys.Add(FMovieSceneChannelValueSetter::Create<FMovieSceneActorReferenceData>(0, NewKey, true));
+			OutGeneratedKeys.Add(FMovieSceneChannelValueSetter::Create<FMovieSceneActorReferenceChannel>(0, NewKey, true));
 		}
 	}
 }
diff --git a/Engine/Source/Editor/MovieSceneTools/Private/TrackEditors/PropertyTrackEditors/ActorReferencePropertyTrackEditor.h b/Engine/Source/Editor/MovieSceneTools/Private/TrackEditors/PropertyTrackEditors/ActorReferencePropertyTrackEditor.h
--- a/Engine/Source/Editor/MovieSceneTools/Private/TrackEditors/PropertyTrackEditors/ActorReferencePropertyTrackEditor.h
+++ b/Engine/Source/Editor/MovieSceneTools/Private/TrackEditors/PropertyTrackEditors/ActorReferencePropertyTrackEditor.h
@@ -64,5 +64,6 @@ protected:
 
 	//~ FPropertyTrackEditor interface
 	virtual FText GetDisplayName() const override;
+	virtual void ProcessKeyOperation(FFrameNumber InKeyTime, const UE::Sequencer::FKeyOperation& Operation, ISequencer& InSequencer, TArray<UE::Sequencer::FAddKeyResult>* OutResults = nullptr) override;
 	virtual void GenerateKeysFromPropertyChanged(const FPropertyChangedParams& PropertyChangedParams, UMovieSceneSection* SectionToKey, FGeneratedTrackKeys& OutGeneratedKeys) override;
 };
diff --git a/Engine/Source/Editor/MovieSceneTools/Public/Channels/BuiltInChannelEditors.h b/Engine/Source/Editor/MovieSceneTools/Public/Channels/BuiltInChannelEditors.h
--- a/Engine/Source/Editor/MovieSceneTools/Public/Channels/BuiltInChannelEditors.h
+++ b/Engine/Source/Editor/MovieSceneTools/Public/Channels/BuiltInChannelEditors.h
@@ -38,7 +38,7 @@ class ISectionLayoutBuilder;
 /** Overrides for adding or updating a key for non-standard channels */
 MOVIESCENETOOLS_API FKeyHandle AddOrUpdateKey(FMovieSceneDoubleChannel* Channel, UMovieSceneSection* SectionToKey,  const TMovieSceneExternalValue<double>& EditorData, FFrameNumber InTime, ISequencer& Sequencer, const FGuid& InObjectBindingID, FTrackInstancePropertyBindings* PropertyBindings);
 MOVIESCENETOOLS_API FKeyHandle AddOrUpdateKey(FMovieSceneFloatChannel* Channel, UMovieSceneSection* SectionToKey,  const TMovieSceneExternalValue<float>& EditorData, FFrameNumber InTime, ISequencer& Sequencer, const FGuid& InObjectBindingID, FTrackInstancePropertyBindings* PropertyBindings);
-MOVIESCENETOOLS_API FKeyHandle AddOrUpdateKey(FMovieSceneActorReferenceData* Channel, UMovieSceneSection* SectionToKey, FFrameNumber InTime, ISequencer& Sequencer, const FGuid& InObjectBindingID, FTrackInstancePropertyBindings* PropertyBindings);
+MOVIESCENETOOLS_API FKeyHandle AddOrUpdateKey(FMovieSceneActorReferenceChannel* Channel, UMovieSceneSection* SectionToKey, FFrameNumber InTime, ISequencer& Sequencer, const FGuid& InObjectBindingID, FTrackInstancePropertyBindings* PropertyBindings);
 MOVIESCENETOOLS_API FKeyHandle AddOrUpdateKey(FMovieSceneTimeWarpChannel* Channel, UMovieSceneSection* SectionToKey, FFrameNumber InTime, ISequencer& Sequencer, const FGuid& InObjectBindingID, FTrackInstancePropertyBindings* PropertyBindings);
 
 /** Key editor overrides */
@@ -50,7 +50,7 @@ MOVIESCENETOOLS_API bool CanCreateKeyEditor(const FMovieSceneFloatChannel*
 MOVIESCENETOOLS_API bool CanCreateKeyEditor(const FMovieSceneStringChannel*     Channel);
 MOVIESCENETOOLS_API bool CanCreateKeyEditor(const FMovieSceneTextChannel*       Channel);
 MOVIESCENETOOLS_API bool CanCreateKeyEditor(const FMovieSceneObjectPathChannel* Channel);
-MOVIESCENETOOLS_API bool CanCreateKeyEditor(const FMovieSceneActorReferenceData* Channel);
+MOVIESCENETOOLS_API bool CanCreateKeyEditor(const FMovieSceneActorReferenceChannel* Channel);
 
 MOVIESCENETOOLS_API TSharedRef<SWidget> CreateKeyEditor(const TMovieSceneChannelHandle<FMovieSceneBoolChannel>&        Channel, const UE::Sequencer::FCreateKeyEditorParams& Params);
 MOVIESCENETOOLS_API TSharedRef<SWidget> CreateKeyEditor(const TMovieSceneChannelHandle<FMovieSceneByteChannel>&        Channel, const UE::Sequencer::FCreateKeyEditorParams& Params);
@@ -60,7 +60,7 @@ MOVIESCENETOOLS_API TSharedRef<SWidget> CreateKeyEditor(const TMovieSceneChannel
 MOVIESCENETOOLS_API TSharedRef<SWidget> CreateKeyEditor(const TMovieSceneChannelHandle<FMovieSceneStringChannel>&      Channel, const UE::Sequencer::FCreateKeyEditorParams& Params);
 MOVIESCENETOOLS_API TSharedRef<SWidget> CreateKeyEditor(const TMovieSceneChannelHandle<FMovieSceneTextChannel>&        Channel, const UE::Sequencer::FCreateKeyEditorParams& Params);
 MOVIESCENETOOLS_API TSharedRef<SWidget> CreateKeyEditor(const TMovieSceneChannelHandle<FMovieSceneObjectPathChannel>&  Channel, const UE::Sequencer::FCreateKeyEditorParams& Params);
-MOVIESCENETOOLS_API TSharedRef<SWidget> CreateKeyEditor(const TMovieSceneChannelHandle<FMovieSceneActorReferenceData>& Channel, const UE::Sequencer::FCreateKeyEditorParams& Params);
+MOVIESCENETOOLS_API TSharedRef<SWidget> CreateKeyEditor(const TMovieSceneChannelHandle<FMovieSceneActorReferenceChannel>& Channel, const UE::Sequencer::FCreateKeyEditorParams& Params);
 MOVIESCENETOOLS_API TSharedRef<SWidget> CreateKeyEditor(const TMovieSceneChannelHandle<FMovieSceneTimeWarpChannel>&    Channel, const UE::Sequencer::FCreateKeyEditorParams& Params);
 
 MOVIESCENETOOLS_API UMovieSceneKeyStructType* InstanceGeneratedStruct(FMovieSceneByteChannel* Channel,       FSequencerKeyStructGenerator* Generator);
@@ -116,4 +116,4 @@ MOVIESCENETOOLS_API TUniquePtr<FCurveModel> CreateCurveEditorModel(const TMovieS
 MOVIESCENETOOLS_API bool ShouldShowCurve(const FMovieSceneFloatChannel* Channel, UMovieSceneSection* InSection);
 MOVIESCENETOOLS_API bool ShouldShowCurve(const FMovieSceneDoubleChannel* Channel, UMovieSceneSection* InSection);
 
-MOVIESCENETOOLS_API TSharedPtr<UE::Sequencer::FChannelModel> CreateChannelModel(const TMovieSceneChannelHandle<FMovieSceneTimeWarpChannel>& InChannelHandle, const UE::Sequencer::FSectionModel& InSection, FName InChannelName);
\ No newline at end of file
+MOVIESCENETOOLS_API TSharedPtr<UE::Sequencer::FChannelModel> CreateChannelModel(const TMovieSceneChannelHandle<FMovieSceneTimeWarpChannel>& InChannelHandle, const UE::Sequencer::FSectionModel& InSection, FName InChannelName);
diff --git a/Engine/Source/Editor/Sequencer/Private/SequencerEdMode.cpp b/Engine/Source/Editor/Sequencer/Private/SequencerEdMode.cpp
--- a/Engine/Source/Editor/Sequencer/Private/SequencerEdMode.cpp
+++ b/Engine/Source/Editor/Sequencer/Private/SequencerEdMode.cpp
@@ -947,7 +947,7 @@ void FSequencerEdMode::DrawAudioTracks(FPrimitiveDrawInterface* PDI)
 			for (UMovieSceneSection* Section : AudioTrack->GetAudioSections())
 			{
 				UMovieSceneAudioSection* AudioSection = Cast<UMovieSceneAudioSection>(Section);
-				const FMovieSceneActorReferenceData& AttachActorData = AudioSection->GetAttachActorData();
+				const FMovieSceneActorReferenceChannel& AttachActorData = AudioSection->GetAttachActorData();
 
 				TMovieSceneChannelData<const FMovieSceneActorReferenceKey> ChannelData = AttachActorData.GetData();
 				TArrayView<const FMovieSceneActorReferenceKey> Values = ChannelData.GetValues().Num() != 0
diff --git a/Engine/Source/Runtime/MovieSceneTracks/Private/Channels/MovieSceneActorReferenceChannel.cpp b/Engine/Source/Runtime/MovieSceneTracks/Private/Channels/MovieSceneActorReferenceChannel.cpp
new file mode 100644
--- /dev/null
+++ b/Engine/Source/Runtime/MovieSceneTracks/Private/Channels/MovieSceneActorReferenceChannel.cpp
@@ -0,0 +1,100 @@
+// Copyright Epic Games, Inc. All Rights Reserved.
+
+#include "Channels/MovieSceneActorReferenceChannel.h"
+
+#include UE_INLINE_GENERATED_CPP_BY_NAME(MovieSceneActorReferenceChannel)
+
+bool FMovieSceneActorReferenceChannel::Evaluate(FFrameTime InTime, FMovieSceneActorReferenceKey& OutValue) const
+{
+	if (KeyTimes.Num())
+	{
+		const int32 Index = FMath::Max(0, Algo::UpperBound(KeyTimes, InTime.FrameNumber)-1);
+		OutValue = KeyValues[Index];
+		return true;
+	}
+
+	OutValue = DefaultValue;
+	return true;
+}
+
+void FMovieSceneActorReferenceChannel::GetKeys(const TRange<FFrameNumber>& WithinRange, TArray<FFrameNumber>* OutKeyTimes, TArray<FKeyHandle>* OutKeyHandles)
+{
+	GetData().GetKeys(WithinRange, OutKeyTimes, OutKeyHandles);
+}
+
+void FMovieSceneActorReferenceChannel::GetKeyTimes(TArrayView<const FKeyHandle> InHandles, TArrayView<FFrameNumber> OutKeyTimes)
+{
+	GetData().GetKeyTimes(InHandles, OutKeyTimes);
+}
+
+void FMovieSceneActorReferenceChannel::SetKeyTimes(TArrayView<const FKeyHandle> InHandles, TArrayView<const FFrameNumber> InKeyTimes)
+{
+	GetData().SetKeyTimes(InHandles, InKeyTimes);
+}
+
+void FMovieSceneActorReferenceChannel::DuplicateKeys(TArrayView<const FKeyHandle> InHandles, TArrayView<FKeyHandle> OutNewHandles)
+{
+	GetData().DuplicateKeys(InHandles, OutNewHandles);
+}
+
+void FMovieSceneActorReferenceChannel::DeleteKeys(TArrayView<const FKeyHandle> InHandles)
+{
+	GetData().DeleteKeys(InHandles);
+}
+
+void FMovieSceneActorReferenceChannel::DeleteKeysFrom(FFrameNumber InTime, bool bDeleteKeysBefore)
+{
+	// Insert a key at the current time to maintain evaluation
+	if (GetData().GetTimes().Num() > 0)
+	{
+		FMovieSceneActorReferenceKey Value;
+		Evaluate(InTime, Value);
+		GetData().UpdateOrAddKey(InTime, Value);
+	}
+
+	GetData().DeleteKeysFrom(InTime, bDeleteKeysBefore);
+}
+
+FKeyHandle FMovieSceneActorReferenceChannel::GetHandle(int32 Index)
+{
+	return GetData().GetHandle(Index);
+}
+
+int32 FMovieSceneActorReferenceChannel::GetIndex(FKeyHandle Handle)
+{
+	return GetData().GetIndex(Handle);
+}
+
+void FMovieSceneActorReferenceChannel::RemapTimes(const UE::MovieScene::IRetimingInterface& Retimer)
+{
+	GetData().RemapTimes(Retimer);
+}
+
+TRange<FFrameNumber> FMovieSceneActorReferenceChannel::ComputeEffectiveRange() const
+{
+	return GetData().GetTotalRange();
+}
+
+int32 FMovieSceneActorReferenceChannel::GetNumKeys() const
+{
+	return KeyTimes.Num();
+}
+
+void FMovieSceneActorReferenceChannel::Reset()
+{
+	KeyTimes.Reset();
+	KeyValues.Reset();
+	KeyHandles.Reset();
+	DefaultValue = FMovieSceneActorReferenceKey();
+}
+
+void FMovieSceneActorReferenceChannel::Offset(FFrameNumber DeltaPosition)
+{
+	GetData().Offset(DeltaPosition);
+}
+
+void FMovieSceneActorReferenceChannel::ClearDefault()
+{
+	DefaultValue = FMovieSceneActorReferenceKey();
+}
+
diff --git a/Engine/Source/Runtime/MovieSceneTracks/Private/Evaluation/MovieSceneActorReferenceTemplate.cpp b/Engine/Source/Runtime/MovieSceneTracks/Private/Evaluation/MovieSceneActorReferenceTemplate.cpp
deleted file mode 100644
diff --git a/Engine/Source/Runtime/MovieSceneTracks/Private/Evaluation/MovieSceneActorReferenceTemplate.h b/Engine/Source/Runtime/MovieSceneTracks/Private/Evaluation/MovieSceneActorReferenceTemplate.h
deleted file mode 100644
diff --git a/Engine/Source/Runtime/MovieSceneTracks/Private/MovieSceneTracksComponentTypes.cpp b/Engine/Source/Runtime/MovieSceneTracks/Private/MovieSceneTracksComponentTypes.cpp
--- a/Engine/Source/Runtime/MovieSceneTracks/Private/MovieSceneTracksComponentTypes.cpp
+++ b/Engine/Source/Runtime/MovieSceneTracks/Private/MovieSceneTracksComponentTypes.cpp
@@ -383,6 +383,34 @@ void FComponentDetachParams::ApplyDetach(USceneComponent* ChildComponentToAttach
 	}
 }
 
+/* ---------------------------------------------------------------------------
+ * Actor reference conversion functions
+ * ---------------------------------------------------------------------------*/
+void ConvertOperationalProperty(const FActorReferenceComponent& InComponentData, TObjectPtr<AActor>& Out)
+{
+	Out = InComponentData.WeakActor.Get();
+}
+
+void ConvertOperationalProperty(const FActorReferenceComponent& InComponentData, TSoftObjectPtr<AActor>& Out)
+{
+	Out = InComponentData.WeakActor.Get();
+}
+
+void ConvertOperationalProperty(TObjectPtr<AActor> const InActor, FActorReferenceComponent& Out)
+{
+	// This conversion is incomplete: we end up with an FActorReferenceComponent that has a resolved actor,
+	// but its object binding ID is left invalid, because we don't know where that actor "came from" in terms
+	// of object binding resolution. This is fine because we should only end up here for saving pre-animated
+	// state, which means we only care about saving and restoring the actor pointer.
+	Out.WeakActor = InActor.Get();
+}
+
+void ConvertOperationalProperty(TSoftObjectPtr<AActor> const InActor, FActorReferenceComponent& Out)
+{
+	// See comment above.
+	Out.WeakActor = InActor.Get();
+}
+
 
 static bool GMovieSceneTracksComponentTypesDestroyed = false;
 static TUniquePtr<FMovieSceneTracksComponentTypes> GMovieSceneTracksComponentTypes;
@@ -488,6 +516,49 @@ struct FObjectHandler : TPropertyComponentHandler<FObjectPropertyTraits, FObject
 	}
 };
 
+struct FActorReferenceHandler : TPropertyComponentHandler<FActorReferencePropertyTraits, FActorReferenceComponent>
+{
+	virtual void DispatchInitializePropertyMetaDataTasks(const FPropertyDefinition& Definition, FSystemTaskPrerequisites& InPrerequisites, FSystemSubsequentTasks& Subsequents, UMovieSceneEntitySystemLinker* Linker) override
+	{
+		FBuiltInComponentTypes* BuiltInComponents = FBuiltInComponentTypes::Get();
+		FMovieSceneTracksComponentTypes* TrackComponents = FMovieSceneTracksComponentTypes::Get();
+
+		FEntityTaskBuilder()
+			.Read(BuiltInComponents->BoundObject)
+			.Read(BuiltInComponents->PropertyBinding)
+			.ReadOptional(BuiltInComponents->CustomPropertyIndex)
+			.Write(TrackComponents->ActorReference.MetaDataComponents.GetType<0>())
+			.FilterAll({ BuiltInComponents->Tags.NeedsLink })
+			.Iterate_PerEntity(&Linker->EntityManager, [TrackComponents](UObject* Object, const FMovieScenePropertyBinding& Binding, const FCustomPropertyIndex* OptionalCustomPropertyIndex, FActorReferencePropertyTraits::FActorReferenceMetaData& OutMetaData)
+			{
+				if (OptionalCustomPropertyIndex)
+				{
+					if (const auto* CustomMetaData = TrackComponents->Accessors.ActorReference.MetaData.Find(OptionalCustomPropertyIndex->Value))
+					{
+						OutMetaData.Type = CustomMetaData->Type;
+						return;
+					}
+				}
+
+				FTrackInstancePropertyBindings Bindings(Binding.PropertyName, Binding.PropertyPath.ToString());
+				ensure(Bindings.HasValidBinding(*Object));
+				FProperty* BoundProperty = Bindings.GetProperty(*Object);
+				if (FObjectProperty* BoundObjectProperty = CastField<FObjectProperty>(BoundProperty))
+				{
+					OutMetaData.Type = EActorReferenceType::ObjectPointer;
+				}
+				else if (FSoftObjectProperty* BoundSoftObjectProperty = CastField<FSoftObjectProperty>(BoundProperty))
+				{
+					OutMetaData.Type = EActorReferenceType::SoftObjectPointer;
+				}
+				else
+				{
+					ensureAlwaysMsgf(false, TEXT("Unsupported actor reference type for property '%s' on object '%s'"), *Bindings.GetPropertyPath(), *GetNameSafe(Object));
+				}
+			});
+	}
+};
+
 FMovieSceneTracksComponentTypes::FMovieSceneTracksComponentTypes()
 {
 	FBuiltInComponentTypes* BuiltInComponents = FBuiltInComponentTypes::Get();
@@ -505,6 +576,7 @@ FMovieSceneTracksComponentTypes::FMovieSceneTracksComponentTypes()
 	ComponentRegistry->NewPropertyType(String, TEXT("FString"));
 	ComponentRegistry->NewPropertyType(Text, TEXT("FText"));
 	ComponentRegistry->NewPropertyType(Object, TEXT("Object"));
+	ComponentRegistry->NewPropertyType(ActorReference, TEXT("Actor Reference"));
 
 	ComponentRegistry->NewPropertyType(Transform, TEXT("FTransform"));
 	ComponentRegistry->NewPropertyType(EulerTransform, TEXT("FEulerTransform"));
@@ -533,6 +605,7 @@ FMovieSceneTracksComponentTypes::FMovieSceneTracksComponentTypes()
 
 	Bool.MetaDataComponents.Initialize(ComponentRegistry, TEXT("Bool Bitfield"));
 	Object.MetaDataComponents.Initialize(ComponentRegistry, TEXT("Object Class"));
+	ActorReference.MetaDataComponents.Initialize(ComponentRegistry, TEXT("Actor Reference Type"));
 
 	ComponentRegistry->NewComponentType(&QuaternionRotationChannel[0], TEXT("Quaternion Rotation Channel 0"));
 	ComponentRegistry->NewComponentType(&QuaternionRotationChannel[1], TEXT("Quaternion Rotation Channel 1"));
@@ -544,6 +617,9 @@ FMovieSceneTracksComponentTypes::FMovieSceneTracksComponentTypes()
 
 	ComponentRegistry->NewComponentType(&ConstraintChannel, TEXT("Constraint Channel"));
 
+	ComponentRegistry->NewComponentType(&ActorReferenceChannel, TEXT("Actor Reference Channel"));
+	ComponentRegistry->NewComponentType(&ActorReferenceResult, TEXT("Actor Reference Result"));
+
 	ComponentRegistry->NewComponentType(&AttachParent, TEXT("Attach Parent"));
 	ComponentRegistry->NewComponentType(&AttachComponent, TEXT("Attachment Component"));
 	ComponentRegistry->NewComponentType(&AttachParentBinding, TEXT("Attach Parent Binding"));
@@ -707,6 +783,13 @@ FMovieSceneTracksComponentTypes::FMovieSceneTracksComponentTypes()
 	.SetCustomAccessors(&Accessors.Object)
 	.Commit(FObjectHandler());
 
+	// --------------------------------------------------------------------------------------------
+	// Set up Actor Reference properties
+	BuiltInComponents->PropertyRegistry.DefineProperty(ActorReference, TEXT("Apply Actor Reference Properties"))
+	.AddSoleChannel(ActorReferenceResult)
+	.SetCustomAccessors(&Accessors.ActorReference)
+	.Commit(FActorReferenceHandler());
+
 	// --------------------------------------------------------------------------------------------
 	// Set up float parameters
 	BuiltInComponents->PropertyRegistry.DefineProperty(FloatParameter, TEXT("Apply Float Parameters"))
@@ -877,6 +960,13 @@ FMovieSceneTracksComponentTypes::FMovieSceneTracksComponentTypes()
 	// Set up camera shake components
 	ComponentRegistry->Factories.DefineMutuallyInclusiveComponent(CameraShake, CameraShakeInstance);
 
+	// --------------------------------------------------------------------------------------------
+	// Set up actor reference components
+	ComponentRegistry->Factories.DuplicateChildComponent(ActorReferenceChannel);
+	ComponentRegistry->Factories.DuplicateChildComponent(ActorReferenceResult);
+	ComponentRegistry->Factories.DefineMutuallyInclusiveComponent(ActorReferenceChannel, ActorReferenceResult);
+	ComponentRegistry->Factories.DefineMutuallyInclusiveComponent(ActorReferenceChannel, BuiltInComponents->EvalTime);
+
 	InitializeMovieSceneTracksAccessors(this);
 }
 
diff --git a/Engine/Source/Runtime/MovieSceneTracks/Private/Sections/MovieSceneActorReferenceSection.cpp b/Engine/Source/Runtime/MovieSceneTracks/Private/Sections/MovieSceneActorReferenceSection.cpp
--- a/Engine/Source/Runtime/MovieSceneTracks/Private/Sections/MovieSceneActorReferenceSection.cpp
+++ b/Engine/Source/Runtime/MovieSceneTracks/Private/Sections/MovieSceneActorReferenceSection.cpp
@@ -1,108 +1,15 @@
 // Copyright Epic Games, Inc. All Rights Reserved.
 
 #include "Sections/MovieSceneActorReferenceSection.h"
+
 #include "Channels/MovieSceneChannelProxy.h"
+#include "Evaluation/MovieSceneSequenceHierarchy.h"
 #include "MovieSceneObjectBindingID.h"
 #include "MovieSceneSequenceID.h"
-#include "Evaluation/MovieSceneSequenceHierarchy.h"
-#include "IMovieScenePlayer.h"
+#include "Tracks/MovieSceneObjectPropertyTrack.h"
 
 #include UE_INLINE_GENERATED_CPP_BY_NAME(MovieSceneActorReferenceSection)
 
-bool FMovieSceneActorReferenceData::Evaluate(FFrameTime InTime, FMovieSceneActorReferenceKey& OutValue) const
-{
-	if (KeyTimes.Num())
-	{
-		const int32 Index = FMath::Max(0, Algo::UpperBound(KeyTimes, InTime.FrameNumber)-1);
-		OutValue = KeyValues[Index];
-		return true;
-	}
-
-	OutValue = DefaultValue;
-	return true;
-}
-
-void FMovieSceneActorReferenceData::GetKeys(const TRange<FFrameNumber>& WithinRange, TArray<FFrameNumber>* OutKeyTimes, TArray<FKeyHandle>* OutKeyHandles)
-{
-	GetData().GetKeys(WithinRange, OutKeyTimes, OutKeyHandles);
-}
-
-void FMovieSceneActorReferenceData::GetKeyTimes(TArrayView<const FKeyHandle> InHandles, TArrayView<FFrameNumber> OutKeyTimes)
-{
-	GetData().GetKeyTimes(InHandles, OutKeyTimes);
-}
-
-void FMovieSceneActorReferenceData::SetKeyTimes(TArrayView<const FKeyHandle> InHandles, TArrayView<const FFrameNumber> InKeyTimes)
-{
-	GetData().SetKeyTimes(InHandles, InKeyTimes);
-}
-
-void FMovieSceneActorReferenceData::DuplicateKeys(TArrayView<const FKeyHandle> InHandles, TArrayView<FKeyHandle> OutNewHandles)
-{
-	GetData().DuplicateKeys(InHandles, OutNewHandles);
-}
-
-void FMovieSceneActorReferenceData::DeleteKeys(TArrayView<const FKeyHandle> InHandles)
-{
-	GetData().DeleteKeys(InHandles);
-}
-
-void FMovieSceneActorReferenceData::DeleteKeysFrom(FFrameNumber InTime, bool bDeleteKeysBefore)
-{
-	// Insert a key at the current time to maintain evaluation
-	if (GetData().GetTimes().Num() > 0)
-	{
-		FMovieSceneActorReferenceKey Value;
-		Evaluate(InTime, Value);
-		GetData().UpdateOrAddKey(InTime, Value);
-	}
-
-	GetData().DeleteKeysFrom(InTime, bDeleteKeysBefore);
-}
-
-FKeyHandle FMovieSceneActorReferenceData::GetHandle(int32 Index)
-{
-	return GetData().GetHandle(Index);
-}
-
-int32 FMovieSceneActorReferenceData::GetIndex(FKeyHandle Handle)
-{
-	return GetData().GetIndex(Handle);
-}
-
-void FMovieSceneActorReferenceData::RemapTimes(const UE::MovieScene::IRetimingInterface& Retimer)
-{
-	GetData().RemapTimes(Retimer);
-}
-
-TRange<FFrameNumber> FMovieSceneActorReferenceData::ComputeEffectiveRange() const
-{
-	return GetData().GetTotalRange();
-}
-
-int32 FMovieSceneActorReferenceData::GetNumKeys() const
-{
-	return KeyTimes.Num();
-}
-
-void FMovieSceneActorReferenceData::Reset()
-{
-	KeyTimes.Reset();
-	KeyValues.Reset();
-	KeyHandles.Reset();
-	DefaultValue = FMovieSceneActorReferenceKey();
-}
-
-void FMovieSceneActorReferenceData::Offset(FFrameNumber DeltaPosition)
-{
-	GetData().Offset(DeltaPosition);
-}
-
-void FMovieSceneActorReferenceData::ClearDefault()
-{
-	DefaultValue = FMovieSceneActorReferenceKey();
-}
-
 UMovieSceneActorReferenceSection::UMovieSceneActorReferenceSection( const FObjectInitializer& ObjectInitializer )
 	: Super( ObjectInitializer )
 {
@@ -177,4 +84,14 @@ void UMovieSceneActorReferenceSection::OnBindingIDsUpdated(const TMap<UE::MovieS
 	}
 }
 
+void UMovieSceneActorReferenceSection::ImportEntityImpl(UMovieSceneEntitySystemLinker* EntityLinker, const FEntityImportParams& Params, FImportedEntity* OutImportedEntity)
+{
+	using namespace UE::MovieScene;
+
+	const FMovieSceneTracksComponentTypes* TracksComponents = FMovieSceneTracksComponentTypes::Get();
+
+	FPropertyTrackEntityImportHelper(TracksComponents->ActorReference)
+		.Add(TracksComponents->ActorReferenceChannel, &ActorReferenceData)
+		.Commit(this, Params, OutImportedEntity);
+}
 
diff --git a/Engine/Source/Runtime/MovieSceneTracks/Private/Systems/ActorReferenceChannelEvaluatorSystem.cpp b/Engine/Source/Runtime/MovieSceneTracks/Private/Systems/ActorReferenceChannelEvaluatorSystem.cpp
new file mode 100644
--- /dev/null
+++ b/Engine/Source/Runtime/MovieSceneTracks/Private/Systems/ActorReferenceChannelEvaluatorSystem.cpp
@@ -0,0 +1,139 @@
+// Copyright Epic Games, Inc. All Rights Reserved.
+
+#include "Systems/ActorReferenceChannelEvaluatorSystem.h"
+
+#include "Channels/MovieSceneActorReferenceChannel.h"
+#include "EntitySystem/BuiltInComponentTypes.h"
+#include "EntitySystem/MovieSceneEntitySystemLinker.h"
+#include "EntitySystem/MovieSceneEntitySystemTask.h"
+#include "EntitySystem/MovieSceneEvalTimeSystem.h"
+#include "EntitySystem/MovieSceneInstanceRegistry.h"
+#include "GameFramework/Actor.h"
+#include "MovieSceneTracksComponentTypes.h"
+
+#include UE_INLINE_GENERATED_CPP_BY_NAME(ActorReferenceChannelEvaluatorSystem)
+
+namespace UE::MovieScene
+{
+	DECLARE_CYCLE_STAT(TEXT("MovieScene: Evaluate actor reference channels"), MovieSceneEval_EvaluateActorReferenceChannelTask, STATGROUP_MovieSceneECS);
+
+	struct FEvaluateActorReferenceChannels
+	{
+		const FInstanceRegistry* InstanceRegistry;
+
+		FEvaluateActorReferenceChannels(const FInstanceRegistry* InInstanceRegistry)
+			: InstanceRegistry(InInstanceRegistry)
+		{}
+
+		void ForEachEntity(FSourceActorReferenceChannel ActorReferenceChannel, FFrameTime FrameTime, FInstanceHandle InstanceHandle, FActorReferenceComponent& OutResult) const
+		{
+			FMovieSceneActorReferenceKey ObjectBinding;
+			const bool bGotKeyValue = ActorReferenceChannel.Source->Evaluate(FrameTime, ObjectBinding);
+			if (bGotKeyValue)
+			{
+				// (Re-)resolve the actor reference if we have a new valid actor to reference.
+				if (ObjectBinding.Object.IsValid() && ObjectBinding.Object != OutResult.ObjectBindingID)
+				{
+					OutResult.WeakActor.Reset();
+					ResolveActorReference(ObjectBinding, InstanceHandle, OutResult);
+				}
+				// Clear the actor pointer if we're supposed to not reference an actor anymore.
+				else if (!ObjectBinding.Object.IsValid() && ObjectBinding.Object != OutResult.ObjectBindingID)
+				{
+					OutResult.ObjectBindingID = ObjectBinding.Object;
+					OutResult.WeakActor.Reset();
+				}
+				// The actor reference hasn't changed but maybe the pointer was valid and became invalid (e.g. a spawnable went away),
+				// or was invalid and became valid (e.g. a spawnable was spawned, unspawned, and has spawned again).
+				// We don't need to handle the first case, since the weak pointer will naturally become null, but we need to handle
+				// the second case.
+				else if (ObjectBinding.Object.IsValid() && !OutResult.WeakActor.IsValid())
+				{
+					ensure(ObjectBinding.Object == OutResult.ObjectBindingID);
+					ResolveActorReference(ObjectBinding, InstanceHandle, OutResult);
+				}
+			}
+			else if (OutResult.ObjectBindingID.IsValid())
+			{
+				// We have no data anymore.
+				OutResult.ObjectBindingID = ObjectBinding.Object;
+				OutResult.WeakActor.Reset();
+			}
+		}
+
+		void ResolveActorReference(const FMovieSceneActorReferenceKey ObjectBinding, FInstanceHandle InstanceHandle, FActorReferenceComponent& OutResult) const
+		{
+			const FSequenceInstance& TargetInstance = InstanceRegistry->GetInstance(InstanceHandle);
+
+			for (TWeakObjectPtr<> WeakObject : ObjectBinding.Object.ResolveBoundObjects(TargetInstance))
+			{
+				if (AActor* Actor = Cast<AActor>(WeakObject.Get()))
+				{
+					OutResult.ObjectBindingID = ObjectBinding.Object;
+					OutResult.WeakActor = Actor;
+
+					// Can only ever reference one actor.
+					break;
+				}
+			}
+		}
+	};
+}
+
+
+UActorReferenceChannelEvaluatorSystem::UActorReferenceChannelEvaluatorSystem(const FObjectInitializer& ObjInit)
+	: Super(ObjInit)
+{
+	using namespace UE::MovieScene;
+
+	const FBuiltInComponentTypes* BuiltInComponents = FBuiltInComponentTypes::Get();
+	const FMovieSceneTracksComponentTypes* TracksComponents = FMovieSceneTracksComponentTypes::Get();
+
+	RelevantComponent = TracksComponents->ActorReferenceChannel;
+	Phase = ESystemPhase::Scheduling;
+
+	if (HasAnyFlags(RF_ClassDefaultObject))
+	{
+		DefineComponentConsumer(GetClass(), BuiltInComponents->EvalTime);
+		DefineComponentProducer(GetClass(), TracksComponents->ActorReferenceResult);
+	}
+}
+
+void UActorReferenceChannelEvaluatorSystem::OnSchedulePersistentTasks(UE::MovieScene::IEntitySystemScheduler* TaskScheduler)
+{
+	using namespace UE::MovieScene;
+
+	const FBuiltInComponentTypes* BuiltInComponents = FBuiltInComponentTypes::Get();
+	const FMovieSceneTracksComponentTypes* TracksComponents = FMovieSceneTracksComponentTypes::Get();
+
+	const FInstanceRegistry* InstanceRegistry = Linker->GetInstanceRegistry();
+
+	FEntityTaskBuilder()
+	.Read(TracksComponents->ActorReferenceChannel)
+	.Read(BuiltInComponents->EvalTime)
+	.Read(BuiltInComponents->InstanceHandle)
+	.Write(TracksComponents->ActorReferenceResult)
+	.FilterNone({ BuiltInComponents->Tags.Ignored })
+	.SetStat(GET_STATID(MovieSceneEval_EvaluateActorReferenceChannelTask))
+	.Fork_PerEntity<FEvaluateActorReferenceChannels>(&Linker->EntityManager, TaskScheduler, InstanceRegistry);
+}
+
+void UActorReferenceChannelEvaluatorSystem::OnRun(FSystemTaskPrerequisites& InPrerequisites, FSystemSubsequentTasks& Subsequents)
+{
+	using namespace UE::MovieScene;
+
+	const FBuiltInComponentTypes* BuiltInComponents = FBuiltInComponentTypes::Get();
+	const FMovieSceneTracksComponentTypes* TracksComponents = FMovieSceneTracksComponentTypes::Get();
+
+	const FInstanceRegistry* InstanceRegistry = Linker->GetInstanceRegistry();
+
+	FEntityTaskBuilder()
+	.Read(TracksComponents->ActorReferenceChannel)
+	.Read(BuiltInComponents->EvalTime)
+	.Read(BuiltInComponents->InstanceHandle)
+	.Write(TracksComponents->ActorReferenceResult)
+	.FilterNone({ BuiltInComponents->Tags.Ignored })
+	.SetStat(GET_STATID(MovieSceneEval_EvaluateActorReferenceChannelTask))
+	.Dispatch_PerEntity<FEvaluateActorReferenceChannels>(&Linker->EntityManager, InPrerequisites, &Subsequents, InstanceRegistry);
+}
+
diff --git a/Engine/Source/Runtime/MovieSceneTracks/Private/Systems/MovieSceneActorReferenceSystem.cpp b/Engine/Source/Runtime/MovieSceneTracks/Private/Systems/MovieSceneActorReferenceSystem.cpp
new file mode 100644
--- /dev/null
+++ b/Engine/Source/Runtime/MovieSceneTracks/Private/Systems/MovieSceneActorReferenceSystem.cpp
@@ -0,0 +1,26 @@
+// Copyright Epic Games, Inc. All Rights Reserved.
+
+#include "Systems/MovieSceneActorReferencePropertySystem.h"
+#include "Systems/ActorReferenceChannelEvaluatorSystem.h"
+#include "MovieSceneTracksComponentTypes.h"
+
+#include UE_INLINE_GENERATED_CPP_BY_NAME(MovieSceneActorReferencePropertySystem)
+
+UMovieSceneActorReferencePropertySystem::UMovieSceneActorReferencePropertySystem(const FObjectInitializer& ObjInit)
+	: Super(ObjInit)
+{
+	BindToProperty(UE::MovieScene::FMovieSceneTracksComponentTypes::Get()->ActorReference);
+
+	if (HasAnyFlags(RF_ClassDefaultObject))
+	{
+		DefineImplicitPrerequisite(UActorReferenceChannelEvaluatorSystem::StaticClass(), GetClass());
+
+		DefineComponentConsumer(GetClass(), UE::MovieScene::FMovieSceneTracksComponentTypes::Get()->ActorReference.PropertyTag);
+	}
+}
+
+void UMovieSceneActorReferencePropertySystem::OnRun(FSystemTaskPrerequisites& InPrerequisites, FSystemSubsequentTasks& Subsequents)
+{
+	Super::OnRun(InPrerequisites, Subsequents);
+}
+
diff --git a/Engine/Source/Runtime/MovieSceneTracks/Private/Systems/MovieSceneAudioSystem.cpp b/Engine/Source/Runtime/MovieSceneTracks/Private/Systems/MovieSceneAudioSystem.cpp
--- a/Engine/Source/Runtime/MovieSceneTracks/Private/Systems/MovieSceneAudioSystem.cpp
+++ b/Engine/Source/Runtime/MovieSceneTracks/Private/Systems/MovieSceneAudioSystem.cpp
@@ -372,7 +372,7 @@ private:
 		// Root audio track
 		if (BoundObject == nullptr)
 		{
-			const FMovieSceneActorReferenceData& AttachActorData = AudioSection->GetAttachActorData();
+			const FMovieSceneActorReferenceChannel& AttachActorData = AudioSection->GetAttachActorData();
 
 			USceneComponent* AttachComponent = nullptr;
 			FMovieSceneActorReferenceKey AttachKey;
diff --git a/Engine/Source/Runtime/MovieSceneTracks/Private/Tracks/MovieSceneActorReferenceTrack.cpp b/Engine/Source/Runtime/MovieSceneTracks/Private/Tracks/MovieSceneActorReferenceTrack.cpp
--- a/Engine/Source/Runtime/MovieSceneTracks/Private/Tracks/MovieSceneActorReferenceTrack.cpp
+++ b/Engine/Source/Runtime/MovieSceneTracks/Private/Tracks/MovieSceneActorReferenceTrack.cpp
@@ -2,7 +2,6 @@
 
 #include "Tracks/MovieSceneActorReferenceTrack.h"
 #include "Sections/MovieSceneActorReferenceSection.h"
-#include "Evaluation/MovieSceneActorReferenceTemplate.h"
 
 #include UE_INLINE_GENERATED_CPP_BY_NAME(MovieSceneActorReferenceTrack)
 
@@ -21,9 +20,3 @@ UMovieSceneSection* UMovieSceneActorReferenceTrack::CreateNewSection()
 	return NewObject<UMovieSceneActorReferenceSection>(this, NAME_None, RF_Transactional);
 }
 
-FMovieSceneEvalTemplatePtr UMovieSceneActorReferenceTrack::CreateTemplateForSection(const UMovieSceneSection& InSection) const
-{
-	return FMovieSceneActorReferenceSectionTemplate(*CastChecked<UMovieSceneActorReferenceSection>(&InSection), *this);
-}
-
-
diff --git a/Engine/Source/Runtime/MovieSceneTracks/Public/Channels/MovieSceneActorReferenceChannel.h b/Engine/Source/Runtime/MovieSceneTracks/Public/Channels/MovieSceneActorReferenceChannel.h
new file mode 100644
--- /dev/null
+++ b/Engine/Source/Runtime/MovieSceneTracks/Public/Channels/MovieSceneActorReferenceChannel.h
@@ -0,0 +1,177 @@
+// Copyright Epic Games, Inc. All Rights Reserved.
+
+#pragma once
+
+#include "Channels/MovieSceneChannel.h"
+#include "Channels/MovieSceneChannelData.h"
+#include "Channels/MovieSceneChannelTraits.h"
+#include "Misc/FrameNumber.h"
+#include "MovieSceneObjectBindingID.h"
+
+#include "MovieSceneActorReferenceChannel.generated.h"
+
+/**
+ * Keyframe value type for actor references.
+ */
+USTRUCT()
+struct FMovieSceneActorReferenceKey
+{
+	GENERATED_BODY()
+
+	FMovieSceneActorReferenceKey()
+	{}
+
+	FMovieSceneActorReferenceKey(const FMovieSceneObjectBindingID& InBindingID)
+		: Object(InBindingID)
+	{}
+
+	friend bool operator==(const FMovieSceneActorReferenceKey& A, const FMovieSceneActorReferenceKey& B)
+	{
+		return A.Object == B.Object && A.ComponentName == B.ComponentName && A.SocketName == B.SocketName;
+	}
+
+	friend bool operator!=(const FMovieSceneActorReferenceKey& A, const FMovieSceneActorReferenceKey& B)
+	{
+		return A.Object != B.Object || A.ComponentName != B.ComponentName || A.SocketName != B.SocketName;
+	}
+
+	UPROPERTY(EditAnywhere, Category="Key")
+	FMovieSceneObjectBindingID Object;
+
+	UPROPERTY(EditAnywhere, Category="Key")
+	FName ComponentName;
+
+	UPROPERTY(EditAnywhere, Category="Key")
+	FName SocketName;
+};
+
+/**
+ * A channel containing keyframes referencing an actor elsewhere in a Sequence.
+ */
+USTRUCT()
+struct FMovieSceneActorReferenceChannel : public FMovieSceneChannel
+{
+	GENERATED_BODY()
+
+	typedef FMovieSceneActorReferenceKey CurveValueType;
+
+	FMovieSceneActorReferenceChannel()
+		: DefaultValue()
+	{}
+
+	/**
+	 * Access a mutable interface for this channel's data
+	 *
+	 * @return An object that is able to manipulate this channel's data
+	 */
+	inline TMovieSceneChannelData<FMovieSceneActorReferenceKey> GetData()
+	{
+		return TMovieSceneChannelData<FMovieSceneActorReferenceKey>(&KeyTimes, &KeyValues, this, &KeyHandles);
+	}
+
+	/**
+	 * Access a constant interface for this channel's data
+	 *
+	 * @return An object that is able to interrogate this channel's data
+	 */
+	inline TMovieSceneChannelData<const FMovieSceneActorReferenceKey> GetData() const
+	{
+		return TMovieSceneChannelData<const FMovieSceneActorReferenceKey>(&KeyTimes, &KeyValues);
+	}
+
+	/**
+	 * Evaluate this channel
+	 *
+	 * @param InTime     The time to evaluate at
+	 * @param OutValue   A value to receive the result
+	 * @return true if the channel was evaluated successfully, false otherwise
+	 */
+	MOVIESCENETRACKS_API bool Evaluate(FFrameTime InTime, FMovieSceneActorReferenceKey& OutValue) const;
+
+public:
+
+	// ~ FMovieSceneChannel Interface
+	MOVIESCENETRACKS_API virtual void GetKeys(const TRange<FFrameNumber>& WithinRange, TArray<FFrameNumber>* OutKeyTimes, TArray<FKeyHandle>* OutKeyHandles) override;
+	MOVIESCENETRACKS_API virtual void GetKeyTimes(TArrayView<const FKeyHandle> InHandles, TArrayView<FFrameNumber> OutKeyTimes) override;
+	MOVIESCENETRACKS_API virtual void SetKeyTimes(TArrayView<const FKeyHandle> InHandles, TArrayView<const FFrameNumber> InKeyTimes) override;
+	MOVIESCENETRACKS_API virtual void DuplicateKeys(TArrayView<const FKeyHandle> InHandles, TArrayView<FKeyHandle> OutNewHandles) override;
+	MOVIESCENETRACKS_API virtual void DeleteKeys(TArrayView<const FKeyHandle> InHandles) override;
+	MOVIESCENETRACKS_API virtual void DeleteKeysFrom(FFrameNumber InTime, bool bDeleteKeysBefore) override;
+	MOVIESCENETRACKS_API virtual void RemapTimes(const UE::MovieScene::IRetimingInterface& Retimer) override;
+	MOVIESCENETRACKS_API virtual TRange<FFrameNumber> ComputeEffectiveRange() const override;
+	MOVIESCENETRACKS_API virtual int32 GetNumKeys() const override;
+	MOVIESCENETRACKS_API virtual void Reset() override;
+	MOVIESCENETRACKS_API virtual void Offset(FFrameNumber DeltaPosition) override;
+	MOVIESCENETRACKS_API virtual void ClearDefault() override;
+	MOVIESCENETRACKS_API virtual FKeyHandle GetHandle(int32 Index) override;
+	MOVIESCENETRACKS_API virtual int32 GetIndex(FKeyHandle Handle) override;
+
+public:
+
+	/**
+	 * Set this channel's default value that should be used when no keys are present
+	 *
+	 * @param InDefaultValue The desired default value
+	 */
+	inline void SetDefault(FMovieSceneActorReferenceKey InDefaultValue)
+	{
+		DefaultValue = InDefaultValue;
+	}
+
+	/**
+	 * Get this channel's default value that will be used when no keys are present
+	 *
+	 * @return (Optional) The channel's default value
+	 */
+	inline const FMovieSceneActorReferenceKey& GetDefault() const
+	{
+		return DefaultValue;
+	}
+
+	/**
+	 * Upgrade legacy data by appending to the end of the array. Assumes sorted data
+	 */
+	void UpgradeLegacyTime(UObject* Context, double Time, FMovieSceneActorReferenceKey Value)
+	{
+		FFrameRate LegacyFrameRate = GetLegacyConversionFrameRate();
+		FFrameNumber KeyTime = UpgradeLegacyMovieSceneTime(nullptr, LegacyFrameRate, Time);
+
+		check(KeyTimes.Num() == 0 || KeyTime >= KeyTimes.Last());
+
+		KeyTimes.Add(KeyTime);
+		KeyValues.Add(Value);
+	}
+
+private:
+
+	/** Sorted array of key times */
+	UPROPERTY(meta=(KeyTimes))
+	TArray<FFrameNumber> KeyTimes;
+
+	/** Default value used when there are no keys */
+	UPROPERTY()
+	FMovieSceneActorReferenceKey DefaultValue;
+
+	/** Array of values that correspond to each key time */
+	UPROPERTY(meta=(KeyValues))
+	TArray<FMovieSceneActorReferenceKey> KeyValues;
+
+	/** This needs to be a UPROPERTY so it gets saved into editor transactions but transient so it doesn't get saved into assets. */
+	UPROPERTY(Transient)
+	FMovieSceneKeyHandleMap KeyHandles;
+};
+
+// Backwards compatibility for C++ code.
+using FMovieSceneActorReferenceData UE_DEPRECATED(5.8, "Please use FMovieSceneActorReferenceChannel") = FMovieSceneActorReferenceChannel;
+
+inline bool EvaluateChannel(const FMovieSceneActorReferenceChannel* InChannel, FFrameTime InTime, FMovieSceneActorReferenceKey& OutValue)
+{
+	return InChannel->Evaluate(InTime, OutValue);
+}
+
+inline bool GetChannelDefault(const FMovieSceneActorReferenceChannel* Channel, FMovieSceneActorReferenceKey& OutDefaultValue)
+{
+	OutDefaultValue = Channel->GetDefault();
+	return true;
+}
+
diff --git a/Engine/Source/Runtime/MovieSceneTracks/Public/MovieSceneTracksComponentTypes.h b/Engine/Source/Runtime/MovieSceneTracks/Public/MovieSceneTracksComponentTypes.h
--- a/Engine/Source/Runtime/MovieSceneTracks/Public/MovieSceneTracksComponentTypes.h
+++ b/Engine/Source/Runtime/MovieSceneTracks/Public/MovieSceneTracksComponentTypes.h
@@ -2,27 +2,26 @@
 
 #pragma once
 
-#include "Math/Range.h"
-#include "Math/Transform.h"
+#include "Channels/MovieSceneActorReferenceChannel.h"
+#include "Channels/MovieSceneFloatChannel.h"
+#include "Channels/MovieSceneObjectPathChannel.h"
 #include "EntitySystem/BuiltInComponentTypes.h"
 #include "EntitySystem/MovieSceneEntityIDs.h"
-#include "EntitySystem/MovieScenePropertySystemTypes.h"
 #include "EntitySystem/MovieScenePropertyMetaData.h"
+#include "EntitySystem/MovieScenePropertySystemTypes.h"
 #include "EntitySystem/MovieScenePropertyTraits.h"
-#include "Engine/EngineTypes.h"
 #include "EulerTransform.h"
-#include "TransformData.h"
-#include "MovieSceneTracksPropertyTypes.h"
-#include "Styling/SlateColor.h"
-#include "ConstraintChannel.h"
+#include "GameFramework/Actor.h"
 #include "Materials/MaterialParameters.h"
-#include "Sections/MovieSceneCameraShakeSection.h"
+#include "Math/Range.h"
+#include "Math/Transform.h"
 #include "Misc/Guid.h"
 #include "MovieSceneTracksPropertyTraits.h"
+#include "MovieSceneTracksPropertyTypes.h"
+#include "Sections/MovieSceneCameraShakeSection.h"
+#include "Styling/SlateColor.h"
 #include "Tracks/MovieSceneMaterialTrack.h"
-#include "Channels/MovieSceneFloatChannel.h"
-#include "Channels/MovieSceneDoubleChannel.h"
-#include "Channels/MovieSceneObjectPathChannel.h"
+
 #include "MovieSceneTracksComponentTypes.generated.h"
 
 class UMaterialParameterCollection;
@@ -210,6 +209,32 @@ namespace UE
 namespace MovieScene
 {
 
+struct FSourceActorReferenceChannel
+{
+	FSourceActorReferenceChannel()
+		: Source(nullptr)
+	{}
+
+	FSourceActorReferenceChannel(const FMovieSceneActorReferenceChannel* InSource)
+		: Source(InSource)
+	{}
+
+	const FMovieSceneActorReferenceChannel* Source;
+};
+
+enum class EActorReferenceType
+{
+	ObjectPointer,
+	SoftObjectPointer
+};
+
+struct FActorReferenceComponent
+{
+	FMovieSceneObjectBindingID ObjectBindingID;
+	TWeakObjectPtr<AActor> WeakActor;
+	EActorReferenceType Type = EActorReferenceType::ObjectPointer;
+};
+
 struct FComponentAttachParamsDestination
 {
 	FName SocketName    = NAME_None;
@@ -494,6 +519,124 @@ struct FObjectPropertyTraits : IPropertyTraits
 	}
 };
 
+struct FActorReferencePropertyTraits : IPropertyTraits
+{
+	static constexpr bool bIsComposite = false;
+
+	struct FActorReferenceMetaData
+	{
+		EActorReferenceType Type = EActorReferenceType::ObjectPointer;
+	};
+
+	using StorageType = FActorReferenceComponent;
+	using MetaDataType = TPropertyMetaData<FActorReferenceMetaData>;
+
+	using ObjectPointerTraitsImpl = TPropertyTraits<TObjectPtr<AActor>, FActorReferenceComponent>;
+	using SoftObjectPointerTraitsImpl = TPropertyTraits<TSoftObjectPtr<AActor>, FActorReferenceComponent>;
+
+	FORCEINLINE static bool SupportsProperty(const FProperty& InProperty)
+	{
+		if (const FObjectProperty* ObjectProperty = CastField<FObjectProperty>(&InProperty))
+		{
+			return ObjectProperty->PropertyClass->IsChildOf<AActor>();
+		}
+		else if (const FSoftObjectProperty* SoftObjectProperty = CastField<FSoftObjectProperty>(&InProperty))
+		{
+			return SoftObjectProperty->PropertyClass->IsChildOf<AActor>();
+		}
+		return false;
+	}
+	FORCEINLINE static FIntermediatePropertyValue CoercePropertyValue(const FProperty& InProperty, const FSourcePropertyValue& InPropertyValue)
+	{
+		// This will return an FActorReferenceComponent that has (probably) a valid actor pointer, but no cached binding ID for it.
+		// We'll have to handle that ourselves in the track editor, since going from a property value to a channel value means 
+		// converting an AActor pointer to a FMovieSceneObjectBindingID, and we don't have the context to do that here.
+		FActorReferenceComponent ComponentData;
+		ComponentData.WeakActor = *InPropertyValue.Cast<AActor*>();
+		return FIntermediatePropertyValue::FromValue(ComponentData);
+	}
+	FORCEINLINE static void UnpackChannels(const FActorReferenceComponent& InValue, const FProperty& Property, FUnpackedChannelValues& OutUnpackedValues)
+	{
+		// We expect the given value to have been computed by UActorReferenceChannelEvaluatorSystem, which caches the binding ID along
+		// with the resolved actor pointer.
+		FMovieSceneActorReferenceKey ChannelValue;
+		ChannelValue.Object = InValue.ObjectBindingID;
+		OutUnpackedValues.Add(FUnpackedChannelValue(TIndexedChannelValue<FMovieSceneActorReferenceChannel>(ChannelValue, 0), NAME_None));
+	}
+
+	static void GetObjectPropertyValue(const UObject* InObject, FActorReferenceMetaData MetaData, const FCustomPropertyAccessor& BaseCustomAccessor, FActorReferenceComponent& OutValue)
+	{
+		const TCustomPropertyAccessor<FActorReferencePropertyTraits>& CustomAccessor = static_cast<const TCustomPropertyAccessor<FActorReferencePropertyTraits>&>(BaseCustomAccessor);
+		OutValue = (*CustomAccessor.Functions.Getter)(InObject, MetaData);
+	}
+	static void GetObjectPropertyValue(const UObject* InObject, FActorReferenceMetaData MetaData, uint16 PropertyOffset, FActorReferenceComponent& OutValue)
+	{
+		switch (MetaData.Type)
+		{
+			case EActorReferenceType::ObjectPointer:
+				ObjectPointerTraitsImpl::GetObjectPropertyValue(InObject, PropertyOffset, OutValue);
+				break;
+			case EActorReferenceType::SoftObjectPointer:
+				SoftObjectPointerTraitsImpl::GetObjectPropertyValue(InObject, PropertyOffset, OutValue);
+				break;
+		}
+	}
+	static void GetObjectPropertyValue(const UObject* InObject, FActorReferenceMetaData MetaData, FTrackInstancePropertyBindings* PropertyBindings, FActorReferenceComponent& OutValue)
+	{
+		switch (MetaData.Type)
+		{
+			case EActorReferenceType::ObjectPointer:
+				ObjectPointerTraitsImpl::GetObjectPropertyValue(InObject, PropertyBindings, OutValue);
+				break;
+			case EActorReferenceType::SoftObjectPointer:
+				SoftObjectPointerTraitsImpl::GetObjectPropertyValue(InObject, PropertyBindings, OutValue);
+				break;
+		}
+	}
+	static void GetObjectPropertyValue(const UObject* InObject, FActorReferenceMetaData MetaData, const FName& PropertyPath, StorageType& OutValue)
+	{
+		switch (MetaData.Type)
+		{
+			case EActorReferenceType::ObjectPointer:
+				ObjectPointerTraitsImpl::GetObjectPropertyValue(InObject, PropertyPath, OutValue);
+				break;
+			case EActorReferenceType::SoftObjectPointer:
+				SoftObjectPointerTraitsImpl::GetObjectPropertyValue(InObject, PropertyPath, OutValue);
+				break;
+		}
+	}
+
+	static void SetObjectPropertyValue(UObject* InObject, FActorReferenceMetaData MetaData, const FCustomPropertyAccessor& BaseCustomAccessor, const FActorReferenceComponent& InValue)
+	{
+		const TCustomPropertyAccessor<FActorReferencePropertyTraits>& CustomAccessor = static_cast<const TCustomPropertyAccessor<FActorReferencePropertyTraits>&>(BaseCustomAccessor);
+		(*CustomAccessor.Functions.Setter)(InObject, MetaData, InValue);
+	}
+	static void SetObjectPropertyValue(UObject* InObject, FActorReferenceMetaData MetaData, uint16 PropertyOffset, const FActorReferenceComponent& InValue)
+	{
+		switch (MetaData.Type)
+		{
+			case EActorReferenceType::ObjectPointer:
+				ObjectPointerTraitsImpl::SetObjectPropertyValue(InObject, PropertyOffset, InValue);
+				break;
+			case EActorReferenceType::SoftObjectPointer:
+				SoftObjectPointerTraitsImpl::SetObjectPropertyValue(InObject, PropertyOffset, InValue);
+				break;
+		}
+	}
+	static void SetObjectPropertyValue(UObject* InObject, FActorReferenceMetaData MetaData, FTrackInstancePropertyBindings* PropertyBindings, const FActorReferenceComponent& InValue)
+	{
+		switch (MetaData.Type)
+		{
+			case EActorReferenceType::ObjectPointer:
+				ObjectPointerTraitsImpl::SetObjectPropertyValue(InObject, PropertyBindings, InValue);
+				break;
+			case EActorReferenceType::SoftObjectPointer:
+				SoftObjectPointerTraitsImpl::SetObjectPropertyValue(InObject, PropertyBindings, InValue);
+				break;
+		}
+	}
+};
+
 using FBytePropertyTraits               = TPropertyTraits<uint8, uint8>;
 using FEnumPropertyTraits               = TPropertyTraits<uint8, uint8>;
 using FIntPropertyTraits                = TDynamicVariantPropertyTraits<int64, int64, int32, int16, int8>;
@@ -538,6 +681,7 @@ struct FMovieSceneTracksComponentTypes
 	TPropertyComponents<FStringPropertyTraits> String;
 	TPropertyComponents<FTextPropertyTraits> Text;
 	TPropertyComponents<FObjectPropertyTraits> Object;
+	TPropertyComponents<FActorReferencePropertyTraits> ActorReference;
 
 	TPropertyComponents<FFloatParameterTraits> FloatParameter;
 	TPropertyComponents<FColorParameterTraits> ColorParameter;
@@ -547,6 +691,9 @@ struct FMovieSceneTracksComponentTypes
 
 	TComponentTypeID<FConstraintComponentData> ConstraintChannel;
 
+	TComponentTypeID<FSourceActorReferenceChannel> ActorReferenceChannel;
+	TComponentTypeID<FActorReferenceComponent> ActorReferenceResult;
+
 	TComponentTypeID<TWeakObjectPtr<USceneComponent>> AttachParent;
 	TComponentTypeID<FAttachmentComponent> AttachComponent;
 	TComponentTypeID<FMovieSceneObjectBindingID> AttachParentBinding;
@@ -619,6 +766,28 @@ struct FMovieSceneTracksComponentTypes
 		TMap<int32, FMovieSceneTracksComponentTypes::FObjectPropertyRegistration::FMetaData> MetaData;
 	};
 
+	struct FActorReferencePropertyRegistration : TCustomPropertyRegistration<FActorReferencePropertyTraits>
+	{
+		struct FMetaData
+		{
+			EActorReferenceType Type = EActorReferenceType::ObjectPointer;
+		};
+
+		void Add(UClass* ClassType, FName PropertyName, GetterFunc Getter, SetterFunc Setter)
+		{
+			TCustomPropertyRegistration<FActorReferencePropertyTraits>::Add(ClassType, PropertyName, Getter, Setter);
+		}
+
+		void Add(UClass* ClassType, FName PropertyName, GetterFunc Getter, SetterFunc Setter, FMetaData InMetaData)
+		{
+			int32 CustomIndex = CustomAccessors.Num();
+			TCustomPropertyRegistration<FActorReferencePropertyTraits>::Add(ClassType, PropertyName, Getter, Setter);
+			MetaData.Add(CustomIndex, InMetaData);
+		}
+
+		TMap<int32, FMovieSceneTracksComponentTypes::FActorReferencePropertyRegistration::FMetaData> MetaData;
+	};
+
 	struct FAccessors
 	{
 		TCustomPropertyRegistration<FBoolPropertyTraits> Bool;
@@ -633,6 +802,7 @@ struct FMovieSceneTracksComponentTypes
 		TCustomPropertyRegistration<FTransformPropertyTraits> Transform;
 		TCustomPropertyRegistration<FComponentTransformPropertyTraits, 1> ComponentTransform;
 		FObjectPropertyRegistration Object;
+		FActorReferencePropertyRegistration ActorReference;
 
 	} Accessors;
 
diff --git a/Engine/Source/Runtime/MovieSceneTracks/Public/MovieSceneTracksPropertyTypes.h b/Engine/Source/Runtime/MovieSceneTracks/Public/MovieSceneTracksPropertyTypes.h
--- a/Engine/Source/Runtime/MovieSceneTracks/Public/MovieSceneTracksPropertyTypes.h
+++ b/Engine/Source/Runtime/MovieSceneTracks/Public/MovieSceneTracksPropertyTypes.h
@@ -21,6 +21,7 @@ namespace MovieScene
 
 struct FObjectComponent;
 struct FGeneratedTrackKeys;
+struct FActorReferenceComponent;
 
 struct FVectorPropertyMetaData
 {
@@ -260,5 +261,10 @@ MOVIESCENETRACKS_API void ConvertOperationalProperty(const FVector4d& In, FDoubl
 MOVIESCENETRACKS_API void ConvertOperationalProperty(const FObjectComponent& In, UObject*& Out);
 MOVIESCENETRACKS_API void ConvertOperationalProperty(UObject* In, FObjectComponent& Out);
 
+MOVIESCENETRACKS_API void ConvertOperationalProperty(const FActorReferenceComponent& InComponentData, TObjectPtr<AActor>& Out);
+MOVIESCENETRACKS_API void ConvertOperationalProperty(const FActorReferenceComponent& InComponentData, TSoftObjectPtr<AActor>& Out);
+MOVIESCENETRACKS_API void ConvertOperationalProperty(const TSoftObjectPtr<AActor> InActor, FActorReferenceComponent& Out);
+MOVIESCENETRACKS_API void ConvertOperationalProperty(const TObjectPtr<AActor> InActor, FActorReferenceComponent& Out);
+
 } // namespace MovieScene
 } // namespace UE
diff --git a/Engine/Source/Runtime/MovieSceneTracks/Public/Sections/MovieSceneActorReferenceSection.h b/Engine/Source/Runtime/MovieSceneTracks/Public/Sections/MovieSceneActorReferenceSection.h
--- a/Engine/Source/Runtime/MovieSceneTracks/Public/Sections/MovieSceneActorReferenceSection.h
+++ b/Engine/Source/Runtime/MovieSceneTracks/Public/Sections/MovieSceneActorReferenceSection.h
@@ -2,191 +2,22 @@
 
 #pragma once
 
-#include "Channels/MovieSceneChannel.h"
-#include "Channels/MovieSceneChannelData.h"
-#include "Channels/MovieSceneChannelTraits.h"
-#include "Containers/Array.h"
-#include "Containers/ArrayView.h"
-#include "Containers/Map.h"
-#include "Containers/UnrealString.h"
-#include "CoreMinimal.h"
+#include "Channels/MovieSceneActorReferenceChannel.h"
 #include "Curves/IntegralCurve.h"
-#include "Curves/KeyHandle.h"
-#include "HAL/Platform.h"
-#include "HAL/PlatformCrt.h"
-#include "Math/Range.h"
-#include "Misc/AssertionMacros.h"
-#include "Misc/FrameNumber.h"
-#include "Misc/FrameRate.h"
-#include "Misc/FrameTime.h"
-#include "Misc/Guid.h"
-#include "MovieSceneFwd.h"
-#include "MovieSceneObjectBindingID.h"
+#include "EntitySystem/IMovieSceneEntityProvider.h"
 #include "MovieSceneSection.h"
-#include "UObject/NameTypes.h"
-#include "UObject/ObjectMacros.h"
-#include "UObject/UObjectGlobals.h"
 
 #include "MovieSceneActorReferenceSection.generated.h"
 
-class IMovieScenePlayer;
-class UObject;
-struct FKeyHandle;
-struct FMovieSceneSequenceHierarchy;
 struct FMovieSceneSequenceID;
 
-USTRUCT()
-struct FMovieSceneActorReferenceKey
-{
-	GENERATED_BODY()
-
-	FMovieSceneActorReferenceKey()
-	{}
-
-	FMovieSceneActorReferenceKey(const FMovieSceneObjectBindingID& InBindingID)
-		: Object(InBindingID)
-	{}
-
-	friend bool operator==(const FMovieSceneActorReferenceKey& A, const FMovieSceneActorReferenceKey& B)
-	{
-		return A.Object == B.Object && A.ComponentName == B.ComponentName && A.SocketName == B.SocketName;
-	}
-
-	friend bool operator!=(const FMovieSceneActorReferenceKey& A, const FMovieSceneActorReferenceKey& B)
-	{
-		return A.Object != B.Object || A.ComponentName != B.ComponentName || A.SocketName != B.SocketName;
-	}
-
-	UPROPERTY(EditAnywhere, Category="Key")
-	FMovieSceneObjectBindingID Object;
-
-	UPROPERTY(EditAnywhere, Category="Key")
-	FName ComponentName;
-
-	UPROPERTY(EditAnywhere, Category="Key")
-	FName SocketName;
-};
-
-/** A curve of events */
-USTRUCT()
-struct FMovieSceneActorReferenceData : public FMovieSceneChannel
-{
-	GENERATED_BODY()
-
-	typedef FMovieSceneActorReferenceKey CurveValueType;
-
-	FMovieSceneActorReferenceData()
-		: DefaultValue()
-	{}
-
-	/**
-	 * Access a mutable interface for this channel's data
-	 *
-	 * @return An object that is able to manipulate this channel's data
-	 */
-	inline TMovieSceneChannelData<FMovieSceneActorReferenceKey> GetData()
-	{
-		return TMovieSceneChannelData<FMovieSceneActorReferenceKey>(&KeyTimes, &KeyValues, this, &KeyHandles);
-	}
-
-	/**
-	 * Access a constant interface for this channel's data
-	 *
-	 * @return An object that is able to interrogate this channel's data
-	 */
-	inline TMovieSceneChannelData<const FMovieSceneActorReferenceKey> GetData() const
-	{
-		return TMovieSceneChannelData<const FMovieSceneActorReferenceKey>(&KeyTimes, &KeyValues);
-	}
-
-	/**
-	 * Evaluate this channel
-	 *
-	 * @param InTime     The time to evaluate at
-	 * @param OutValue   A value to receive the result
-	 * @return true if the channel was evaluated successfully, false otherwise
-	 */
-	MOVIESCENETRACKS_API bool Evaluate(FFrameTime InTime, FMovieSceneActorReferenceKey& OutValue) const;
-
-public:
-
-	// ~ FMovieSceneChannel Interface
-	MOVIESCENETRACKS_API virtual void GetKeys(const TRange<FFrameNumber>& WithinRange, TArray<FFrameNumber>* OutKeyTimes, TArray<FKeyHandle>* OutKeyHandles) override;
-	MOVIESCENETRACKS_API virtual void GetKeyTimes(TArrayView<const FKeyHandle> InHandles, TArrayView<FFrameNumber> OutKeyTimes) override;
-	MOVIESCENETRACKS_API virtual void SetKeyTimes(TArrayView<const FKeyHandle> InHandles, TArrayView<const FFrameNumber> InKeyTimes) override;
-	MOVIESCENETRACKS_API virtual void DuplicateKeys(TArrayView<const FKeyHandle> InHandles, TArrayView<FKeyHandle> OutNewHandles) override;
-	MOVIESCENETRACKS_API virtual void DeleteKeys(TArrayView<const FKeyHandle> InHandles) override;
-	MOVIESCENETRACKS_API virtual void DeleteKeysFrom(FFrameNumber InTime, bool bDeleteKeysBefore) override;
-	MOVIESCENETRACKS_API virtual void RemapTimes(const UE::MovieScene::IRetimingInterface& Retimer) override;
-	MOVIESCENETRACKS_API virtual TRange<FFrameNumber> ComputeEffectiveRange() const override;
-	MOVIESCENETRACKS_API virtual int32 GetNumKeys() const override;
-	MOVIESCENETRACKS_API virtual void Reset() override;
-	MOVIESCENETRACKS_API virtual void Offset(FFrameNumber DeltaPosition) override;
-	MOVIESCENETRACKS_API virtual void ClearDefault() override;
-	MOVIESCENETRACKS_API virtual FKeyHandle GetHandle(int32 Index) override;
-	MOVIESCENETRACKS_API virtual int32 GetIndex(FKeyHandle Handle) override;
-
-public:
-
-	/**
-	 * Set this channel's default value that should be used when no keys are present
-	 *
-	 * @param InDefaultValue The desired default value
-	 */
-	inline void SetDefault(FMovieSceneActorReferenceKey InDefaultValue)
-	{
-		DefaultValue = InDefaultValue;
-	}
-
-	/**
-	 * Get this channel's default value that will be used when no keys are present
-	 *
-	 * @return (Optional) The channel's default value
-	 */
-	inline const FMovieSceneActorReferenceKey& GetDefault() const
-	{
-		return DefaultValue;
-	}
-
-	/**
-	 * Upgrade legacy data by appending to the end of the array. Assumes sorted data
-	 */
-	void UpgradeLegacyTime(UObject* Context, double Time, FMovieSceneActorReferenceKey Value)
-	{
-		FFrameRate LegacyFrameRate = GetLegacyConversionFrameRate();
-		FFrameNumber KeyTime = UpgradeLegacyMovieSceneTime(nullptr, LegacyFrameRate, Time);
-
-		check(KeyTimes.Num() == 0 || KeyTime >= KeyTimes.Last());
-
-		KeyTimes.Add(KeyTime);
-		KeyValues.Add(Value);
-	}
-
-private:
-
-	/** Sorted array of key times */
-	UPROPERTY(meta=(KeyTimes))
-	TArray<FFrameNumber> KeyTimes;
-
-	/** Default value used when there are no keys */
-	UPROPERTY()
-	FMovieSceneActorReferenceKey DefaultValue;
-
-	/** Array of values that correspond to each key time */
-	UPROPERTY(meta=(KeyValues))
-	TArray<FMovieSceneActorReferenceKey> KeyValues;
-
-	/** This needs to be a UPROPERTY so it gets saved into editor transactions but transient so it doesn't get saved into assets. */
-	UPROPERTY(Transient)
-	FMovieSceneKeyHandleMap KeyHandles;
-};
-
 /**
  * A single actor reference point section
  */
 UCLASS(MinimalAPI)
 class UMovieSceneActorReferenceSection
 	: public UMovieSceneSection
+	, public IMovieSceneEntityProvider
 {
 	GENERATED_UCLASS_BODY()
 
@@ -198,12 +29,17 @@ public:
 	//~ UMovieSceneSection interface
 	virtual void OnBindingIDsUpdated(const TMap<UE::MovieScene::FFixedObjectBindingID, UE::MovieScene::FFixedObjectBindingID>& OldFixedToNewFixedMap, FMovieSceneSequenceID LocalSequenceID, TSharedRef<UE::MovieScene::FSharedPlaybackState> SharedPlaybackState) override;
 
-	const FMovieSceneActorReferenceData& GetActorReferenceData() const { return ActorReferenceData; }
+	const FMovieSceneActorReferenceChannel& GetActorReferenceData() const { return ActorReferenceData; }
+
+private:
+
+	// IMovieSceneEntityProvider interface
+	virtual void ImportEntityImpl(UMovieSceneEntitySystemLinker* EntityLinker, const FEntityImportParams& Params, FImportedEntity* OutImportedEntity) override;
 
 private:
 
 	UPROPERTY()
-	FMovieSceneActorReferenceData ActorReferenceData;
+	FMovieSceneActorReferenceChannel ActorReferenceData;
 
 private:
 
@@ -215,13 +51,3 @@ private:
 	TArray<FString> ActorGuidStrings_DEPRECATED;
 };
 
-inline bool EvaluateChannel(const FMovieSceneActorReferenceData* InChannel, FFrameTime InTime, FMovieSceneActorReferenceKey& OutValue)
-{
-	return InChannel->Evaluate(InTime, OutValue);
-}
-
-inline bool GetChannelDefault(const FMovieSceneActorReferenceData* Channel, FMovieSceneActorReferenceKey& OutDefaultValue)
-{
-	OutDefaultValue = Channel->GetDefault();
-	return true;
-}
\ No newline at end of file
diff --git a/Engine/Source/Runtime/MovieSceneTracks/Public/Sections/MovieSceneAudioSection.h b/Engine/Source/Runtime/MovieSceneTracks/Public/Sections/MovieSceneAudioSection.h
--- a/Engine/Source/Runtime/MovieSceneTracks/Public/Sections/MovieSceneAudioSection.h
+++ b/Engine/Source/Runtime/MovieSceneTracks/Public/Sections/MovieSceneAudioSection.h
@@ -2,17 +2,17 @@
 
 #pragma once
 
+#include "Channels/MovieSceneActorReferenceChannel.h"
 #include "Channels/MovieSceneAudioTriggerChannel.h"
 #include "Channels/MovieSceneBoolChannel.h"
 #include "Channels/MovieSceneFloatChannel.h"
 #include "Channels/MovieSceneIntegerChannel.h"
 #include "Channels/MovieSceneStringChannel.h"
 #include "Components/AudioComponent.h"
-#include "Decorations/MovieSceneScalingAnchors.h"
 #include "CoreMinimal.h"
+#include "Decorations/MovieSceneScalingAnchors.h"
 #include "EntitySystem/IMovieSceneEntityProvider.h"
 #include "MovieSceneSection.h"
-#include "Sections/MovieSceneActorReferenceSection.h"
 #include "Sound/SoundAttenuation.h"
 #include "UObject/ObjectMacros.h"
 
@@ -136,7 +136,7 @@ public:
 	/*
 	 * @return The attach actor data
 	 */
-	const FMovieSceneActorReferenceData& GetAttachActorData() const { return AttachActorData; }
+	const FMovieSceneActorReferenceChannel& GetAttachActorData() const { return AttachActorData; }
 
 	/*
 	 * @return The attach component given the bound actor and the actor attach key with the component and socket names
@@ -264,7 +264,7 @@ private:
 	TMap<FName, FMovieSceneAudioTriggerChannel> Inputs_Trigger;
 
 	UPROPERTY()
-	FMovieSceneActorReferenceData AttachActorData;
+	FMovieSceneActorReferenceChannel AttachActorData;
 
 	/* Allow looping if the section length is greater than the sound duration */
 	UPROPERTY(EditAnywhere, Category = "Audio")
diff --git a/Engine/Source/Runtime/MovieSceneTracks/Public/Systems/ActorReferenceChannelEvaluatorSystem.h b/Engine/Source/Runtime/MovieSceneTracks/Public/Systems/ActorReferenceChannelEvaluatorSystem.h
new file mode 100644
--- /dev/null
+++ b/Engine/Source/Runtime/MovieSceneTracks/Public/Systems/ActorReferenceChannelEvaluatorSystem.h
@@ -0,0 +1,24 @@
+// Copyright Epic Games, Inc. All Rights Reserved.
+
+#pragma once
+
+#include "EntitySystem/MovieSceneEntitySystem.h"
+
+#include "ActorReferenceChannelEvaluatorSystem.generated.h"
+
+/**
+ * System that is responsible for evaluating actor reference channels.
+ */
+UCLASS(MinimalAPI)
+class UActorReferenceChannelEvaluatorSystem : public UMovieSceneEntitySystem
+{
+public:
+
+	GENERATED_BODY()
+
+	UActorReferenceChannelEvaluatorSystem(const FObjectInitializer& ObjInit);
+
+	virtual void OnSchedulePersistentTasks(UE::MovieScene::IEntitySystemScheduler* TaskScheduler) override;
+	virtual void OnRun(FSystemTaskPrerequisites& InPrerequisites, FSystemSubsequentTasks& Subsequents) override;
+};
+
diff --git a/Engine/Source/Runtime/MovieSceneTracks/Public/Systems/MovieSceneActorReferencePropertySystem.h b/Engine/Source/Runtime/MovieSceneTracks/Public/Systems/MovieSceneActorReferencePropertySystem.h
new file mode 100644
--- /dev/null
+++ b/Engine/Source/Runtime/MovieSceneTracks/Public/Systems/MovieSceneActorReferencePropertySystem.h
@@ -0,0 +1,18 @@
+// Copyright Epic Games, Inc. All Rights Reserved.
+
+#pragma once
+
+#include "Systems/MovieScenePropertySystem.h"
+
+#include "MovieSceneActorReferencePropertySystem.generated.h"
+
+UCLASS(MinimalAPI)
+class UMovieSceneActorReferencePropertySystem : public UMovieScenePropertySystem
+{
+	GENERATED_BODY()
+
+	UMovieSceneActorReferencePropertySystem(const FObjectInitializer& ObjInit);
+
+	MOVIESCENETRACKS_API virtual void OnRun(FSystemTaskPrerequisites& InPrerequisites, FSystemSubsequentTasks& Subsequents) override;
+};
+
diff --git a/Engine/Source/Runtime/MovieSceneTracks/Public/Tracks/MovieSceneActorReferenceTrack.h b/Engine/Source/Runtime/MovieSceneTracks/Public/Tracks/MovieSceneActorReferenceTrack.h
--- a/Engine/Source/Runtime/MovieSceneTracks/Public/Tracks/MovieSceneActorReferenceTrack.h
+++ b/Engine/Source/Runtime/MovieSceneTracks/Public/Tracks/MovieSceneActorReferenceTrack.h
@@ -5,14 +5,13 @@
 #include "CoreMinimal.h"
 #include "UObject/ObjectMacros.h"
 #include "Tracks/MovieScenePropertyTrack.h"
-#include "Compilation/IMovieSceneTrackTemplateProducer.h"
 #include "MovieSceneActorReferenceTrack.generated.h"
 
 /**
  * Handles manipulation of actor reference properties in a movie scene
  */
 UCLASS( MinimalAPI )
-class UMovieSceneActorReferenceTrack : public UMovieScenePropertyTrack, public IMovieSceneTrackTemplateProducer
+class UMovieSceneActorReferenceTrack : public UMovieScenePropertyTrack
 {
 	GENERATED_UCLASS_BODY()
 
@@ -20,5 +19,5 @@ public:
 	/** UMovieSceneTrack interface */
 	virtual bool SupportsType(TSubclassOf<UMovieSceneSection> SectionClass) const override;
 	virtual UMovieSceneSection* CreateNewSection() override;
-	virtual FMovieSceneEvalTemplatePtr CreateTemplateForSection(const UMovieSceneSection& InSection) const override;
 };
+
