[Sequencer] MetaSound Sources intermittently don't play in sequencer

#Synced-CL 48625497

diff --git a/Engine/Source/Runtime/MovieSceneTracks/Private/Systems/MovieSceneAudioSystem.cpp b/Engine/Source/Runtime/MovieSceneTracks/Private/Systems/MovieSceneAudioSystem.cpp
--- a/Engine/Source/Runtime/MovieSceneTracks/Private/Systems/MovieSceneAudioSystem.cpp
+++ b/Engine/Source/Runtime/MovieSceneTracks/Private/Systems/MovieSceneAudioSystem.cpp
@@ -1143,11 +1143,12 @@ UMovieSceneAudioSystem::FAudioComponentEvaluationData* UMovieSceneAudioSystem::G
 			}
 		}
 
-		// If no exact match, check for any AudioComponent that isn't busy
+		// If no exact match, check for any AudioComponent that isn't busy and doesn't contain a procedurally generated sound
 		for (FAudioComponentBySectionKey::ElementType& Pair : *Map)
 		{
 			UAudioComponent* ExistingComponent = Pair.Value.AudioComponent.Get();
-			if (ExistingComponent && !ExistingComponent->IsPlaying())
+			if (ExistingComponent && !ExistingComponent->IsPlaying() && 
+				(ExistingComponent->Sound == nullptr || !ExistingComponent->Sound->IsProcedurallyGenerated()))
 			{
 				// Replace this entry with the new entity ID to claim it
 				FAudioComponentEvaluationData MovedData(Pair.Value);
