Sequencer: Fix unsafe calls to GetPlaybackContext()->GetWorld() that don't verify that the Playback Context is valid first.
#rnx
#Synced-CL 48875823

diff --git a/Engine/Source/Editor/MovieSceneTools/Private/Bindings/MovieSceneSpawnableActorBindingCustomization.cpp b/Engine/Source/Editor/MovieSceneTools/Private/Bindings/MovieSceneSpawnableActorBindingCustomization.cpp
--- a/Engine/Source/Editor/MovieSceneTools/Private/Bindings/MovieSceneSpawnableActorBindingCustomization.cpp
+++ b/Engine/Source/Editor/MovieSceneTools/Private/Bindings/MovieSceneSpawnableActorBindingCustomization.cpp
@@ -69,7 +69,13 @@ void FMovieSceneSpawnableActorBindingBaseCustomization::RefreshComboList()
 
 		if (SequencerPtr.IsValid())
 		{
-			UWorld* World = SequencerPtr.Pin()->GetPlaybackContext()->GetWorld();
+			UObject* PlaybackContext = SequencerPtr.Pin()->GetPlaybackContext();
+			if (!PlaybackContext)
+			{
+				return;
+			}
+
+			UWorld* World = PlaybackContext->GetWorld();
 			if (!World)
 			{
 				return;
diff --git a/Engine/Source/Editor/MovieSceneTools/Private/LevelSequenceFBXInterop.cpp b/Engine/Source/Editor/MovieSceneTools/Private/LevelSequenceFBXInterop.cpp
--- a/Engine/Source/Editor/MovieSceneTools/Private/LevelSequenceFBXInterop.cpp
+++ b/Engine/Source/Editor/MovieSceneTools/Private/LevelSequenceFBXInterop.cpp
@@ -237,7 +237,8 @@ void FLevelSequenceFBXInterop::ExportFBXInternal(const FString& ExportFilename,
 		UMovieSceneSequence* MovieSceneSequence = Sequencer->GetFocusedMovieSceneSequence();
 		UMovieSceneSequence* RootMovieSceneSequence = Sequencer->GetRootMovieSceneSequence();
 		UMovieScene* MovieScene = MovieSceneSequence->GetMovieScene();
-		UWorld* World = Sequencer->GetPlaybackContext()->GetWorld();
+		UObject* PlaybackContext = Sequencer->GetPlaybackContext();
+		UWorld* World = PlaybackContext ? PlaybackContext->GetWorld() : nullptr;
 		FMovieSceneSequenceIDRef Template = Sequencer->GetFocusedTemplateID();
 		UnFbx::FFbxExporter::FLevelSequenceNodeNameAdapter NodeNameAdapter(MovieScene, Sequencer.Get(), Template);
 
diff --git a/Engine/Source/Editor/MovieSceneTools/Private/MovieSceneToolHelpers.cpp b/Engine/Source/Editor/MovieSceneTools/Private/MovieSceneToolHelpers.cpp
--- a/Engine/Source/Editor/MovieSceneTools/Private/MovieSceneToolHelpers.cpp
+++ b/Engine/Source/Editor/MovieSceneTools/Private/MovieSceneToolHelpers.cpp
@@ -3758,7 +3758,8 @@ private:
 		// Import static cameras first
 		ImportFBXCamera(FbxImporter, Sequence, *Sequencer, ObjectBindingMap, bMatchByNameOnly, bCreateCameras.IsSet() ? bCreateCameras.GetValue() : ImportFBXSettings->bCreateCameras);
 
-		UWorld* World = Sequencer->GetPlaybackContext()->GetWorld();
+		UObject* PlaybackContext = Sequencer->GetPlaybackContext();
+		UWorld* World = PlaybackContext ? PlaybackContext->GetWorld() : nullptr;
 		bool bValid = MovieSceneToolHelpers::ImportFBXIfReady(World, Sequence, Sequencer, Sequencer->GetFocusedTemplateID(), ObjectBindingMap, ImportFBXSettings, InOutParams, Sequencer);
 	
 		Sequencer->NotifyMovieSceneDataChanged(EMovieSceneDataChangeType::MovieSceneStructureItemAdded);
@@ -4081,16 +4082,19 @@ bool MovieSceneToolHelpers::ExportFBX(UWorld* World, const FAnimExportSequencePa
 	Exporter->SetTransformBaking(false);
 	Exporter->SetKeepHierarchy(true);
 
-	ExportLevelMesh(Exporter, World->PersistentLevel, AESP.Player, Bindings, NodeNameAdapter, Template);
-
-	// Export streaming levels and actors
-	for (ULevelStreaming* StreamingLevel : World->GetStreamingLevels())
+	if (World)
 	{
-		if (StreamingLevel)
+		ExportLevelMesh(Exporter, World->PersistentLevel, AESP.Player, Bindings, NodeNameAdapter, Template);
+
+		// Export streaming levels and actors
+		for (ULevelStreaming* StreamingLevel : World->GetStreamingLevels())
 		{
-			if (ULevel* Level = StreamingLevel->GetLoadedLevel())
+			if (StreamingLevel)
 			{
-				ExportLevelMesh(Exporter, Level, AESP.Player, Bindings, NodeNameAdapter, Template);
+				if (ULevel* Level = StreamingLevel->GetLoadedLevel())
+				{
+					ExportLevelMesh(Exporter, Level, AESP.Player, Bindings, NodeNameAdapter, Template);
+				}
 			}
 		}
 	}
diff --git a/Engine/Source/Editor/MovieSceneTools/Private/TrackEditors/LevelVisibilityTrackEditor.cpp b/Engine/Source/Editor/MovieSceneTools/Private/TrackEditors/LevelVisibilityTrackEditor.cpp
--- a/Engine/Source/Editor/MovieSceneTools/Private/TrackEditors/LevelVisibilityTrackEditor.cpp
+++ b/Engine/Source/Editor/MovieSceneTools/Private/TrackEditors/LevelVisibilityTrackEditor.cpp
@@ -194,7 +194,13 @@ void FLevelVisibilityTrackEditor::OnAddNewSection( UMovieSceneTrack* LevelVisibi
 
 void FLevelVisibilityTrackEditor::GetCurrentLevelVisibility(TArray<FName>& OutVisibleLevelNames, TArray<FName>& OutHiddenLevelNames)
 {
-	UWorld* World = GetSequencer()->GetPlaybackContext()->GetWorld();
+	UObject* PlaybackContext = GetSequencer()->GetPlaybackContext();
+	if (!PlaybackContext)
+	{
+		return;
+	}
+
+	UWorld* World = PlaybackContext->GetWorld();
 	if (!World)
 	{
 		return;
diff --git a/Engine/Source/Editor/MovieSceneTools/Private/TrackEditors/MaterialParameterCollectionTrackEditor.cpp b/Engine/Source/Editor/MovieSceneTools/Private/TrackEditors/MaterialParameterCollectionTrackEditor.cpp
--- a/Engine/Source/Editor/MovieSceneTools/Private/TrackEditors/MaterialParameterCollectionTrackEditor.cpp
+++ b/Engine/Source/Editor/MovieSceneTools/Private/TrackEditors/MaterialParameterCollectionTrackEditor.cpp
@@ -315,7 +315,11 @@ void FMaterialParameterCollectionTrackEditor::AddScalarParameter(UMovieSceneMate
 	}
 
 	float Value = Parameter.DefaultValue;
-	if (UWorld* World = GetSequencer()->GetPlaybackContext()->GetWorld())
+
+	UObject* PlaybackContext = GetSequencer()->GetPlaybackContext();
+	UWorld* World = PlaybackContext ? PlaybackContext->GetWorld() : nullptr;
+
+	if (World)
 	{
 		if (UMaterialParameterCollectionInstance* Instance = World->GetParameterCollectionInstance(Track->MPC))
 		{
@@ -340,7 +344,11 @@ void FMaterialParameterCollectionTrackEditor::AddVectorParameter(UMovieSceneMate
 	}
 
 	FLinearColor Value = Parameter.DefaultValue;
-	if (UWorld* World = GetSequencer()->GetPlaybackContext()->GetWorld())
+
+	UObject* PlaybackContext = GetSequencer()->GetPlaybackContext();
+	UWorld* World = PlaybackContext ? PlaybackContext->GetWorld() : nullptr;
+
+	if (World)
 	{
 		if (UMaterialParameterCollectionInstance* Instance = World->GetParameterCollectionInstance(Track->MPC))
 		{
diff --git a/Engine/Source/Editor/Sequencer/Private/Sequencer.cpp b/Engine/Source/Editor/Sequencer/Private/Sequencer.cpp
--- a/Engine/Source/Editor/Sequencer/Private/Sequencer.cpp
+++ b/Engine/Source/Editor/Sequencer/Private/Sequencer.cpp
@@ -9377,10 +9377,11 @@ void FSequencer::OnLoadRecordedData()
 				NotifyMovieSceneDataChanged(EMovieSceneDataChangeType::MovieSceneStructureItemAdded);
 
 			}; //callback
-			UWorld* PlaybackContext = GetPlaybackContext()->GetWorld();
+			UObject* PlaybackContext = GetPlaybackContext();
+			UWorld* World = PlaybackContext ? PlaybackContext->GetWorld() : nullptr;
 			for (const FString& FileName : OpenFilenames)
 			{
-				Recorder->LoadRecordedSequencerFile(FocusedMovieSceneSequence, PlaybackContext, FileName, OnReadComplete);
+				Recorder->LoadRecordedSequencerFile(FocusedMovieSceneSequence, World, FileName, OnReadComplete);
 			}
 		}
 	}
diff --git a/Engine/Source/Editor/Sequencer/Private/SequencerUtilities.cpp b/Engine/Source/Editor/Sequencer/Private/SequencerUtilities.cpp
--- a/Engine/Source/Editor/Sequencer/Private/SequencerUtilities.cpp
+++ b/Engine/Source/Editor/Sequencer/Private/SequencerUtilities.cpp
@@ -5249,12 +5249,18 @@ void FSequencerUtilities::SyncSectionsUsingSourceTimecode(TSharedRef<ISequencer>
 
 void FSequencerUtilities::FixActorReferences(TSharedRef<ISequencer> Sequencer)
 {
-	UWorld* PlaybackContext = Sequencer->GetPlaybackContext()->GetWorld();
+	UObject* PlaybackContext = Sequencer->GetPlaybackContext();
 	if (!PlaybackContext)
 	{
 		return;
 	}
 
+	UWorld* World = PlaybackContext->GetWorld();
+	if (!World)
+	{
+		return;
+	}
+
 	UMovieSceneSequence* FocusedSequence = Sequencer->GetFocusedMovieSceneSequence();
 	if (!FocusedSequence)
 	{
@@ -5277,7 +5283,7 @@ void FSequencerUtilities::FixActorReferences(TSharedRef<ISequencer> Sequencer)
 
 	TMap<FString, AActor*> ActorNameToActorMap;
 
-	for (TActorIterator<AActor> ActorItr(PlaybackContext); ActorItr; ++ActorItr)
+	for (TActorIterator<AActor> ActorItr(World); ActorItr; ++ActorItr)
 	{
 		// Same as with the Object Iterator, access the subclass instance with the * or -> operators.
 		AActor* Actor = *ActorItr;
diff --git a/Engine/Source/Runtime/MovieSceneTracks/Private/Systems/MovieSceneAudioSystem.cpp b/Engine/Source/Runtime/MovieSceneTracks/Private/Systems/MovieSceneAudioSystem.cpp
--- a/Engine/Source/Runtime/MovieSceneTracks/Private/Systems/MovieSceneAudioSystem.cpp
+++ b/Engine/Source/Runtime/MovieSceneTracks/Private/Systems/MovieSceneAudioSystem.cpp
@@ -400,27 +400,30 @@ private:
 			{
 				// Initialize the sound
 				UWorld* World = PlaybackContext ? PlaybackContext->GetWorld() : nullptr;
-				EvaluationData = AudioSystem->AddRootAudioComponent(InstanceHandle, AudioSection, World);
-				UAudioComponent* AudioComponent = EvaluationData ? EvaluationData->AudioComponent.Get() : nullptr;
-
-				if (ensure(AudioComponent))
+				if (World)
 				{
-					AudioSystem->PreAnimatedStorage->BeginTrackingEntity(EntityID, bWantsRestoreState, RootInstanceHandle, AudioComponent);
-					AudioSystem->PreAnimatedStorage->CachePreAnimatedValue(
-						FCachePreAnimatedValueParams(), AudioComponent,
-						[](FObjectKey InKey) { return EPreAnimatedAudioStateType::AudioComponentLifespan; });
+					EvaluationData = AudioSystem->AddRootAudioComponent(InstanceHandle, AudioSection, World);
+					UAudioComponent* AudioComponent = EvaluationData ? EvaluationData->AudioComponent.Get() : nullptr;
 
-					if (AudioSection->GetOnQueueSubtitles().IsBound())
-					{
-						AudioComponent->OnQueueSubtitles = AudioSection->GetOnQueueSubtitles();
-					}
-					if (AudioSection->GetOnAudioFinished().IsBound())
+					if (ensure(AudioComponent))
 					{
-						AudioComponent->OnAudioFinished = AudioSection->GetOnAudioFinished();
-					}
-					if (AudioSection->GetOnAudioPlaybackPercent().IsBound())
-					{
-						AudioComponent->OnAudioPlaybackPercent = AudioSection->GetOnAudioPlaybackPercent();
+						AudioSystem->PreAnimatedStorage->BeginTrackingEntity(EntityID, bWantsRestoreState, RootInstanceHandle, AudioComponent);
+						AudioSystem->PreAnimatedStorage->CachePreAnimatedValue(
+							FCachePreAnimatedValueParams(), AudioComponent,
+							[](FObjectKey InKey) { return EPreAnimatedAudioStateType::AudioComponentLifespan; });
+
+						if (AudioSection->GetOnQueueSubtitles().IsBound())
+						{
+							AudioComponent->OnQueueSubtitles = AudioSection->GetOnQueueSubtitles();
+						}
+						if (AudioSection->GetOnAudioFinished().IsBound())
+						{
+							AudioComponent->OnAudioFinished = AudioSection->GetOnAudioFinished();
+						}
+						if (AudioSection->GetOnAudioPlaybackPercent().IsBound())
+						{
+							AudioComponent->OnAudioPlaybackPercent = AudioSection->GetOnAudioPlaybackPercent();
+						}
 					}
 				}
 			}
