For LevelSequences, the Camera Cut Tracks now call SetViewTarget on the PlayerController rather than the PlayerController's CameraManager. APlayerController::SetViewTarget is already configured to look up the CameraManager and call SetViewTarget, but other APlayerController subclasses will now receive the notification of the changing view target and can customize  how they respond to that event.

The console variable "Sequencer.CameraCutSetViewTargetThroughPlayerController" can be set to false to restore the previous behavior of calling SetViewTarget on the camera manager.

#rnx

diff --git a/Engine/Source/Runtime/MovieSceneTracks/Private/TrackInstances/MovieSceneCameraCutGameHandler.cpp b/Engine/Source/Runtime/MovieSceneTracks/Private/TrackInstances/MovieSceneCameraCutGameHandler.cpp
--- a/Engine/Source/Runtime/MovieSceneTracks/Private/TrackInstances/MovieSceneCameraCutGameHandler.cpp
+++ b/Engine/Source/Runtime/MovieSceneTracks/Private/TrackInstances/MovieSceneCameraCutGameHandler.cpp
@@ -22,8 +22,14 @@
 
 namespace UE::MovieScene
 {
-
-const APlayerController* GetPlaybackController(const UObject* PlaybackContext)
+static bool bCameraCutSetViewTargetThroughPlayerController = true;
+FAutoConsoleVariableRef CVarCameraCutSetViewTargetThroughPlayerController(
+	TEXT("Sequencer.CameraCutSetViewTargetThroughPlayerController"),
+	bCameraCutSetViewTargetThroughPlayerController,
+	TEXT("When true (default), camera cut tracks will call SetViewTarget on the affected player controller instead of that controller's camera manager.\n")
+);
+
+APlayerController* GetPlaybackController(const UObject* PlaybackContext)
 {
 	UWorld* World = PlaybackContext ? PlaybackContext->GetWorld() : nullptr;
 	if (World == nullptr || World->GetGameInstance() == nullptr)
@@ -33,7 +39,7 @@ const APlayerController* GetPlaybackController(const UObject* PlaybackContext)
 	
 	if (const AActor* PlaybackContextActor = Cast<AActor>(PlaybackContext))
 	{
-		if (const APlayerController* OwnerController = Cast<APlayerController>(PlaybackContextActor->GetOwner()))
+		if (APlayerController* OwnerController = Cast<APlayerController>(PlaybackContextActor->GetOwner()))
 		{
 			return OwnerController;
 		}
@@ -106,17 +112,28 @@ void FPreAnimatedCameraCutTraits::RestorePreAnimatedValue(
 		}
 	}
 
-	APlayerCameraManager* CameraManager = (PC != nullptr) ? PC->PlayerCameraManager.Get() : nullptr;
-
 	// Restore previous view target.
 	// If the previous view target is not valid anymore, we still set it on the camera manger. This will by
 	// default fall back to using the player controller as the view target.
-	if (CameraManager)
+	if (PC)
 	{
 		AActor* PreviousViewTarget = Cast<AActor>(CachedValue.LastViewTarget.ResolveObjectPtr());
-		CameraManager->SetViewTarget(PreviousViewTarget);
-		// TODO james.fleming ideally we would cache this before, just in case it had been set true (which is not usual, but could be possible) 
-		CameraManager->bClientSimulatingViewTarget = false;
+
+		if (bCameraCutSetViewTargetThroughPlayerController)
+		{
+			PC->SetViewTarget(PreviousViewTarget);
+		}
+
+		if (APlayerCameraManager* CameraManager = PC->PlayerCameraManager.Get())
+		{
+			if (!bCameraCutSetViewTargetThroughPlayerController)
+			{
+				CameraManager->SetViewTarget(PreviousViewTarget);
+			}
+
+			// TODO james.fleming ideally we would cache this before, just in case it had been set true (which is not usual, but could be possible) 
+			CameraManager->bClientSimulatingViewTarget = false;
+		}
 	}
 
 	// Restore previous aspect ratio axis constraint. Use the cached local player if there's no local player
@@ -253,7 +270,7 @@ void FCameraCutGameHandler::SetCameraCut(
 	}
 
 	CA_SUPPRESS(6011);
-	const APlayerController* PC = GetPlaybackController(PlaybackContext);
+	APlayerController* PC = GetPlaybackController(PlaybackContext);
 	APlayerCameraManager* CameraManager = (PC != nullptr) ? PC->PlayerCameraManager.Get() : nullptr;
 
 	// If the player controller is missing, there is no camera manager for us to manage the view target
@@ -382,9 +399,22 @@ void FCameraCutGameHandler::SetCameraCut(
 		UE_LOG(LogMovieScene, Verbose, TEXT("Starting new camera cut: '%s'"),
 			(CameraActor ? *CameraActor->GetName() : TEXT("None")));
 	}
-	if (bDoSetViewTarget && ensureMsgf(CameraManager, TEXT("Can't set view target when there is no player controller!")))
+	if (bDoSetViewTarget)
 	{
-		CameraManager->SetViewTarget(CameraActor, TransitionParams);
+		if (bCameraCutSetViewTargetThroughPlayerController)
+		{
+			if (ensureMsgf(PC, TEXT("Can't set view target when there is no player controller!")))
+			{
+				PC->SetViewTarget(CameraActor, TransitionParams);
+			}
+		}
+		else
+		{
+			if (ensureMsgf(CameraManager, TEXT("Can't set view target when there is no camera manager!")))
+			{
+				CameraManager->SetViewTarget(CameraActor, TransitionParams);
+			}
+		}
 	}
 
 	// Override the aspect ratio constraint if this sequence requires it.
diff --git a/Engine/Source/Runtime/MovieSceneTracks/Private/TrackInstances/MovieSceneCameraCutGameHandler.h b/Engine/Source/Runtime/MovieSceneTracks/Private/TrackInstances/MovieSceneCameraCutGameHandler.h
--- a/Engine/Source/Runtime/MovieSceneTracks/Private/TrackInstances/MovieSceneCameraCutGameHandler.h
+++ b/Engine/Source/Runtime/MovieSceneTracks/Private/TrackInstances/MovieSceneCameraCutGameHandler.h
@@ -72,5 +72,5 @@ private:
 	const FSequenceInstance& SequenceInstance;
 };
 
-	const APlayerController* GetPlaybackController(const UObject* PlaybackContext);
+	APlayerController* GetPlaybackController(const UObject* PlaybackContext);
 }  // namespace UE::MovieScene
