Sequencer: Enable completion mode for actor reference track

This allows the actor reference track to be disabled/muted and have it restore state, etc.

#Synced-CL 48964464

diff --git a/Engine/Source/Runtime/MovieSceneTracks/Private/Sections/MovieSceneActorReferenceSection.cpp b/Engine/Source/Runtime/MovieSceneTracks/Private/Sections/MovieSceneActorReferenceSection.cpp
--- a/Engine/Source/Runtime/MovieSceneTracks/Private/Sections/MovieSceneActorReferenceSection.cpp
+++ b/Engine/Source/Runtime/MovieSceneTracks/Private/Sections/MovieSceneActorReferenceSection.cpp
@@ -13,6 +13,8 @@
 UMovieSceneActorReferenceSection::UMovieSceneActorReferenceSection( const FObjectInitializer& ObjectInitializer )
 	: Super( ObjectInitializer )
 {
+	EvalOptions.EnableAndSetCompletionMode(EMovieSceneCompletionMode::ProjectDefault);
+
 #if WITH_EDITOR
 
 	ChannelProxy = MakeShared<FMovieSceneChannelProxy>(ActorReferenceData, FMovieSceneChannelMetaData());
