[Sequencer] Audio should not default to looping when added to a sequence

#Synced-CL 49069346

diff --git a/Engine/Source/Editor/MovieSceneTools/Private/TrackEditors/AudioTrackEditor.cpp b/Engine/Source/Editor/MovieSceneTools/Private/TrackEditors/AudioTrackEditor.cpp
--- a/Engine/Source/Editor/MovieSceneTools/Private/TrackEditors/AudioTrackEditor.cpp
+++ b/Engine/Source/Editor/MovieSceneTools/Private/TrackEditors/AudioTrackEditor.cpp
@@ -886,10 +886,12 @@ FText FAudioSection::GetSectionToolTip() const
 
 		const float AudioStartTime = AudioSection->GetStartOffset() / TickResolution;
 		const float SectionDuration = (AudioSection->GetExclusiveEndFrame() - AudioSection->GetInclusiveStartFrame())/ TickResolution;
+		
+		FText BaseToolTip;
 
 		if (AudioSection->GetLooping())
 		{
-			ToolTip = FText::Format(LOCTEXT("ToolTipContentFormatLooping", "Start: {0}s\nDuration: {1}s\nLooping"),
+			BaseToolTip = FText::Format(LOCTEXT("ToolTipContentFormatLooping", "Start: {0}s\nDuration: {1}s\nLooping"),
 				AudioStartTime,
 				SectionDuration);
 		}
@@ -900,12 +902,34 @@ FText FAudioSection::GetSectionToolTip() const
 
 			if (Duration > 0.0f)
 			{
-				ToolTip = FText::Format(LOCTEXT("ToolTipContentFormat", "{0}s - {1}s ({2} seconds)"),
+				BaseToolTip = FText::Format(LOCTEXT("ToolTipContentFormat", "{0}s - {1}s ({2} seconds)"),
 					AudioStartTime,
 					AudioStartTime + Duration,
 					Duration);
 			}
 		}
+
+		FText CompletionModeToolTip;
+		switch (AudioSection->EvalOptions.CompletionMode)
+		{
+		case EMovieSceneCompletionMode::KeepState:
+			CompletionModeToolTip = LOCTEXT("ToolTipCompletionModeKeep", "\nPlays to end of sound");
+			break;
+		case EMovieSceneCompletionMode::RestoreState:
+			CompletionModeToolTip = LOCTEXT("ToolTipCompletionModeRestore", "\nStops on section end");
+			break;
+		default:
+			CompletionModeToolTip = LOCTEXT("ToolTipCompletionModeDefault", "\nCompletionMode: ProjectDefault");
+		}
+
+		FText LoopStateDivergenceToolTip;
+
+		if (AudioSection->GetLooping() == AudioSection->GetSound()->IsOneShot())
+		{
+			LoopStateDivergenceToolTip = LOCTEXT("ToolTipLoopStateDivergence", "\nOverriding loop state from source content");
+		}
+
+		ToolTip = FText::Format(LOCTEXT("FullToolTipContentFormat", "{0}{1}{2}"), BaseToolTip, LoopStateDivergenceToolTip, CompletionModeToolTip);
 	}
 
 	ToolTip = FText::Join(FText::GetEmpty(), ToolTip, AttachToolTip);
diff --git a/Engine/Source/Runtime/MovieSceneTracks/Private/Sections/MovieSceneAudioSection.cpp b/Engine/Source/Runtime/MovieSceneTracks/Private/Sections/MovieSceneAudioSection.cpp
--- a/Engine/Source/Runtime/MovieSceneTracks/Private/Sections/MovieSceneAudioSection.cpp
+++ b/Engine/Source/Runtime/MovieSceneTracks/Private/Sections/MovieSceneAudioSection.cpp
@@ -290,6 +290,18 @@ void UMovieSceneAudioSection::SetupSoundInputParameters(USoundBase* InSoundBase)
 	}
 }
 
+void UMovieSceneAudioSection::SetCompletionMode()
+{
+	if (bLooping)
+	{
+		EvalOptions.CompletionMode = EMovieSceneCompletionMode::RestoreState;
+	}
+	else
+	{
+		EvalOptions.CompletionMode = EMovieSceneCompletionMode::KeepState;
+	}
+}
+
 TOptional<FFrameTime> UMovieSceneAudioSection::GetOffsetTime() const
 {
 	return TOptional<FFrameTime>(StartFrameOffset);
@@ -351,6 +363,43 @@ void UMovieSceneAudioSection::PostLoad()
 		StartFrameOffset = ConvertFrameTime(FFrameTime::FromDecimal(DisplayRate.AsDecimal() * StartOffsetToUpgrade.GetValue()), DisplayRate, TickResolution).FrameNumber;
 	}
 }
+
+#if WITH_EDITOR
+void UMovieSceneAudioSection::PostEditChangeProperty(FPropertyChangedEvent& PropertyChangedEvent)
+{
+	Super::PostEditChangeProperty(PropertyChangedEvent);
+
+	if (PropertyChangedEvent.MemberProperty && Sound)
+	{
+		const FName bLoopingName = GET_MEMBER_NAME_CHECKED(UMovieSceneAudioSection, bLooping);
+
+		// Clamp the StartOffset to the length of the SoundWave when not looping
+		// Setting the StartOffset to the size of the SoundWave will result in silence, irrespective of the size of the section
+		const FName PropertyName = PropertyChangedEvent.MemberProperty->GetFName();
+		const bool bHasDeterministicDuration = !bLooping && !Sound->IsProcedurallyGenerated() && Sound->IsOneShot();
+
+		if ((PropertyName == GET_MEMBER_NAME_CHECKED(UMovieSceneAudioSection, StartFrameOffset) || (PropertyName == bLoopingName))
+			&& bHasDeterministicDuration)
+		{
+			// SoundDuration is local-specifc, meaning it can vary based on the value of GetCurrentLanguage()
+			const float SoundDuration = MovieSceneHelpers::GetSoundDuration(Sound);
+			const FFrameRate FrameRate = GetTypedOuter<UMovieScene>()->GetTickResolution();
+			const FFrameNumber MaxFrame = (SoundDuration * FrameRate).FloorToFrame();
+
+			StartFrameOffset = FMath::Min(StartFrameOffset, MaxFrame);
+		}
+
+		// Change Completion Mode to prevent looping sounds from playing indefinitely.
+		if (PropertyName == bLoopingName)
+		{
+			if (Sound->IsOneShot())
+			{
+				SetCompletionMode();
+			}
+		}
+	}
+}
+#endif // WITH_EDITOR
 	
 TOptional<TRange<FFrameNumber> > UMovieSceneAudioSection::GetAutoSizeRange() const
 {
@@ -419,14 +468,15 @@ UMovieSceneSection* UMovieSceneAudioSection::SplitSection(FQualifiedFrameTime Sp
 
 void UMovieSceneAudioSection::SetSound(USoundBase* InSound)
 {
+	const bool bIsNewSound = (Sound == nullptr);
+
 	Sound = InSound;
 
-	// For a procedurally generated sound, only have the section loop if it's specifically a looping sound
-	// this is because there's no way to determine a duration for the sound so the user has no way of
-	// now when to end the region to avoid retriggering the one-shot procedural asset.
-	if (Sound && Sound->IsProcedurallyGenerated())
+	if (Sound && bIsNewSound)
 	{
 		bLooping = !Sound->IsOneShot();
+
+		SetCompletionMode();
 	}
 
 	InvalidateChannelProxy();
diff --git a/Engine/Source/Runtime/MovieSceneTracks/Public/Sections/MovieSceneAudioSection.h b/Engine/Source/Runtime/MovieSceneTracks/Public/Sections/MovieSceneAudioSection.h
--- a/Engine/Source/Runtime/MovieSceneTracks/Public/Sections/MovieSceneAudioSection.h
+++ b/Engine/Source/Runtime/MovieSceneTracks/Public/Sections/MovieSceneAudioSection.h
@@ -145,6 +145,9 @@ public:
 
 	/** ~UObject interface */
 	MOVIESCENETRACKS_API virtual void PostLoad() override;
+#if WITH_EDITOR
+	MOVIESCENETRACKS_API virtual void PostEditChangeProperty(FPropertyChangedEvent& PropertyChangedEvent) override;
+#endif //WITH_EDITOR
 
 	/** Called when subtitles are sent to the SubtitleManager.  Set this delegate if you want to hijack the subtitles for other purposes */
 	void SetOnQueueSubtitles(const FOnQueueSubtitles& InOnQueueSubtitles)
@@ -219,6 +222,8 @@ private:
 
 	MOVIESCENETRACKS_API void SetupSoundInputParameters(USoundBase* InSoundBase);
 
+	MOVIESCENETRACKS_API void SetCompletionMode();
+
 	/** The sound cue or wave that this section plays */
 	UPROPERTY(EditAnywhere, Category="Audio", BlueprintGetter=GetSound, BlueprintSetter=SetSound)
 	TObjectPtr<USoundBase> Sound;
