Sequencer: add hierarchical-bias support to component attachment track

#Synced-CL 48703199

diff --git a/Engine/Source/Runtime/MovieSceneTracks/Private/Systems/MovieSceneComponentAttachmentSystem.cpp b/Engine/Source/Runtime/MovieSceneTracks/Private/Systems/MovieSceneComponentAttachmentSystem.cpp
--- a/Engine/Source/Runtime/MovieSceneTracks/Private/Systems/MovieSceneComponentAttachmentSystem.cpp
+++ b/Engine/Source/Runtime/MovieSceneTracks/Private/Systems/MovieSceneComponentAttachmentSystem.cpp
@@ -7,6 +7,7 @@
 
 #include "EntitySystem/MovieSceneBoundObjectInstantiator.h"
 #include "EntitySystem/MovieSceneBoundSceneComponentInstantiator.h"
+#include "EntitySystem/MovieSceneEntityGroupingSystem.h"
 #include "EntitySystem/Interrogation/MovieSceneInterrogationLinker.h"
 
 #include "PreAnimatedState/MovieScenePreAnimatedComponentTransformStorage.h"
@@ -28,6 +29,50 @@ namespace UE
 namespace MovieScene
 {
 
+struct FComponentAttachmentGroupingPolicy
+{
+	using GroupKeyType = FObjectKey;
+
+	void InitializeGroupKeys(
+		TEntityGroupingHandlerBase<FComponentAttachmentGroupingPolicy>& Handler,
+		FEntityGroupBuilder* Builder,
+		FEntityAllocationIteratorItem Item,
+		FReadEntityIDs EntityIDs,
+		TWrite<FEntityGroupID> GroupIDs,
+		TRead<UObject*> BoundObjects,
+		TRead<FAttachmentComponent> AttachComponents)
+	{
+		const FEntityAllocation* Allocation = Item.GetAllocation();
+		for (int32 Index = 0; Index < Allocation->Num(); ++Index)
+		{
+			if (!GroupIDs[Index].IsValid())
+			{
+				GroupKeyType Key = FObjectKey(BoundObjects[Index]);
+				const int32 NewGroupIndex = Handler.GetOrAllocateGroupIndex(Key, Builder);
+				FEntityGroupID NewGroupID = Builder->MakeGroupID(NewGroupIndex);
+
+				Builder->AddEntityToGroup(EntityIDs[Index], NewGroupID);
+
+				GroupIDs[Index] = NewGroupID;
+			}
+		}
+	}
+
+#if WITH_EDITOR
+	bool OnObjectsReplaced(GroupKeyType& InOutKey, const TMap<UObject*, UObject*>& ReplacementMap)
+	{
+		if (UObject* KeyObj = InOutKey.ResolveObjectPtr())
+		{
+			if (UObject* NewKeyObj = ReplacementMap.FindRef(KeyObj))
+			{
+				InOutKey = FObjectKey(NewKeyObj);
+				return true;
+			}
+		}
+		return false;
+	}
+#endif
+};
 
 struct FInitializeAttachParentsTask
 {
@@ -53,6 +98,7 @@ struct FInitializeAttachParentsTask
 struct FAttachmentHandler
 {
 	UMovieSceneComponentAttachmentSystem* AttachmentSystem;
+	FBuiltInComponentTypes* BuiltInComponents;
 	FMovieSceneTracksComponentTypes* TrackComponents;
 	FEntityManager* EntityManager;
 
@@ -60,6 +106,7 @@ struct FAttachmentHandler
 		: AttachmentSystem(InAttachmentSystem)
 	{
 		EntityManager = &InAttachmentSystem->GetLinker()->EntityManager;
+		BuiltInComponents = FBuiltInComponentTypes::Get();
 		TrackComponents = FMovieSceneTracksComponentTypes::Get();
 	}
 
@@ -83,6 +130,11 @@ struct FAttachmentHandler
 
 		for (FMovieSceneEntityID Entity : Inputs)
 		{
+			if (EntityManager->HasComponent(Entity, BuiltInComponents->Tags.Ignored))
+			{
+				continue;
+			}
+
 			TOptionalComponentReader<TWeakObjectPtr<USceneComponent>> AttachParentComponent = EntityManager->ReadComponent(Entity, TrackComponents->AttachParent);
 			TOptionalComponentReader<FAttachmentComponent>            AttachmentComponent   = EntityManager->ReadComponent(Entity, TrackComponents->AttachComponent);
 			if (AttachParentComponent && AttachmentComponent)
@@ -174,6 +226,7 @@ UMovieSceneComponentAttachmentSystem::UMovieSceneComponentAttachmentSystem(const
 
 	SystemCategories |= FSystemInterrogator::GetExcludedFromInterrogationCategory();
 
+	FBuiltInComponentTypes* BuiltInComponents = FBuiltInComponentTypes::Get();
 	FMovieSceneTracksComponentTypes* TrackComponents = FMovieSceneTracksComponentTypes::Get();
 	RelevantComponent = TrackComponents->AttachParentBinding;
 
@@ -186,13 +239,16 @@ UMovieSceneComponentAttachmentSystem::UMovieSceneComponentAttachmentSystem(const
 		DefineImplicitPrerequisite(UMovieSceneComponentAttachmentInvalidatorSystem::StaticClass(), GetClass());
 		DefineImplicitPrerequisite(GetClass(), UMovieSceneRestorePreAnimatedStateSystem::StaticClass());
 
-		DefineComponentConsumer(GetClass(), FBuiltInComponentTypes::Get()->BoundObject);
-		DefineComponentConsumer(GetClass(), FBuiltInComponentTypes::Get()->SymbolicTags.CreatesEntities);
+		DefineComponentConsumer(GetClass(), BuiltInComponents->BoundObject);
+		DefineComponentConsumer(GetClass(), BuiltInComponents->Tags.Ignored);
+		DefineComponentConsumer(GetClass(), BuiltInComponents->SymbolicTags.CreatesEntities);
 	}
 }
 
 void UMovieSceneComponentAttachmentSystem::OnLink()
 {
+	using namespace UE::MovieScene;
+
 	UMovieSceneRestorePreAnimatedStateSystem* RestoreSystem = Linker->LinkSystem<UMovieSceneRestorePreAnimatedStateSystem>();
 	Linker->SystemGraph.AddReference(this, RestoreSystem);
 
@@ -200,6 +256,13 @@ void UMovieSceneComponentAttachmentSystem::OnLink()
 	Linker->SystemGraph.AddReference(this, AttachmentInvalidator);
 
 	AttachmentTracker.Initialize(this);
+
+	FBuiltInComponentTypes*          BuiltInComponents = FBuiltInComponentTypes::Get();
+	FMovieSceneTracksComponentTypes* TracksComponents  = FMovieSceneTracksComponentTypes::Get();
+
+	UMovieSceneEntityGroupingSystem* GroupingSystem = Linker->LinkSystem<UMovieSceneEntityGroupingSystem>();
+	GroupingKey = GroupingSystem->AddGrouping(
+			FComponentAttachmentGroupingPolicy(), BuiltInComponents->BoundObject, TracksComponents->AttachComponent);
 }
 
 void UMovieSceneComponentAttachmentSystem::OnUnlink()
@@ -207,6 +270,13 @@ void UMovieSceneComponentAttachmentSystem::OnUnlink()
 	using namespace UE::MovieScene;
 
 	AttachmentTracker.Destroy(FAttachmentHandler(this));
+
+	UMovieSceneEntityGroupingSystem* GroupingSystem = Linker->FindSystem<UMovieSceneEntityGroupingSystem>();
+	if (ensure(GroupingSystem))
+	{
+		GroupingSystem->RemoveGrouping(GroupingKey);
+	}
+	GroupingKey = FEntityGroupingPolicyKey();
 }
 
 void UMovieSceneComponentAttachmentSystem::OnRun(FSystemTaskPrerequisites& InPrerequisites, FSystemSubsequentTasks& Subsequents)
@@ -227,6 +297,7 @@ void UMovieSceneComponentAttachmentSystem::OnRun(FSystemTaskPrerequisites& InPre
 	.Read(TrackComponents->AttachComponent)
 	.Write(TrackComponents->AttachParent)
 	.FilterAll({ Components->Tags.NeedsLink })
+	.FilterNone({ Components->Tags.Ignored })
 	.RunInline_PerEntity(&Linker->EntityManager, InitAttachParents);
 
 	// Step 2: Update all invalidated inputs and outputs for attachments
diff --git a/Engine/Source/Runtime/MovieSceneTracks/Public/Systems/MovieSceneComponentAttachmentSystem.h b/Engine/Source/Runtime/MovieSceneTracks/Public/Systems/MovieSceneComponentAttachmentSystem.h
--- a/Engine/Source/Runtime/MovieSceneTracks/Public/Systems/MovieSceneComponentAttachmentSystem.h
+++ b/Engine/Source/Runtime/MovieSceneTracks/Public/Systems/MovieSceneComponentAttachmentSystem.h
@@ -3,7 +3,9 @@
 #pragma once
 
 #include "EntitySystem/MovieSceneEntityInstantiatorSystem.h"
+
 #include "MovieSceneTracksComponentTypes.h"
+#include "EntitySystem/BuiltInComponentTypes.h"
 #include "EntitySystem/MovieScenePreAnimatedStateSystem.h"
 #include "EntitySystem/MovieSceneOverlappingEntityTracker.h"
 
@@ -71,6 +73,8 @@ private:
 	UE::MovieScene::TOverlappingEntityTracker<UE::MovieScene::FPreAnimAttachment, FObjectKey> AttachmentTracker;
 
 	TArray<TTuple<USceneComponent*, UE::MovieScene::FPreAnimAttachment>> PendingAttachmentsToRestore;
+
+	UE::MovieScene::FEntityGroupingPolicyKey GroupingKey;
 };
 
 
