Sequencer: Fix for scrubbing audio with modulation

#Synced-CL 49295446

diff --git a/Engine/Source/Runtime/MovieSceneTracks/Private/Systems/MovieSceneAudioSystem.cpp b/Engine/Source/Runtime/MovieSceneTracks/Private/Systems/MovieSceneAudioSystem.cpp
--- a/Engine/Source/Runtime/MovieSceneTracks/Private/Systems/MovieSceneAudioSystem.cpp
+++ b/Engine/Source/Runtime/MovieSceneTracks/Private/Systems/MovieSceneAudioSystem.cpp
@@ -794,6 +794,7 @@ private:
 					if (!AudioComponent.IsPlaying())
 					{
 						ScrubbedSound->SetSoundWave(SoundWave);
+						ScrubbedSound->ModulationSettings = SoundWave->ModulationSettings;
 						ScrubbedSound->SetPlayheadTime(AudioTime);
 						float MaxScrubWidthSeconds = 0.001f * (float)ScrubWidthMillisecondsCVar;
 						ScrubbedSound->SetGrainDurationRange({ MaxScrubWidthSeconds, 0.05f });
