Fix race when writing to sparse array in FGatherTransformOriginsFromSubscenes.

Previously FGatherTransformOriginsFromSubscenes was forked per allocation, causing multiple tasks to run in parrallel, since it writes to an array, this task isn't thread safe and can cause a crash. The task is now scheduled per allocation, so it will be run sequentially preventing the race condition.

#Synced-CL 49164812

diff --git a/Engine/Source/Runtime/MovieSceneTracks/Private/Systems/MovieSceneTransformOriginSystem.cpp b/Engine/Source/Runtime/MovieSceneTracks/Private/Systems/MovieSceneTransformOriginSystem.cpp
--- a/Engine/Source/Runtime/MovieSceneTracks/Private/Systems/MovieSceneTransformOriginSystem.cpp
+++ b/Engine/Source/Runtime/MovieSceneTracks/Private/Systems/MovieSceneTransformOriginSystem.cpp
@@ -374,7 +374,7 @@ void UMovieSceneTransformOriginSystem::OnSchedulePersistentTasks(UE::MovieScene:
 	.FilterNone({BuiltInComponents->Tags.ImportedEntity})
 	.CombineFilter(LocationAndRotationFilterResults)
 	.SetParams(FTaskParams(TEXT("Gather Transform Origins From Subscenes")).ForcePrePostTask())
-	.Fork_PerAllocation<FGatherTransformOriginsFromSubscenes>(&Linker->EntityManager, TaskScheduler, &InstanceHandleToParentHandle, &TransformOriginsByInstanceID, InstanceRegistry);
+	.Schedule_PerAllocation<FGatherTransformOriginsFromSubscenes>(&Linker->EntityManager, TaskScheduler, &InstanceHandleToParentHandle, &TransformOriginsByInstanceID, InstanceRegistry);
 
 	FTaskID AssignTask = FEntityTaskBuilder()
 	.Read(BuiltInComponents->InstanceHandle)
