Constraints: fixed constraint channel label

#rnx
#Synced-CL 49605136

diff --git a/Engine/Source/Runtime/MovieSceneTracks/Private/Sections/MovieScene3DTransformSection.cpp b/Engine/Source/Runtime/MovieSceneTracks/Private/Sections/MovieScene3DTransformSection.cpp
--- a/Engine/Source/Runtime/MovieSceneTracks/Private/Sections/MovieScene3DTransformSection.cpp
+++ b/Engine/Source/Runtime/MovieSceneTracks/Private/Sections/MovieScene3DTransformSection.cpp
@@ -812,7 +812,8 @@ EMovieSceneChannelProxyType UMovieScene3DTransformSection::CacheChannelProxy()
 				{
 					if (WeakConstraint.IsValid())
 					{
-						FString ParentStr; WeakConstraint->GetLabel().Split(TEXT("."), &ParentStr, nullptr);
+						constexpr TCHAR ParentChildDelimiter[] = TEXT("|");
+						FString ParentStr; WeakConstraint->GetLabel().Split(ParentChildDelimiter, &ParentStr, nullptr);
 						if (!ParentStr.IsEmpty())
 						{
 							return ParentStr;
