ECS & Override blending: fix early exit due to override type not being used

#rnx

#Synced-CL 49645631

diff --git a/Engine/Source/Runtime/MovieSceneTracks/Private/Systems/MovieScenePiecewiseDoubleBlenderSystem.cpp b/Engine/Source/Runtime/MovieSceneTracks/Private/Systems/MovieScenePiecewiseDoubleBlenderSystem.cpp
--- a/Engine/Source/Runtime/MovieSceneTracks/Private/Systems/MovieScenePiecewiseDoubleBlenderSystem.cpp
+++ b/Engine/Source/Runtime/MovieSceneTracks/Private/Systems/MovieScenePiecewiseDoubleBlenderSystem.cpp
@@ -999,7 +999,7 @@ void UMovieScenePiecewiseDoubleBlenderSystem::ReinitializeAccumulationBuffers()
 		const bool bHasOverrides		 = AnyCompositeAllocationsMask.Contains(BuiltInComponents->Tags.OverrideBlend);
 		const bool bHasAdditivesFromBase = AnyCompositeAllocationsMask.Contains(BuiltInComponents->Tags.AdditiveFromBaseBlend);
 
-		if (!(bHasAbsolutes || bHasRelatives || bHasAdditives || bHasAdditivesFromBase))
+		if (!(bHasAbsolutes || bHasRelatives || bHasAdditives || bHasAdditivesFromBase || bHasOverrides))
 		{
 			continue;
 		}
