[Sequencer] Audio should not default to looping when added to a sequence

#Synced-CL 49558946

diff --git a/Engine/Source/Runtime/MovieSceneTracks/Private/Sections/MovieSceneAudioSection.cpp b/Engine/Source/Runtime/MovieSceneTracks/Private/Sections/MovieSceneAudioSection.cpp
--- a/Engine/Source/Runtime/MovieSceneTracks/Private/Sections/MovieSceneAudioSection.cpp
+++ b/Engine/Source/Runtime/MovieSceneTracks/Private/Sections/MovieSceneAudioSection.cpp
@@ -86,6 +86,7 @@ UMovieSceneAudioSection::UMovieSceneAudioSection( const FObjectInitializer& Obje
 		(GetLinkerCustomVersion(FSequencerObjectVersion::GUID) < FSequencerObjectVersion::WhenFinishedDefaultsToProjectDefault ? 
 			EMovieSceneCompletionMode::RestoreState : 
 			EMovieSceneCompletionMode::ProjectDefault);
+	CachedCompletionMode = EvalOptions.CompletionMode;
 
 	SoundVolume.SetDefault(1.f);
 	PitchMultiplier.SetDefault(1.f);
@@ -292,13 +293,13 @@ void UMovieSceneAudioSection::SetupSoundInputParameters(USoundBase* InSoundBase)
 
 void UMovieSceneAudioSection::SetCompletionMode()
 {
-	if (bLooping)
+	if (bLooping) // All looping sounds are set to restore state to prevent "stuck on" sounds
 	{
 		EvalOptions.CompletionMode = EMovieSceneCompletionMode::RestoreState;
 	}
 	else
 	{
-		EvalOptions.CompletionMode = EMovieSceneCompletionMode::KeepState;
+		EvalOptions.CompletionMode = CachedCompletionMode;
 	}
 }
 
@@ -372,6 +373,7 @@ void UMovieSceneAudioSection::PostEditChangeProperty(FPropertyChangedEvent& Prop
 	if (PropertyChangedEvent.MemberProperty && Sound)
 	{
 		const FName bLoopingName = GET_MEMBER_NAME_CHECKED(UMovieSceneAudioSection, bLooping);
+		const FName EvalOptionsName = GET_MEMBER_NAME_CHECKED(UMovieSceneSection, EvalOptions);
 
 		// Clamp the StartOffset to the length of the SoundWave when not looping
 		// Setting the StartOffset to the size of the SoundWave will result in silence, irrespective of the size of the section
@@ -392,10 +394,15 @@ void UMovieSceneAudioSection::PostEditChangeProperty(FPropertyChangedEvent& Prop
 		// Change Completion Mode to prevent looping sounds from playing indefinitely.
 		if (PropertyName == bLoopingName)
 		{
-			if (Sound->IsOneShot())
-			{
-				SetCompletionMode();
-			}
+			SetCompletionMode();
+		}
+		else if (PropertyName == EvalOptionsName)
+		{
+			// Cache user intended changes to the completion mode
+			CachedCompletionMode = EvalOptions.CompletionMode;
+
+			// Prevent the user from setting a looping sound to Keep State
+			SetCompletionMode();
 		}
 	}
 }
diff --git a/Engine/Source/Runtime/MovieSceneTracks/Public/Sections/MovieSceneAudioSection.h b/Engine/Source/Runtime/MovieSceneTracks/Public/Sections/MovieSceneAudioSection.h
--- a/Engine/Source/Runtime/MovieSceneTracks/Public/Sections/MovieSceneAudioSection.h
+++ b/Engine/Source/Runtime/MovieSceneTracks/Public/Sections/MovieSceneAudioSection.h
@@ -224,6 +224,11 @@ private:
 
 	MOVIESCENETRACKS_API void SetCompletionMode();
 
+	// Cache the users intended CompletionMode. 
+	// Used to revert the CompletionMode back to the users intent when not looping (when looping, completion mode is always RestoreState)
+	UPROPERTY()
+	EMovieSceneCompletionMode CachedCompletionMode = EMovieSceneCompletionMode::ProjectDefault;
+
 	/** The sound cue or wave that this section plays */
 	UPROPERTY(EditAnywhere, Category="Audio", BlueprintGetter=GetSound, BlueprintSetter=SetSound)
 	TObjectPtr<USoundBase> Sound;
