Sequencer: Rename parameter from InFrameRate to InTickReoslution since it's really expecting a Tick Resolution

#rnx

#Synced-CL 49850366

diff --git a/Engine/Plugins/Experimental/GeometryCollectionPlugin/Source/GeometryCollectionTracks/Private/MovieSceneGeometryCollectionSection.cpp b/Engine/Plugins/Experimental/GeometryCollectionPlugin/Source/GeometryCollectionTracks/Private/MovieSceneGeometryCollectionSection.cpp
--- a/Engine/Plugins/Experimental/GeometryCollectionPlugin/Source/GeometryCollectionTracks/Private/MovieSceneGeometryCollectionSection.cpp
+++ b/Engine/Plugins/Experimental/GeometryCollectionPlugin/Source/GeometryCollectionTracks/Private/MovieSceneGeometryCollectionSection.cpp
@@ -146,10 +146,10 @@ void UMovieSceneGeometryCollectionSection::GetSnapTimes(TArray<FFrameNumber>& Ou
 	}
 }
 
-float UMovieSceneGeometryCollectionSection::MapTimeToAnimation(FFrameTime InPosition, FFrameRate InFrameRate) const
+float UMovieSceneGeometryCollectionSection::MapTimeToAnimation(FFrameTime InPosition, FFrameRate InTickResolution) const
 {
 	FMovieSceneGeometryCollectionSectionTemplateParameters TemplateParams(Params, GetInclusiveStartFrame(), GetExclusiveEndFrame());
-	return TemplateParams.MapTimeToAnimation(InPosition, InFrameRate);
+	return TemplateParams.MapTimeToAnimation(InPosition, InTickResolution);
 }
 
 
diff --git a/Engine/Plugins/Experimental/GeometryCollectionPlugin/Source/GeometryCollectionTracks/Private/MovieSceneGeometryCollectionTemplate.cpp b/Engine/Plugins/Experimental/GeometryCollectionPlugin/Source/GeometryCollectionTracks/Private/MovieSceneGeometryCollectionTemplate.cpp
--- a/Engine/Plugins/Experimental/GeometryCollectionPlugin/Source/GeometryCollectionTracks/Private/MovieSceneGeometryCollectionTemplate.cpp
+++ b/Engine/Plugins/Experimental/GeometryCollectionPlugin/Source/GeometryCollectionTracks/Private/MovieSceneGeometryCollectionTemplate.cpp
@@ -122,23 +122,23 @@ void FMovieSceneGeometryCollectionSectionTemplate::Evaluate(const FMovieSceneEva
 	ExecutionTokens.Add(FGeometryCollectionExecutionToken(Params));
 }
 
-float FMovieSceneGeometryCollectionSectionTemplateParameters::MapTimeToAnimation(FFrameTime InPosition, FFrameRate InFrameRate) const
+float FMovieSceneGeometryCollectionSectionTemplateParameters::MapTimeToAnimation(FFrameTime InPosition, FFrameRate InTickResolution) const
 {
 	InPosition = FMath::Clamp(InPosition, FFrameTime(SectionStartTime), FFrameTime(SectionEndTime-1));
 
 	const float SectionPlayRate = PlayRate;
 	const float AnimPlayRate = FMath::IsNearlyZero(SectionPlayRate) ? 1.0f : SectionPlayRate;
 
-	const float SeqLength = GetSequenceLength() - InFrameRate.AsSeconds(StartFrameOffset + EndFrameOffset);
+	const float SeqLength = GetSequenceLength() - InTickResolution.AsSeconds(StartFrameOffset + EndFrameOffset);
 
-	float AnimPosition = FFrameTime::FromDecimal((InPosition - SectionStartTime).AsDecimal() * AnimPlayRate) / InFrameRate;
+	float AnimPosition = FFrameTime::FromDecimal((InPosition - SectionStartTime).AsDecimal() * AnimPlayRate) / InTickResolution;
 	if (SeqLength > 0.f)
 	{
 		// We don't support looping right now as it confuses the scrub with the current cache evaluation, just clamp so it holds on the
 		// last frame.
 		AnimPosition = FMath::Clamp(AnimPosition, AnimPosition, SeqLength);
 	}
-	AnimPosition += InFrameRate.AsSeconds(StartFrameOffset);
+	AnimPosition += InTickResolution.AsSeconds(StartFrameOffset);
 	// if (bReverse)
 	// {
 	// 	AnimPosition = (SeqLength - (AnimPosition - InFrameRate.AsSeconds(StartFrameOffset))) + InFrameRate.AsSeconds(StartFrameOffset);
diff --git a/Engine/Plugins/Experimental/GeometryCollectionPlugin/Source/GeometryCollectionTracks/Public/MovieSceneGeometryCollectionSection.h b/Engine/Plugins/Experimental/GeometryCollectionPlugin/Source/GeometryCollectionTracks/Public/MovieSceneGeometryCollectionSection.h
--- a/Engine/Plugins/Experimental/GeometryCollectionPlugin/Source/GeometryCollectionTracks/Public/MovieSceneGeometryCollectionSection.h
+++ b/Engine/Plugins/Experimental/GeometryCollectionPlugin/Source/GeometryCollectionTracks/Public/MovieSceneGeometryCollectionSection.h
@@ -79,7 +79,7 @@ public:
 	FMovieSceneGeometryCollectionParams Params;
 
 	/** Get Frame Time as Animation Time*/
-    virtual float MapTimeToAnimation(FFrameTime InPosition, FFrameRate InFrameRate) const;
+    virtual float MapTimeToAnimation(FFrameTime InPosition, FFrameRate InTickResolution) const;
 
 protected:
 	//~ UMovieSceneSection interface
diff --git a/Engine/Plugins/Experimental/GeometryCollectionPlugin/Source/GeometryCollectionTracks/Public/MovieSceneGeometryCollectionTemplate.h b/Engine/Plugins/Experimental/GeometryCollectionPlugin/Source/GeometryCollectionTracks/Public/MovieSceneGeometryCollectionTemplate.h
--- a/Engine/Plugins/Experimental/GeometryCollectionPlugin/Source/GeometryCollectionTracks/Public/MovieSceneGeometryCollectionTemplate.h
+++ b/Engine/Plugins/Experimental/GeometryCollectionPlugin/Source/GeometryCollectionTracks/Public/MovieSceneGeometryCollectionTemplate.h
@@ -20,7 +20,7 @@ struct FMovieSceneGeometryCollectionSectionTemplateParameters : public FMovieSce
 		, SectionEndTime(InSectionEndTime)
 	{}
 
-	float MapTimeToAnimation(FFrameTime InPosition, FFrameRate InFrameRate) const;
+	float MapTimeToAnimation(FFrameTime InPosition, FFrameRate InTickResolution) const;
 
 	UPROPERTY()
 	FFrameNumber SectionStartTime;
diff --git a/Engine/Plugins/FX/NiagaraSimCaching/Source/NiagaraSimCaching/Private/Niagara/Sequencer/MovieSceneNiagaraCacheTemplate.cpp b/Engine/Plugins/FX/NiagaraSimCaching/Source/NiagaraSimCaching/Private/Niagara/Sequencer/MovieSceneNiagaraCacheTemplate.cpp
--- a/Engine/Plugins/FX/NiagaraSimCaching/Source/NiagaraSimCaching/Private/Niagara/Sequencer/MovieSceneNiagaraCacheTemplate.cpp
+++ b/Engine/Plugins/FX/NiagaraSimCaching/Source/NiagaraSimCaching/Private/Niagara/Sequencer/MovieSceneNiagaraCacheTemplate.cpp
@@ -109,13 +109,13 @@ struct FNiagaraCacheExecutionToken : IMovieSceneExecutionToken
 		}
 	}
 
-	float MapTimeToAnimation(const FMovieSceneBaseCacheParams& BaseParams, float ComponentDuration, FFrameTime InPosition, FFrameRate InFrameRate, const FMovieSceneNiagaraSectionTemplateParameter& Section) const
+	float MapTimeToAnimation(const FMovieSceneBaseCacheParams& BaseParams, float ComponentDuration, FFrameTime InPosition, FFrameRate InTickResolution, const FMovieSceneNiagaraSectionTemplateParameter& Section) const
 	{
 		FFrameNumber SectionStartTime = Section.GetInclusiveStartFrame();
 		FFrameNumber SectionEndTime = Section.GetExclusiveEndFrame();
 		
-		const float SequenceLength = Section.Params.SectionStretchMode == ENiagaraSimCacheSectionStretchMode::TimeDilate ? static_cast<float>(InFrameRate.AsSeconds(SectionEndTime - SectionStartTime)) : ComponentDuration;
-		const FFrameTime AnimationLength = SequenceLength * InFrameRate;
+		const float SequenceLength = Section.Params.SectionStretchMode == ENiagaraSimCacheSectionStretchMode::TimeDilate ? static_cast<float>(InTickResolution.AsSeconds(SectionEndTime - SectionStartTime)) : ComponentDuration;
+		const FFrameTime AnimationLength = SequenceLength * InTickResolution;
 		const int32 LengthInFrames = AnimationLength.FrameNumber.Value + static_cast<int>(AnimationLength.GetSubFrame() + 0.5f) + 1;
 	
 		//we only play end if we are not looping, and assuming we are looping if Length is greater than default length;
@@ -125,15 +125,15 @@ struct FNiagaraCacheExecutionToken : IMovieSceneExecutionToken
 
 		const float SectionPlayRate = BaseParams.PlayRate;
 		const float AnimPlayRate = FMath::IsNearlyZero(SectionPlayRate) ? 1.0f : SectionPlayRate;
-		const float SeqLength = SequenceLength - static_cast<float>(InFrameRate.AsSeconds(BaseParams.StartFrameOffset + BaseParams.EndFrameOffset));
+		const float SeqLength = SequenceLength - static_cast<float>(InTickResolution.AsSeconds(BaseParams.StartFrameOffset + BaseParams.EndFrameOffset));
 
-		float AnimPosition = static_cast<float>(FFrameTime::FromDecimal((InPosition - SectionStartTime).AsDecimal() * AnimPlayRate) / InFrameRate);
-		AnimPosition += static_cast<float>(InFrameRate.AsSeconds(BaseParams.FirstLoopStartFrameOffset));
+		float AnimPosition = static_cast<float>(FFrameTime::FromDecimal((InPosition - SectionStartTime).AsDecimal() * AnimPlayRate) / InTickResolution);
+		AnimPosition += static_cast<float>(InTickResolution.AsSeconds(BaseParams.FirstLoopStartFrameOffset));
 		if (SeqLength > 0.f && (bLooping || !FMath::IsNearlyEqual(AnimPosition, SeqLength, 1e-4f)))
 		{
 			AnimPosition = FMath::Fmod(AnimPosition, SeqLength);
 		}
-		AnimPosition += static_cast<float>(InFrameRate.AsSeconds(BaseParams.StartFrameOffset));
+		AnimPosition += static_cast<float>(InTickResolution.AsSeconds(BaseParams.StartFrameOffset));
 		if (BaseParams.bReverse)
 		{
 			AnimPosition = SequenceLength - AnimPosition;
diff --git a/Engine/Plugins/Runtime/GeometryCache/Source/GeometryCacheTracks/Classes/MovieSceneGeometryCacheSection.h b/Engine/Plugins/Runtime/GeometryCache/Source/GeometryCacheTracks/Classes/MovieSceneGeometryCacheSection.h
--- a/Engine/Plugins/Runtime/GeometryCache/Source/GeometryCacheTracks/Classes/MovieSceneGeometryCacheSection.h
+++ b/Engine/Plugins/Runtime/GeometryCache/Source/GeometryCacheTracks/Classes/MovieSceneGeometryCacheSection.h
@@ -68,7 +68,7 @@ public:
 	FMovieSceneGeometryCacheParams Params;
 
 	/** Get Frame Time as Animation Time*/
-	virtual float MapTimeToAnimation(float ComponentDuration, FFrameTime InPosition, FFrameRate InFrameRate) const;
+	virtual float MapTimeToAnimation(float ComponentDuration, FFrameTime InPosition, FFrameRate InTickResolution) const;
 
 protected:
 	//~ UMovieSceneSection interface
diff --git a/Engine/Plugins/Runtime/GeometryCache/Source/GeometryCacheTracks/Private/MovieSceneGeometryCacheSection.cpp b/Engine/Plugins/Runtime/GeometryCache/Source/GeometryCacheTracks/Private/MovieSceneGeometryCacheSection.cpp
--- a/Engine/Plugins/Runtime/GeometryCache/Source/GeometryCacheTracks/Private/MovieSceneGeometryCacheSection.cpp
+++ b/Engine/Plugins/Runtime/GeometryCache/Source/GeometryCacheTracks/Private/MovieSceneGeometryCacheSection.cpp
@@ -198,10 +198,10 @@ void UMovieSceneGeometryCacheSection::GetSnapTimes(TArray<FFrameNumber>& OutSnap
 	}
 }
 
-float UMovieSceneGeometryCacheSection::MapTimeToAnimation(float ComponentDuration, FFrameTime InPosition, FFrameRate InFrameRate) const
+float UMovieSceneGeometryCacheSection::MapTimeToAnimation(float ComponentDuration, FFrameTime InPosition, FFrameRate InTickResolution) const
 {
 	FMovieSceneGeometryCacheSectionTemplateParameters TemplateParams(Params, GetInclusiveStartFrame(), GetExclusiveEndFrame());
-	return TemplateParams.MapTimeToAnimation(ComponentDuration, InPosition, InFrameRate);
+	return TemplateParams.MapTimeToAnimation(ComponentDuration, InPosition, InTickResolution);
 }
 
 
diff --git a/Engine/Plugins/Runtime/GeometryCache/Source/GeometryCacheTracks/Private/MovieSceneGeometryCacheTemplate.cpp b/Engine/Plugins/Runtime/GeometryCache/Source/GeometryCacheTracks/Private/MovieSceneGeometryCacheTemplate.cpp
--- a/Engine/Plugins/Runtime/GeometryCache/Source/GeometryCacheTracks/Private/MovieSceneGeometryCacheTemplate.cpp
+++ b/Engine/Plugins/Runtime/GeometryCache/Source/GeometryCacheTracks/Private/MovieSceneGeometryCacheTemplate.cpp
@@ -128,10 +128,10 @@ void FMovieSceneGeometryCacheSectionTemplate::Evaluate(const FMovieSceneEvaluati
 		ExecutionTokens.Add(FGeometryCacheExecutionToken(Params));
 }
 
-float FMovieSceneGeometryCacheSectionTemplateParameters::MapTimeToAnimation(float ComponentDuration, FFrameTime InPosition, FFrameRate InFrameRate) const
+float FMovieSceneGeometryCacheSectionTemplateParameters::MapTimeToAnimation(float ComponentDuration, FFrameTime InPosition, FFrameRate InTickResolution) const
 {
 	const float SequenceLength = ComponentDuration;
-	const FFrameTime AnimationLength = SequenceLength * InFrameRate;
+	const FFrameTime AnimationLength = SequenceLength * InTickResolution;
 	const int32 LengthInFrames = AnimationLength.FrameNumber.Value + (int)(AnimationLength.GetSubFrame() + 0.5f) + 1;
 	//we only play end if we are not looping, and assuming we are looping if Length is greater than default length;
 	const bool bLooping = (SectionEndTime.Value - SectionStartTime.Value + StartFrameOffset + EndFrameOffset) > LengthInFrames;
@@ -141,16 +141,16 @@ float FMovieSceneGeometryCacheSectionTemplateParameters::MapTimeToAnimation(floa
 	const float SectionPlayRate = PlayRate;
 	const float AnimPlayRate = FMath::IsNearlyZero(SectionPlayRate) ? 1.0f : SectionPlayRate;
 
-	const float FirstLoopSeqLength = static_cast<float>(SequenceLength - InFrameRate.AsSeconds(FirstLoopStartFrameOffset + StartFrameOffset + EndFrameOffset));
-	const float SeqLength = static_cast<float>(SequenceLength - InFrameRate.AsSeconds(StartFrameOffset + EndFrameOffset));
+	const float FirstLoopSeqLength = static_cast<float>(SequenceLength - InTickResolution.AsSeconds(FirstLoopStartFrameOffset + StartFrameOffset + EndFrameOffset));
+	const float SeqLength = static_cast<float>(SequenceLength - InTickResolution.AsSeconds(StartFrameOffset + EndFrameOffset));
 
-	float AnimPosition = static_cast<float>(FFrameTime::FromDecimal((InPosition - SectionStartTime).AsDecimal() * AnimPlayRate) / InFrameRate);
-	AnimPosition += static_cast<float>(InFrameRate.AsSeconds(FirstLoopStartFrameOffset));
+	float AnimPosition = static_cast<float>(FFrameTime::FromDecimal((InPosition - SectionStartTime).AsDecimal() * AnimPlayRate) / InTickResolution);
+	AnimPosition += static_cast<float>(InTickResolution.AsSeconds(FirstLoopStartFrameOffset));
 	if (SeqLength > 0.f && (bLooping || !FMath::IsNearlyEqual(AnimPosition, SeqLength, 1e-4f)))
 	{
 		AnimPosition = FMath::Fmod(AnimPosition, SeqLength);
 	}
-	AnimPosition += static_cast<float>(InFrameRate.AsSeconds(StartFrameOffset));
+	AnimPosition += static_cast<float>(InTickResolution.AsSeconds(StartFrameOffset));
 	if (bReverse)
 	{
 		AnimPosition = SequenceLength - AnimPosition;
diff --git a/Engine/Plugins/Runtime/GeometryCache/Source/GeometryCacheTracks/Private/MovieSceneGeometryCacheTemplate.h b/Engine/Plugins/Runtime/GeometryCache/Source/GeometryCacheTracks/Private/MovieSceneGeometryCacheTemplate.h
--- a/Engine/Plugins/Runtime/GeometryCache/Source/GeometryCacheTracks/Private/MovieSceneGeometryCacheTemplate.h
+++ b/Engine/Plugins/Runtime/GeometryCache/Source/GeometryCacheTracks/Private/MovieSceneGeometryCacheTemplate.h
@@ -18,7 +18,7 @@ struct FMovieSceneGeometryCacheSectionTemplateParameters : public FMovieSceneGeo
 		, SectionEndTime(InSectionEndTime)
 	{}
 
-	float MapTimeToAnimation(float ComponentDuration, FFrameTime InPosition, FFrameRate InFrameRate) const;
+	float MapTimeToAnimation(float ComponentDuration, FFrameTime InPosition, FFrameRate InTickResolution) const;
 
 	UPROPERTY()
 	FFrameNumber SectionStartTime;
diff --git a/Engine/Plugins/Runtime/HairStrands/Source/HairStrandsCore/Private/MovieSceneGroomCacheSection.cpp b/Engine/Plugins/Runtime/HairStrands/Source/HairStrandsCore/Private/MovieSceneGroomCacheSection.cpp
--- a/Engine/Plugins/Runtime/HairStrands/Source/HairStrandsCore/Private/MovieSceneGroomCacheSection.cpp
+++ b/Engine/Plugins/Runtime/HairStrands/Source/HairStrandsCore/Private/MovieSceneGroomCacheSection.cpp
@@ -118,10 +118,10 @@ void UMovieSceneGroomCacheSection::GetSnapTimes(TArray<FFrameNumber>& OutSnapTim
 	}
 }
 
-float UMovieSceneGroomCacheSection::MapTimeToAnimation(float ComponentDuration, FFrameTime InPosition, FFrameRate InFrameRate) const
+float UMovieSceneGroomCacheSection::MapTimeToAnimation(float ComponentDuration, FFrameTime InPosition, FFrameRate InTickResolution) const
 {
 	FMovieSceneGroomCacheSectionTemplateParameters TemplateParams(Params, GetInclusiveStartFrame(), GetExclusiveEndFrame());
-	return TemplateParams.MapTimeToAnimation(ComponentDuration, InPosition, InFrameRate);
+	return TemplateParams.MapTimeToAnimation(ComponentDuration, InPosition, InTickResolution);
 }
 
 
diff --git a/Engine/Plugins/Runtime/HairStrands/Source/HairStrandsCore/Private/MovieSceneGroomCacheTemplate.h b/Engine/Plugins/Runtime/HairStrands/Source/HairStrandsCore/Private/MovieSceneGroomCacheTemplate.h
--- a/Engine/Plugins/Runtime/HairStrands/Source/HairStrandsCore/Private/MovieSceneGroomCacheTemplate.h
+++ b/Engine/Plugins/Runtime/HairStrands/Source/HairStrandsCore/Private/MovieSceneGroomCacheTemplate.h
@@ -19,7 +19,7 @@ struct FMovieSceneGroomCacheSectionTemplateParameters : public FMovieSceneGroomC
 		, SectionEndTime(InSectionEndTime)
 	{}
 
-	float MapTimeToAnimation(float ComponentDuration, FFrameTime InPosition, FFrameRate InFrameRate) const;
+	float MapTimeToAnimation(float ComponentDuration, FFrameTime InPosition, FFrameRate InTickResolution) const;
 
 	UPROPERTY()
 	FFrameNumber SectionStartTime;
diff --git a/Engine/Plugins/Runtime/HairStrands/Source/HairStrandsCore/Public/MovieSceneGroomCacheSection.h b/Engine/Plugins/Runtime/HairStrands/Source/HairStrandsCore/Public/MovieSceneGroomCacheSection.h
--- a/Engine/Plugins/Runtime/HairStrands/Source/HairStrandsCore/Public/MovieSceneGroomCacheSection.h
+++ b/Engine/Plugins/Runtime/HairStrands/Source/HairStrandsCore/Public/MovieSceneGroomCacheSection.h
@@ -56,7 +56,7 @@ public:
 	FMovieSceneGroomCacheParams Params;
 
 	/** Get Frame Time as Animation Time*/
-	virtual float MapTimeToAnimation(float ComponentDuration, FFrameTime InPosition, FFrameRate InFrameRate) const;
+	virtual float MapTimeToAnimation(float ComponentDuration, FFrameTime InPosition, FFrameRate InTickResolution) const;
 
 protected:
 	//~ UMovieSceneSection interface
diff --git a/Engine/Source/Runtime/MovieSceneTracks/Private/Evaluation/MovieSceneBaseCacheTemplate.cpp b/Engine/Source/Runtime/MovieSceneTracks/Private/Evaluation/MovieSceneBaseCacheTemplate.cpp
--- a/Engine/Source/Runtime/MovieSceneTracks/Private/Evaluation/MovieSceneBaseCacheTemplate.cpp
+++ b/Engine/Source/Runtime/MovieSceneTracks/Private/Evaluation/MovieSceneBaseCacheTemplate.cpp
@@ -7,10 +7,10 @@
 DECLARE_CYCLE_STAT(TEXT("Chaos Cache Evaluate"), MovieSceneEval_BaseCache_Evaluate, STATGROUP_MovieSceneEval);
 DECLARE_CYCLE_STAT(TEXT("Chaos Cache Token Execute"), MovieSceneEval_BaseCache_TokenExecute, STATGROUP_MovieSceneEval);
 
-float FMovieSceneBaseCacheSectionTemplateParameters::MapTimeToAnimation(const FMovieSceneBaseCacheParams& BaseParams, float ComponentDuration, FFrameTime InPosition, FFrameRate InFrameRate) const
+float FMovieSceneBaseCacheSectionTemplateParameters::MapTimeToAnimation(const FMovieSceneBaseCacheParams& BaseParams, float ComponentDuration, FFrameTime InPosition, FFrameRate InTickResolution) const
 {
 	const float SequenceLength = ComponentDuration;
-	const FFrameTime AnimationLength = SequenceLength * InFrameRate;
+	const FFrameTime AnimationLength = SequenceLength * InTickResolution;
 	const int32 LengthInFrames = AnimationLength.FrameNumber.Value + (int)(AnimationLength.GetSubFrame() + 0.5f) + 1;
 	
 	//we only play end if we are not looping, and assuming we are looping if Length is greater than default length;
@@ -21,16 +21,16 @@ float FMovieSceneBaseCacheSectionTemplateParameters::MapTimeToAnimation(const FM
 	const float SectionPlayRate = BaseParams.PlayRate;
 	const float AnimPlayRate = FMath::IsNearlyZero(SectionPlayRate) ? 1.0f : SectionPlayRate;
 
-	const float FirstLoopSeqLength = SequenceLength - InFrameRate.AsSeconds(BaseParams.FirstLoopStartFrameOffset + BaseParams.StartFrameOffset + BaseParams.EndFrameOffset);
-	const float SeqLength = SequenceLength - InFrameRate.AsSeconds(BaseParams.StartFrameOffset + BaseParams.EndFrameOffset);
+	const float FirstLoopSeqLength = SequenceLength - InTickResolution.AsSeconds(BaseParams.FirstLoopStartFrameOffset + BaseParams.StartFrameOffset + BaseParams.EndFrameOffset);
+	const float SeqLength = SequenceLength - InTickResolution.AsSeconds(BaseParams.StartFrameOffset + BaseParams.EndFrameOffset);
 
-	float AnimPosition = FFrameTime::FromDecimal((InPosition - SectionStartTime).AsDecimal() * AnimPlayRate) / InFrameRate;
-	AnimPosition += InFrameRate.AsSeconds(BaseParams.FirstLoopStartFrameOffset);
+	float AnimPosition = FFrameTime::FromDecimal((InPosition - SectionStartTime).AsDecimal() * AnimPlayRate) / InTickResolution;
+	AnimPosition += InTickResolution.AsSeconds(BaseParams.FirstLoopStartFrameOffset);
 	if (SeqLength > 0.f && (bLooping || !FMath::IsNearlyEqual(AnimPosition, SeqLength, 1e-4f)))
 	{
 		AnimPosition = FMath::Fmod(AnimPosition, SeqLength);
 	}
-	AnimPosition += InFrameRate.AsSeconds(BaseParams.StartFrameOffset);
+	AnimPosition += InTickResolution.AsSeconds(BaseParams.StartFrameOffset);
 	if (BaseParams.bReverse)
 	{
 		AnimPosition = SequenceLength - AnimPosition;
diff --git a/Engine/Source/Runtime/MovieSceneTracks/Private/Sections/MovieSceneBaseCacheSection.cpp b/Engine/Source/Runtime/MovieSceneTracks/Private/Sections/MovieSceneBaseCacheSection.cpp
--- a/Engine/Source/Runtime/MovieSceneTracks/Private/Sections/MovieSceneBaseCacheSection.cpp
+++ b/Engine/Source/Runtime/MovieSceneTracks/Private/Sections/MovieSceneBaseCacheSection.cpp
@@ -122,12 +122,12 @@ void UMovieSceneBaseCacheSection::GetSnapTimes(TArray<FFrameNumber>& OutSnapTime
 	}
 }
 
-float UMovieSceneBaseCacheSection::MapTimeToAnimation(float ComponentDuration, FFrameTime InPosition, FFrameRate InFrameRate) const
+float UMovieSceneBaseCacheSection::MapTimeToAnimation(float ComponentDuration, FFrameTime InPosition, FFrameRate InTickResolution) const
 {
 	if(ParamsPtr)
 	{
 		const FMovieSceneBaseCacheSectionTemplateParameters TemplateParamsPtr(GetInclusiveStartFrame(), GetExclusiveEndFrame());
-		return TemplateParamsPtr.MapTimeToAnimation(*ParamsPtr, ComponentDuration, InPosition, InFrameRate);
+		return TemplateParamsPtr.MapTimeToAnimation(*ParamsPtr, ComponentDuration, InPosition, InTickResolution);
 	}
 	return 0.0f;
 }
diff --git a/Engine/Source/Runtime/MovieSceneTracks/Private/Sections/MovieSceneSkeletalAnimationSection.cpp b/Engine/Source/Runtime/MovieSceneTracks/Private/Sections/MovieSceneSkeletalAnimationSection.cpp
--- a/Engine/Source/Runtime/MovieSceneTracks/Private/Sections/MovieSceneSkeletalAnimationSection.cpp
+++ b/Engine/Source/Runtime/MovieSceneTracks/Private/Sections/MovieSceneSkeletalAnimationSection.cpp
@@ -85,16 +85,16 @@ FMovieSceneSequenceTransform FMovieSceneSkeletalAnimationParams::MakeTransform(c
 	return TimingParams.MakeTransform(OuterFrameRate, OuterRange, SequenceLength, OverrideSequence->RateScale);
 }
 
-double FMovieSceneSkeletalAnimationParams::MapTimeToAnimation(const UMovieSceneSection* InSection, FFrameTime InPosition, FFrameRate InFrameRate, UAnimSequenceBase* OverrideSequence) const
+double FMovieSceneSkeletalAnimationParams::MapTimeToAnimation(const UMovieSceneSection* InSection, FFrameTime InPosition, FFrameRate InTickResolution, UAnimSequenceBase* OverrideSequence) const
 {
 	const FFrameNumber SectionStartTime = InSection->GetInclusiveStartFrame();
 	const FFrameNumber SectionEndTime = InSection->GetExclusiveEndFrame();
-	return MapTimeToAnimation(SectionStartTime, SectionEndTime, InPosition, InFrameRate, OverrideSequence);
+	return MapTimeToAnimation(SectionStartTime, SectionEndTime, InPosition, InTickResolution, OverrideSequence);
 }
 
-double FMovieSceneSkeletalAnimationParams::MapTimeToAnimation(FFrameNumber InSectionStartTime, FFrameNumber InSectionEndTime, FFrameTime InPosition, FFrameRate InFrameRate, UAnimSequenceBase* OverrideSequence) const
+double FMovieSceneSkeletalAnimationParams::MapTimeToAnimation(FFrameNumber InSectionStartTime, FFrameNumber InSectionEndTime, FFrameTime InPosition, FFrameRate InTickResolution, UAnimSequenceBase* OverrideSequence) const
 {
-	return MakeTransform(InFrameRate, TRange<FFrameNumber>(InSectionStartTime, InSectionEndTime), OverrideSequence).TransformTime(InPosition).AsDecimal();
+	return MakeTransform(InTickResolution, TRange<FFrameNumber>(InSectionStartTime, InSectionEndTime), OverrideSequence).TransformTime(InPosition).AsDecimal();
 }
 
 UMovieSceneSkeletalAnimationSection::UMovieSceneSkeletalAnimationSection( const FObjectInitializer& ObjectInitializer )
@@ -473,9 +473,9 @@ void UMovieSceneSkeletalAnimationSection::GetSnapTimes(TArray<FFrameNumber>& Out
 	}
 }
 
-double UMovieSceneSkeletalAnimationSection::MapTimeToAnimation(FFrameTime InPosition, FFrameRate InFrameRate) const
+double UMovieSceneSkeletalAnimationSection::MapTimeToAnimation(FFrameTime InPosition, FFrameRate InTickResolution) const
 {
-	return Params.MapTimeToAnimation(this, InPosition, InFrameRate);
+	return Params.MapTimeToAnimation(this, InPosition, InTickResolution);
 }
 
 float UMovieSceneSkeletalAnimationSection::GetTotalWeightValue(FFrameTime InTime) const
@@ -641,7 +641,7 @@ void UMovieSceneSkeletalAnimationSection::GetRootMotion(FFrameTime CurrentTime,
 	}
 }
 
-bool UMovieSceneSkeletalAnimationSection::GetRootMotionVelocity(FFrameTime PreviousTime, FFrameTime CurrentTime, FFrameRate FrameRate, 
+bool UMovieSceneSkeletalAnimationSection::GetRootMotionVelocity(FFrameTime PreviousTime, FFrameTime CurrentTime, FFrameRate InTickResolution,
 	FTransform& OutVelocity, float& OutWeight) const
 {
 	UAnimSequence* AnimSequence = Cast<UAnimSequence>(Params.Animation);
@@ -652,8 +652,8 @@ bool UMovieSceneSkeletalAnimationSection::GetRootMotionVelocity(FFrameTime Previ
 		OutWeight = ManualWeight * EvaluateEasing(CurrentTime);
 		//mz todo we should be able to cache the PreviousTimeSeconds;
 		//mz todo need to get the starting value.
-		float PreviousTimeSeconds = static_cast<float>(MapTimeToAnimation(PreviousTime, FrameRate));
-		float CurrentTimeSeconds  = static_cast<float>(MapTimeToAnimation(CurrentTime, FrameRate));
+		float PreviousTimeSeconds = static_cast<float>(MapTimeToAnimation(PreviousTime, InTickResolution));
+		float CurrentTimeSeconds  = static_cast<float>(MapTimeToAnimation(CurrentTime, InTickResolution));
 		OutVelocity = AnimSequence->ExtractRootMotionFromRange(PreviousTimeSeconds, CurrentTimeSeconds, FAnimExtractContext());
 		return true;
 	}
@@ -889,7 +889,7 @@ void UMovieSceneSkeletalAnimationSection::ClearMatchedOffsetTransforms()
 }
 
 
-void UMovieSceneSkeletalAnimationSection::MatchSectionByBoneTransform(USkeletalMeshComponent* SkelMeshComp, FFrameTime CurrentFrame, FFrameRate FrameRate,
+void UMovieSceneSkeletalAnimationSection::MatchSectionByBoneTransform(USkeletalMeshComponent* SkelMeshComp, FFrameTime CurrentFrame, FFrameRate InTickResolution,
 	const FName& BoneName)
 {
 	MatchedBoneName = BoneName;
@@ -903,7 +903,7 @@ void UMovieSceneSkeletalAnimationSection::MatchSectionByBoneTransform(USkeletalM
 		FVector PreviousMatchedLocationOffset = MatchedLocationOffset;
 		FRotator PreviousMatchedRotationOffset = MatchedRotationOffset;
 
-		Track->MatchSectionByBoneTransform(bMatchWithPrevious,SkelMeshComp, this, CurrentFrame, FrameRate, BoneName, DiffTransform,DiffTranslate,DiffRotate);
+		Track->MatchSectionByBoneTransform(bMatchWithPrevious,SkelMeshComp, this, CurrentFrame, InTickResolution, BoneName, DiffTransform,DiffTranslate,DiffRotate);
 
 		MatchedLocationOffset = bMatchTranslation ? DiffTranslate : FVector::ZeroVector;
 		MatchedRotationOffset = DiffRotate.Rotator();
diff --git a/Engine/Source/Runtime/MovieSceneTracks/Private/Tracks/MovieSceneCommonAnimationTrack.cpp b/Engine/Source/Runtime/MovieSceneTracks/Private/Tracks/MovieSceneCommonAnimationTrack.cpp
--- a/Engine/Source/Runtime/MovieSceneTracks/Private/Tracks/MovieSceneCommonAnimationTrack.cpp
+++ b/Engine/Source/Runtime/MovieSceneTracks/Private/Tracks/MovieSceneCommonAnimationTrack.cpp
@@ -947,7 +947,7 @@ static ESkelAnimRotationOrder FindBestRotationOrder(bool bMatchRoll, bool bMatch
 	return ESkelAnimRotationOrder::XYZ;
 }
 
-void UMovieSceneCommonAnimationTrack::MatchSectionByBoneTransform(bool bMatchWithPrevious, USkeletalMeshComponent* SkelMeshComp, UMovieSceneSkeletalAnimationSection* CurrentSection, FFrameTime CurrentFrame, FFrameRate FrameRate,
+void UMovieSceneCommonAnimationTrack::MatchSectionByBoneTransform(bool bMatchWithPrevious, USkeletalMeshComponent* SkelMeshComp, UMovieSceneSkeletalAnimationSection* CurrentSection, FFrameTime CurrentFrame, FFrameRate InTickResolution,
 	const FName& BoneName, FTransform& SecondSectionRootDiff, FVector& TranslationDiff, FQuat& RotationDiff) //add options for z and for rotation.
 {
 	SortSections();
@@ -979,13 +979,13 @@ void UMovieSceneCommonAnimationTrack::MatchSectionByBoneTransform(bool bMatchWit
 
 		if (FirstAnimSequence && SecondAnimSequence)
 		{
-			double FirstSectionTime = FirstSection->MapTimeToAnimation(CurrentFrame, FrameRate);
+			double FirstSectionTime = FirstSection->MapTimeToAnimation(CurrentFrame, InTickResolution);
 			//use same index for all
 			int32 Index = CurrentSection->SetBoneIndexForRootMotionCalculations(bBlendFirstChildOfRoot);
 			FCompactPoseBoneIndex ParentIndex(Index);
 			
 			FTransform  FirstTransform = GetTransformForBoneRelativeToIndex(FirstAnimSequence, SkelMeshComp, BoneName, ParentIndex, FirstSectionTime);
-			double SecondSectionTime = CurrentSection->MapTimeToAnimation(CurrentFrame, FrameRate);
+			double SecondSectionTime = CurrentSection->MapTimeToAnimation(CurrentFrame, InTickResolution);
 			FTransform  SecondTransform = GetTransformForBoneRelativeToIndex(SecondAnimSequence, SkelMeshComp, BoneName, ParentIndex,SecondSectionTime);
 			//Need to match the translations and rotations here 
 			//First need to get the correct rotation order based upon what's matching, otherwise if not all are matched 
@@ -1049,9 +1049,9 @@ void UMovieSceneCommonAnimationTrack::MatchSectionByBoneTransform(bool bMatchWit
 			//use same index for all
 			int32 Index = CurrentSection->SetBoneIndexForRootMotionCalculations(bBlendFirstChildOfRoot);
 			FCompactPoseBoneIndex ParentIndex(Index);
-			float FirstSectionTime = static_cast<float>(CurrentSection->MapTimeToAnimation(CurrentFrame, FrameRate));
+			float FirstSectionTime = static_cast<float>(CurrentSection->MapTimeToAnimation(CurrentFrame, InTickResolution));
 			FTransform  FirstTransform = GetTransformForBoneRelativeToIndex(FirstAnimSequence, SkelMeshComp, BoneName,ParentIndex, FirstSectionTime);
-			float SecondSectionTime = static_cast<float>(SecondSection->MapTimeToAnimation(CurrentFrame, FrameRate));
+			float SecondSectionTime = static_cast<float>(SecondSection->MapTimeToAnimation(CurrentFrame, InTickResolution));
 			FTransform  SecondTransform = GetTransformForBoneRelativeToIndex(SecondAnimSequence, SkelMeshComp, BoneName,ParentIndex, SecondSectionTime);
 
 			//Need to match the translations and rotations here 
diff --git a/Engine/Source/Runtime/MovieSceneTracks/Public/Evaluation/MovieSceneBaseCacheTemplate.h b/Engine/Source/Runtime/MovieSceneTracks/Public/Evaluation/MovieSceneBaseCacheTemplate.h
--- a/Engine/Source/Runtime/MovieSceneTracks/Public/Evaluation/MovieSceneBaseCacheTemplate.h
+++ b/Engine/Source/Runtime/MovieSceneTracks/Public/Evaluation/MovieSceneBaseCacheTemplate.h
@@ -1,4 +1,4 @@
-﻿// Copyright Epic Games, Inc. All Rights Reserved.
+// Copyright Epic Games, Inc. All Rights Reserved.
 
 #pragma once
 
@@ -18,7 +18,7 @@ struct FMovieSceneBaseCacheSectionTemplateParameters
 	{}
 
 	/** Get Frame Time as Animation Time */
-	MOVIESCENETRACKS_API float MapTimeToAnimation(const FMovieSceneBaseCacheParams& BaseParams, float ComponentDuration, FFrameTime InPosition, FFrameRate InFrameRate) const;
+	MOVIESCENETRACKS_API float MapTimeToAnimation(const FMovieSceneBaseCacheParams& BaseParams, float ComponentDuration, FFrameTime InPosition, FFrameRate InTickResolution) const;
 
 	/** Section start time */
 	UPROPERTY()
diff --git a/Engine/Source/Runtime/MovieSceneTracks/Public/Sections/MovieSceneBaseCacheSection.h b/Engine/Source/Runtime/MovieSceneTracks/Public/Sections/MovieSceneBaseCacheSection.h
--- a/Engine/Source/Runtime/MovieSceneTracks/Public/Sections/MovieSceneBaseCacheSection.h
+++ b/Engine/Source/Runtime/MovieSceneTracks/Public/Sections/MovieSceneBaseCacheSection.h
@@ -55,7 +55,7 @@ public:
 	FMovieSceneBaseCacheParams* ParamsPtr = nullptr;
 	
 	/** Get Frame Time as Animation Time*/
-	MOVIESCENETRACKS_API virtual float MapTimeToAnimation(float ComponentDuration, FFrameTime InPosition, FFrameRate InFrameRate) const;
+	MOVIESCENETRACKS_API virtual float MapTimeToAnimation(float ComponentDuration, FFrameTime InPosition, FFrameRate InTickResolution) const;
 	MOVIESCENETRACKS_API virtual TOptional<TRange<FFrameNumber> > GetAutoSizeRange() const override;
 	
 protected:
diff --git a/Engine/Source/Runtime/MovieSceneTracks/Public/Sections/MovieSceneSkeletalAnimationSection.h b/Engine/Source/Runtime/MovieSceneTracks/Public/Sections/MovieSceneSkeletalAnimationSection.h
--- a/Engine/Source/Runtime/MovieSceneTracks/Public/Sections/MovieSceneSkeletalAnimationSection.h
+++ b/Engine/Source/Runtime/MovieSceneTracks/Public/Sections/MovieSceneSkeletalAnimationSection.h
@@ -34,12 +34,12 @@ struct FMovieSceneSkeletalAnimationParams
 	 * Convert a sequence frame to a time in seconds inside the animation clip, taking into account start/end offsets,
 	 * looping, etc.
 	 */
-	double MOVIESCENETRACKS_API MapTimeToAnimation(const UMovieSceneSection* InSection, FFrameTime InPosition, FFrameRate InFrameRate, UAnimSequenceBase* OverrideSequence = nullptr) const;
+	double MOVIESCENETRACKS_API MapTimeToAnimation(const UMovieSceneSection* InSection, FFrameTime InPosition, FFrameRate InTickResolution, UAnimSequenceBase* OverrideSequence = nullptr) const;
 
 	/**
 	 * As above, but with already computed section bounds.
 	 */
-	double MOVIESCENETRACKS_API MapTimeToAnimation(FFrameNumber InSectionStartTime, FFrameNumber InSectionEndTime, FFrameTime InPosition, FFrameRate InFrameRate, UAnimSequenceBase* OverrideSequence = nullptr) const;
+	double MOVIESCENETRACKS_API MapTimeToAnimation(FFrameNumber InSectionStartTime, FFrameNumber InSectionEndTime, FFrameTime InPosition, FFrameRate InTickResolution, UAnimSequenceBase* OverrideSequence = nullptr) const;
 
 	/**
 	 * Make a transform structure from these animation parameters
@@ -120,7 +120,7 @@ public:
 	FMovieSceneSkeletalAnimationParams Params;
 
 	/** Get Frame Time as Animation Time*/
-	MOVIESCENETRACKS_API double MapTimeToAnimation(FFrameTime InPosition, FFrameRate InFrameRate) const;
+	MOVIESCENETRACKS_API double MapTimeToAnimation(FFrameTime InPosition, FFrameRate InTickResolution) const;
 
 	MOVIESCENETRACKS_API UAnimSequenceBase* GetAnimation() const;
 	MOVIESCENETRACKS_API UAnimSequenceBase* GetPlaybackAnimation() const;
@@ -270,10 +270,10 @@ public:
 	// Functions/params related to root motion calcuations.
 	MOVIESCENETRACKS_API FMovieSceneSkeletalAnimRootMotionTrackParams* GetRootMotionParams() const;
 
-	MOVIESCENETRACKS_API bool GetRootMotionVelocity(FFrameTime PreviousTime, FFrameTime CurrentTime, FFrameRate FrameRate, FTransform& OutVelocity, float& OutWeight) const;
+	MOVIESCENETRACKS_API bool GetRootMotionVelocity(FFrameTime PreviousTime, FFrameTime CurrentTime, FFrameRate InTickResolution, FTransform& OutVelocity, float& OutWeight) const;
 	MOVIESCENETRACKS_API int32 SetBoneIndexForRootMotionCalculations(bool bBlendFirstChildOfRoot);
 	MOVIESCENETRACKS_API bool GetRootMotionTransform(FAnimationPoseData& PoseData, FRootMotionTransformParam& InOutParams) const;
-	MOVIESCENETRACKS_API void MatchSectionByBoneTransform(USkeletalMeshComponent* SkelMeshComp, FFrameTime CurrentFrame, FFrameRate FrameRate,
+	MOVIESCENETRACKS_API void MatchSectionByBoneTransform(USkeletalMeshComponent* SkelMeshComp, FFrameTime CurrentFrame, FFrameRate InTickResolution,
 		const FName& BoneName); //add options for z and rotation
 
 	MOVIESCENETRACKS_API void ClearMatchedOffsetTransforms();
diff --git a/Engine/Source/Runtime/MovieSceneTracks/Public/Tracks/MovieSceneCommonAnimationTrack.h b/Engine/Source/Runtime/MovieSceneTracks/Public/Tracks/MovieSceneCommonAnimationTrack.h
--- a/Engine/Source/Runtime/MovieSceneTracks/Public/Tracks/MovieSceneCommonAnimationTrack.h
+++ b/Engine/Source/Runtime/MovieSceneTracks/Public/Tracks/MovieSceneCommonAnimationTrack.h
@@ -97,7 +97,7 @@ public:
 	UE_API void SetRootMotionsDirty();
 	UE_API void SetUpRootMotions(bool bForce);
 	UE_API TOptional<FTransform>  GetRootMotion(FFrameTime CurrentTime);
-	UE_API void MatchSectionByBoneTransform(bool bMatchWithPrevious, USkeletalMeshComponent* SkelMeshComp, UMovieSceneSkeletalAnimationSection* CurrentSection, FFrameTime CurrentFrame, FFrameRate FrameRate,
+	UE_API void MatchSectionByBoneTransform(bool bMatchWithPrevious, USkeletalMeshComponent* SkelMeshComp, UMovieSceneSkeletalAnimationSection* CurrentSection, FFrameTime CurrentFrame, FFrameRate InTickResolution,
 		const FName& BoneName, FTransform& SecondSectionRootDiff, FVector& TranslationDiff, FQuat& RotationDiff); //add options for z and for rotation.
 #if WITH_EDITORONLY_DATA
 	UE_API void ToggleShowRootMotionTrail();
