Sequencer: fix pre-animated state problems with the camera cut track specifically, and in general too

When a contributor gets removed for a given animated "thing", we were previously restoring pre-animated state if no RestoreState contributors were left. This was incorrect, since KeepState contributors were, by definition, still animating that thing. This could create problem for instance when animating transforms. A KeepState section following a RestoreState section would still be moving the object that the previous section was moving, but when that previous section ended, pre-animated Component Mobility would be restored to Static. This prevented the KeepState section to continue moving the object. Now, in RemoveMetaData(), we *discard* the transient pre-animated state if no KeepState contributors are left, but KeepState contributors still exist. That transient state gets either thrown away (if we don't care about global state capture) or promoted to persistent (if we do).

A related problem was that several places (including in the Component Mobility system) were not tracking entities that are KeepState until global capture is enabled. This meant that when global capture was disabled (i.e. in the game, not the Sequencer editor) these KeepState sections following a RestoreState section would not be tracked. So inside RemoveMetaData(), we would get a total contributor count of zero even though there was still one (KeepState) section animating things. These early-outs have been removed.

This change also adds an ensure in RestorePreAnimatedStateStorage() to warn us if pre-animated state gets promoted to persistent storage for no reason.

The camera cut track instance also gets a few fixes. Before, it was saving pre-animated state all the time, letting the core system figure out if it needed to do something about it (on the first call) or not (on subsequent calls). Now, however, we want to only call it when a new camera cut starts. This is becasue beginning tracking an input cannot be done the usual way: we want it to happen only after we have resolved the bound camera actor and verified it's valid.

#virtualized
#Synced-CL 49781371

diff --git a/Engine/Source/Runtime/MovieScene/Private/Evaluation/PreAnimatedState/MovieScenePreAnimatedStateExtension.cpp b/Engine/Source/Runtime/MovieScene/Private/Evaluation/PreAnimatedState/MovieScenePreAnimatedStateExtension.cpp
--- a/Engine/Source/Runtime/MovieScene/Private/Evaluation/PreAnimatedState/MovieScenePreAnimatedStateExtension.cpp
+++ b/Engine/Source/Runtime/MovieScene/Private/Evaluation/PreAnimatedState/MovieScenePreAnimatedStateExtension.cpp
@@ -175,14 +175,28 @@ void FPreAnimatedStateExtension::RemoveMetaData(const FPreAnimatedStateMetaData&
 	{
 		if (--Aggregate->NumRestoreContributors == 0)
 		{
-			EPreAnimatedStorageRequirement NewRequirement = TotalNum != 0
-				? EPreAnimatedStorageRequirement::Persistent
-				: EPreAnimatedStorageRequirement::None;
-
-			TSharedPtr<IPreAnimatedStorage> Storage = GetStorageChecked(MetaData.Entry.ValueHandle.TypeID);
+			EPreAnimatedStorageRequirement NewRequirement = EPreAnimatedStorageRequirement::None;
 
 			FRestoreStateParams Params = {Linker, MetaData.RootInstanceHandle};
-			NewRequirement = Storage->RestorePreAnimatedStateStorage(MetaData.Entry.ValueHandle.StorageIndex, EPreAnimatedStorageRequirement::Transient, NewRequirement, Params);
+			FPreAnimatedStorageIndex StorageIndex = MetaData.Entry.ValueHandle.StorageIndex;
+			TSharedPtr<IPreAnimatedStorage> Storage = GetStorageChecked(MetaData.Entry.ValueHandle.TypeID);
+			if (TotalNum == 0)
+			{
+				// We are completely done with the value (there are no more contributors). Restore the pre-animated 
+				// value and discard it.
+				EPreAnimatedStorageRequirement SourceRequirement = EPreAnimatedStorageRequirement::Transient;
+				EPreAnimatedStorageRequirement TargetRequirement = EPreAnimatedStorageRequirement::None;
+				NewRequirement = Storage->RestorePreAnimatedStateStorage(StorageIndex, SourceRequirement, TargetRequirement, Params);
+			}
+			else
+			{
+				// We still have contributors animating the value, but they're all set to KeepState.
+				// At this point we can discard the pre-animated value, we don't need it anymore, unless we want to
+				// later restore global state, in which case we keep that value in persistent storage.
+				EPreAnimatedStorageRequirement SourceRequirement = IsCapturingGlobalState() ?
+					EPreAnimatedStorageRequirement::Transient : EPreAnimatedStorageRequirement::Persistent;
+				NewRequirement = Storage->DiscardPreAnimatedStateStorage(StorageIndex, SourceRequirement);
+			}
 
 			if (NewRequirement == EPreAnimatedStorageRequirement::None)
 			{
diff --git a/Engine/Source/Runtime/MovieScene/Public/Evaluation/PreAnimatedState/MovieScenePreAnimatedObjectStorage.h b/Engine/Source/Runtime/MovieScene/Public/Evaluation/PreAnimatedState/MovieScenePreAnimatedObjectStorage.h
--- a/Engine/Source/Runtime/MovieScene/Public/Evaluation/PreAnimatedState/MovieScenePreAnimatedObjectStorage.h
+++ b/Engine/Source/Runtime/MovieScene/Public/Evaluation/PreAnimatedState/MovieScenePreAnimatedObjectStorage.h
@@ -118,11 +118,6 @@ public:
 		const int32 Num = Params.Num;
 		const bool  bWantsRestore = Params.bWantsRestoreState;
 
-		if (!this->ParentExtension->IsCapturingGlobalState() && !bWantsRestore)
-		{
-			return;
-		}
-
 		FPreAnimatedEntityCaptureSource* EntityMetaData = this->ParentExtension->GetOrCreateEntityMetaData();
 		for (int32 Index = 0; Index < Num; ++Index)
 		{
@@ -148,11 +143,6 @@ public:
 
 	void BeginTrackingEntity(FMovieSceneEntityID EntityID, bool bWantsRestoreState, FRootInstanceHandle RootInstanceHandle, UObject* BoundObject) override
 	{
-		if (!this->ParentExtension->IsCapturingGlobalState() && !bWantsRestoreState)
-		{
-			return;
-		}
-
 		FPreAnimatedEntityCaptureSource* EntityMetaData = this->ParentExtension->GetOrCreateEntityMetaData();
 
 		FPreAnimatedStateEntry Entry = MakeEntry(BoundObject);
diff --git a/Engine/Source/Runtime/MovieScene/Public/Evaluation/PreAnimatedState/MovieScenePreAnimatedPropertyStorage.h b/Engine/Source/Runtime/MovieScene/Public/Evaluation/PreAnimatedState/MovieScenePreAnimatedPropertyStorage.h
--- a/Engine/Source/Runtime/MovieScene/Public/Evaluation/PreAnimatedState/MovieScenePreAnimatedPropertyStorage.h
+++ b/Engine/Source/Runtime/MovieScene/Public/Evaluation/PreAnimatedState/MovieScenePreAnimatedPropertyStorage.h
@@ -172,11 +172,6 @@ struct TPreAnimatedPropertyStorageImpl<PropertyTraits, TPropertyMetaData<MetaDat
 		const int32 Num = Params.Num;
 		const bool  bWantsRestore = Params.bWantsRestoreState;
 
-		if (!this->ParentExtension->IsCapturingGlobalState() && !bWantsRestore)
-		{
-			return;
-		}
-
 		FPreAnimatedEntityCaptureSource* EntityMetaData = this->ParentExtension->GetOrCreateEntityMetaData();
 		for (int32 Index = 0; Index < Num; ++Index)
 		{
diff --git a/Engine/Source/Runtime/MovieScene/Public/Evaluation/PreAnimatedState/MovieScenePreAnimatedStateStorage.h b/Engine/Source/Runtime/MovieScene/Public/Evaluation/PreAnimatedState/MovieScenePreAnimatedStateStorage.h
--- a/Engine/Source/Runtime/MovieScene/Public/Evaluation/PreAnimatedState/MovieScenePreAnimatedStateStorage.h
+++ b/Engine/Source/Runtime/MovieScene/Public/Evaluation/PreAnimatedState/MovieScenePreAnimatedStateStorage.h
@@ -279,6 +279,10 @@ public:
 
 		if (TargetRequirement == EPreAnimatedStorageRequirement::Persistent)
 		{
+			ensureMsgf(
+					this->ParentExtension->IsCapturingGlobalState(),
+					TEXT("We shouldn't move cached values to the persistent storage if we are not capturing global state!"));
+
 			// Restore the value but keep the value cached
 			FCachedData& PersistentData = PreAnimatedStorage[StorageIndex];
 
diff --git a/Engine/Source/Runtime/MovieSceneTracks/Private/Systems/MovieSceneComponentMobilitySystem.cpp b/Engine/Source/Runtime/MovieSceneTracks/Private/Systems/MovieSceneComponentMobilitySystem.cpp
--- a/Engine/Source/Runtime/MovieSceneTracks/Private/Systems/MovieSceneComponentMobilitySystem.cpp
+++ b/Engine/Source/Runtime/MovieSceneTracks/Private/Systems/MovieSceneComponentMobilitySystem.cpp
@@ -191,11 +191,6 @@ void UMovieSceneComponentMobilitySystem::SavePreAnimatedState(const FPreAnimatio
 		const FEntityAllocation* Allocation = Item.GetAllocation();
 		const bool bWantsRestore = Item.GetAllocationType().Contains(FBuiltInComponentTypes::Get()->Tags.RestoreState);
 
-		if (!Extension->IsCapturingGlobalState() && !bWantsRestore)
-		{
-			return;
-		}
-
 		FCachePreAnimatedValueParams CacheValueParams;
 		CacheValueParams.bForcePersist = Item.GetAllocationType().Contains(FMovieSceneTracksComponentTypes::Get()->AttachParent);
 
diff --git a/Engine/Source/Runtime/MovieSceneTracks/Private/TrackInstances/MovieSceneCameraCutGameHandler.cpp b/Engine/Source/Runtime/MovieSceneTracks/Private/TrackInstances/MovieSceneCameraCutGameHandler.cpp
--- a/Engine/Source/Runtime/MovieSceneTracks/Private/TrackInstances/MovieSceneCameraCutGameHandler.cpp
+++ b/Engine/Source/Runtime/MovieSceneTracks/Private/TrackInstances/MovieSceneCameraCutGameHandler.cpp
@@ -104,9 +104,9 @@ void FPreAnimatedCameraCutTraits::RestorePreAnimatedValue(
 
 
 	APlayerController* PC = nullptr;
-	if (ULocalPlayer* PreviousViewTarget = Cast<ULocalPlayer>(CachedValue.LastLocalPlayer.ResolveObjectPtr()))
+	if (ULocalPlayer* PreviousLocalPlayer = Cast<ULocalPlayer>(CachedValue.LastLocalPlayer.ResolveObjectPtr()))
 	{
-		if (APlayerController* OwnerController = Cast<APlayerController>(PreviousViewTarget->GetPlayerController(World)))
+		if (APlayerController* OwnerController = Cast<APlayerController>(PreviousLocalPlayer->GetPlayerController(World)))
 		{
 			PC = OwnerController;
 		}
@@ -311,7 +311,7 @@ void FCameraCutGameHandler::SetCameraCut(
 	if (CameraObject == nullptr)
 	{
 		TSharedPtr<FPreAnimatedCameraCutStorage> PreAnimatedStorage = Linker->PreAnimatedState.FindStorage(FPreAnimatedCameraCutStorage::StorageID);
-		FPreAnimatedStorageIndex StorageIndex = PreAnimatedStorage->FindStorageIndex(0);
+		FPreAnimatedStorageIndex StorageIndex = PreAnimatedStorage ? PreAnimatedStorage->FindStorageIndex(0) : FPreAnimatedStorageIndex();
 		if (ensureMsgf(StorageIndex.IsValid(), TEXT("Blending camera back to gameplay but can't find pre-animated camera info!")))
 		{
 			FPreAnimatedCameraCutState CachedValue = PreAnimatedStorage->GetCachedValue(StorageIndex);
diff --git a/Engine/Source/Runtime/MovieSceneTracks/Private/TrackInstances/MovieSceneCameraCutTrackInstance.cpp b/Engine/Source/Runtime/MovieSceneTracks/Private/TrackInstances/MovieSceneCameraCutTrackInstance.cpp
--- a/Engine/Source/Runtime/MovieSceneTracks/Private/TrackInstances/MovieSceneCameraCutTrackInstance.cpp
+++ b/Engine/Source/Runtime/MovieSceneTracks/Private/TrackInstances/MovieSceneCameraCutTrackInstance.cpp
@@ -128,6 +128,7 @@ struct FBlendedCameraCut
 
 	float PreviewBlendFactor = -1.f;
 	bool bCanBlend = false;
+	bool bSavePreAnimatedValue = false;
 
 	FBlendedCameraCut()
 	{}
@@ -232,11 +233,17 @@ public:
 
 		// Save pre-animated state only now, because we don't want to start tracking this camera cut
 		// unless we know it actually resolves to a camera actor (see above) and will actually do something.
-		FScopedPreAnimatedCaptureSource CaptureSource(Linker, Params.Input);
-		FCameraCutGameHandler::CachePreAnimatedValue(Linker, SequenceInstance);
+		if (Params.bSavePreAnimatedValue)
+		{
+			FScopedPreAnimatedCaptureSource CaptureSource(Linker, Params.Input);
+			if (Linker->PreAnimatedState.IsCapturingGlobalState() || CaptureSource.WantsRestoreState())
+			{
+				FCameraCutGameHandler::CachePreAnimatedValue(Linker, SequenceInstance);
 #if WITH_EDITOR
-		FCameraCutEditorHandler::CachePreAnimatedValue(Linker, SequenceInstance);
+				FCameraCutEditorHandler::CachePreAnimatedValue(Linker, SequenceInstance);
 #endif
+			}
+		}
 
 		// We do a straight cut if we have jumped in time and we either (1) always cut on loops or (2) it's not a loop.
 		UMovieSceneSequence* RootSequence = SequenceInstance.GetSharedPlaybackState()->GetRootSequence();
@@ -432,6 +439,7 @@ void UMovieSceneCameraCutTrackInstance::OnAnimate()
 			FBlendedCameraCut Params(Input, CameraBindingID, SequenceInstance.GetSequenceID());
 			Params.LocalDirection = Context.GetDirection();
 			Params.bCanBlend = Track->bCanBlend;
+			Params.bSavePreAnimatedValue = InputInfo.bSavePreAnimatedValue;
 
 			// Get start/current/end time.
 			Params.LocalContextTime = Context.GetTime();
@@ -565,6 +573,8 @@ void UMovieSceneCameraCutTrackInstance::OnAnimate()
 			// Fully active.
 			FinalCameraCut = PriorityCameraCut;
 		}
+
+		FinalCameraCut.bSavePreAnimatedValue = (CameraCutParams[0].bSavePreAnimatedValue || CameraCutParams[1].bSavePreAnimatedValue);
 	}
 	else if (CameraCutParams.Num() == 1)
 	{
@@ -592,6 +602,12 @@ void UMovieSceneCameraCutTrackInstance::OnAnimate()
 		const FSequenceInstance& SequenceInstance = InstanceRegistry->GetInstance(FinalCameraCut.Input.InstanceHandle);
 		Animator.AnimateBlendedCameraCut(FinalCameraCut, Linker, SequenceInstance);
 	}
+
+	// Mark all existing inputs as not needing to have their pre-animated state saved anymore.
+	for (FCameraCutInputInfo& InputInfo : SortedInputInfos)
+	{
+		InputInfo.bSavePreAnimatedValue = false;
+	}
 }
 
 void UMovieSceneCameraCutTrackInstance::OnEndUpdateInputs()
@@ -631,6 +647,9 @@ void UMovieSceneCameraCutTrackInstance::OnEndUpdateInputs()
 					InputInfo.GlobalStartTime = WorldTime;
 				}
 			}
+
+			// This is a new input, we may need to save its pre-animated state.
+			InputInfo.bSavePreAnimatedValue = true;
 		}
 
 		SortedInputInfos.Add(InputInfo);
diff --git a/Engine/Source/Runtime/MovieSceneTracks/Public/TrackInstances/MovieSceneCameraCutTrackInstance.h b/Engine/Source/Runtime/MovieSceneTracks/Public/TrackInstances/MovieSceneCameraCutTrackInstance.h
--- a/Engine/Source/Runtime/MovieSceneTracks/Public/TrackInstances/MovieSceneCameraCutTrackInstance.h
+++ b/Engine/Source/Runtime/MovieSceneTracks/Public/TrackInstances/MovieSceneCameraCutTrackInstance.h
@@ -82,6 +82,7 @@ private:
 	{
 		FMovieSceneTrackInstanceInput Input;
 		float GlobalStartTime = 0.f;
+		bool bSavePreAnimatedValue = false;
 	};
 
 	FCameraCutCache CameraCutCache;
