Sequencer: AutoSetNetAddressableName was called too aggressively - in CopyObjectTemplate, which meant that the user couldn't toggle bNetAddressableName reliably because it would be reset when the object template was saved. It's now only called in SetObjectTemplate, essentially when the class changes.

#Synced-CL 50777405

diff --git a/Engine/Source/Runtime/MovieSceneTracks/Private/Bindings/MovieSceneSpawnableActorBinding.cpp b/Engine/Source/Runtime/MovieSceneTracks/Private/Bindings/MovieSceneSpawnableActorBinding.cpp
--- a/Engine/Source/Runtime/MovieSceneTracks/Private/Bindings/MovieSceneSpawnableActorBinding.cpp
+++ b/Engine/Source/Runtime/MovieSceneTracks/Private/Bindings/MovieSceneSpawnableActorBinding.cpp
@@ -448,8 +448,6 @@ void UMovieSceneSpawnableActorBinding::CopyObjectTemplate(UObject* InSourceObjec
 		// TODO: We should move this out of FMovieSceneSpawnable eventually
 		FMovieSceneSpawnable::MarkSpawnableTemplate(*ActorTemplate);
 
-		AutoSetNetAddressableName();
-
 		MovieSceneSequence.MarkPackageDirty();
 	}
 }
@@ -629,7 +627,6 @@ void UMovieSceneSpawnableActorBinding::PostDuplicate(EDuplicateMode::Type Duplic
 			{
 				FMovieSceneSpawnable::MarkSpawnableTemplate(*ActorTemplate);
 			}
-			AutoSetNetAddressableName();
 		}
 	}
 }
