Sequencer: Fix for input channels not resetting when a metasound is changed

#Synced-CL 50918309

diff --git a/Engine/Source/Runtime/MovieSceneTracks/Private/Sections/MovieSceneAudioSection.cpp b/Engine/Source/Runtime/MovieSceneTracks/Private/Sections/MovieSceneAudioSection.cpp
--- a/Engine/Source/Runtime/MovieSceneTracks/Private/Sections/MovieSceneAudioSection.cpp
+++ b/Engine/Source/Runtime/MovieSceneTracks/Private/Sections/MovieSceneAudioSection.cpp
@@ -405,6 +405,13 @@ void UMovieSceneAudioSection::PostEditChangeProperty(FPropertyChangedEvent& Prop
 			SetCompletionMode();
 		}
 	}
+
+	// If the sound is changed, invalidate channel proxy to regenerate input tracks
+	// Outside of main if statement as this needs to handle Sound being a nullptr as well
+	if (PropertyChangedEvent.MemberProperty && PropertyChangedEvent.MemberProperty->GetFName() == GET_MEMBER_NAME_CHECKED(UMovieSceneAudioSection, Sound))
+	{
+		SetSound(Sound);
+	}
 }
 #endif // WITH_EDITOR
 	
@@ -475,7 +482,7 @@ UMovieSceneSection* UMovieSceneAudioSection::SplitSection(FQualifiedFrameTime Sp
 
 void UMovieSceneAudioSection::SetSound(USoundBase* InSound)
 {
-	const bool bIsNewSound = (Sound == nullptr);
+	const bool bIsNewSound = (Sound != InSound);
 
 	Sound = InSound;
 
