[Sequencer] Set AudioComponents that outlive MovieSceneAudioSystem to destroy on completion

AudioComponents can continue to play beyond the time a MovieScene has completed, and are added to the world to do so. This can show up as a leak when GC sees the AudioComponent with no owning actor. This change identifies AudioComponents that will continue to play and sets the OnAudioFinished callback to a lambda that destroys the AudioComponent.

[RN] Setting AudioComponents created by the MovieSceneAudioSection to delete on completion to prevent leaking.
#Synced-CL 51265256

diff --git a/Engine/Source/Runtime/MovieSceneTracks/Private/Systems/MovieSceneAudioSystem.cpp b/Engine/Source/Runtime/MovieSceneTracks/Private/Systems/MovieSceneAudioSystem.cpp
--- a/Engine/Source/Runtime/MovieSceneTracks/Private/Systems/MovieSceneAudioSystem.cpp
+++ b/Engine/Source/Runtime/MovieSceneTracks/Private/Systems/MovieSceneAudioSystem.cpp
@@ -970,10 +970,33 @@ void UMovieSceneAudioSystem::OnUnlink()
 			{
 				UObject* Actor = AudioComponentsForActor.Key.ResolveObjectPtr();
 				UObject* Section = AudioComponentForSection.Key.Value.ResolveObjectPtr();
-				UE_LOG(LogMovieScene, Warning, TEXT("Cleaning audio component '%s' for section '%s' on actor '%s'"),
+				bool bIsCurrentlyLooping = AudioComponent->GetSound() && AudioComponent->GetSound()->IsLooping() && AudioComponent->IsPlaying();
+				UE_LOG(LogMovieScene, Warning, TEXT("%sequencer sound playing past section end - Audio component '%s' for section '%s' on actor '%s'."),
+						bIsCurrentlyLooping ? TEXT("Looping s") : TEXT("S"),
 						*AudioComponent->GetPathName(),
 						Section ? *Section->GetPathName() : TEXT("<null>"),
 						Actor ? *Actor->GetPathName() : TEXT("<null>"));
+
+				ensure(!bIsCurrentlyLooping);
+
+				// If we are unlinking a playing audio component, set it up to destroy on completion so it can continue to play 
+				// after the movie scene is complete
+				if (AudioComponent->IsPlaying())
+				{
+					AudioComponent->OnAudioFinishedNative.AddLambda([](UAudioComponent* FinishingComponent)
+						{
+							UE_LOG(LogMovieScene, Display, TEXT("Destroying audio component '%s' upon completion"),
+								*FinishingComponent->GetPathName());
+							FinishingComponent->DestroyComponent();
+						});
+				}
+				// otherwise if it is not playing, just destroy the audio component now
+				else
+				{
+					UE_LOG(LogMovieScene, Display, TEXT("Destroying stopped audio component '%s' on MovieSceneAudioSystem::OnUnlink"),
+						*AudioComponent->GetPathName());
+					AudioComponent->DestroyComponent();
+				}
 			}
 		}
 	}
