Sequencer: Rebind delegates or clear (if not bound) for OnQueueSubtitles, OnAudioFinished, and OnAudioPlaybackPercent when an AudioComponent is recycled

Follow up to 3293909 which added the ability for these delegates to exist

#Synced-CL 51308618

diff --git a/Engine/Source/Runtime/MovieSceneTracks/Private/Systems/MovieSceneAudioSystem.cpp b/Engine/Source/Runtime/MovieSceneTracks/Private/Systems/MovieSceneAudioSystem.cpp
--- a/Engine/Source/Runtime/MovieSceneTracks/Private/Systems/MovieSceneAudioSystem.cpp
+++ b/Engine/Source/Runtime/MovieSceneTracks/Private/Systems/MovieSceneAudioSystem.cpp
@@ -411,19 +411,6 @@ private:
 						AudioSystem->PreAnimatedStorage->CachePreAnimatedValue(
 							FCachePreAnimatedValueParams(), AudioComponent,
 							[](FObjectKey InKey) { return EPreAnimatedAudioStateType::AudioComponentLifespan; });
-
-						if (AudioSection->GetOnQueueSubtitles().IsBound())
-						{
-							AudioComponent->OnQueueSubtitles = AudioSection->GetOnQueueSubtitles();
-						}
-						if (AudioSection->GetOnAudioFinished().IsBound())
-						{
-							AudioComponent->OnAudioFinished = AudioSection->GetOnAudioFinished();
-						}
-						if (AudioSection->GetOnAudioPlaybackPercent().IsBound())
-						{
-							AudioComponent->OnAudioPlaybackPercent = AudioSection->GetOnAudioPlaybackPercent();
-						}
 					}
 				}
 			}
@@ -442,6 +429,33 @@ private:
 					{
 						AudioComponent->DetachFromComponent(FDetachmentTransformRules::KeepRelativeTransform);
 					}
+
+					if (AudioSection->GetOnQueueSubtitles().IsBound())
+					{
+						AudioComponent->OnQueueSubtitles = AudioSection->GetOnQueueSubtitles();
+					}
+					else
+					{
+						AudioComponent->OnQueueSubtitles.Clear();
+					}
+
+					if (AudioSection->GetOnAudioFinished().IsBound())
+					{
+						AudioComponent->OnAudioFinished = AudioSection->GetOnAudioFinished();
+					}
+					else
+					{
+						AudioComponent->OnAudioFinished.Clear();
+					}
+
+					if (AudioSection->GetOnAudioPlaybackPercent().IsBound())
+					{
+						AudioComponent->OnAudioPlaybackPercent = AudioSection->GetOnAudioPlaybackPercent();
+					}
+					else
+					{
+						AudioComponent->OnAudioPlaybackPercent.Clear();
+					}
 				}
 
 				EvaluationData->VolumeMultiplier = VolumeMultiplier * AudioSection->EvaluateEasing(Context.GetTime());
@@ -467,24 +481,43 @@ private:
 					AudioSystem->PreAnimatedStorage->CachePreAnimatedValue(
 						FCachePreAnimatedValueParams(), AudioComponent,
 						[](FObjectKey InKey) { return EPreAnimatedAudioStateType::AudioComponentLifespan; });
+				}
+			}
 
+			if (EvaluationData)
+			{
+				UAudioComponent* AudioComponent = EvaluationData->AudioComponent.Get();
+
+				if (AudioComponent)
+				{
 					if (AudioSection->GetOnQueueSubtitles().IsBound())
 					{
 						AudioComponent->OnQueueSubtitles = AudioSection->GetOnQueueSubtitles();
 					}
+					else
+					{
+						AudioComponent->OnQueueSubtitles.Clear();
+					}
+
 					if (AudioSection->GetOnAudioFinished().IsBound())
 					{
 						AudioComponent->OnAudioFinished = AudioSection->GetOnAudioFinished();
 					}
+					else
+					{
+						AudioComponent->OnAudioFinished.Clear();
+					}
+
 					if (AudioSection->GetOnAudioPlaybackPercent().IsBound())
 					{
 						AudioComponent->OnAudioPlaybackPercent = AudioSection->GetOnAudioPlaybackPercent();
 					}
+					else
+					{
+						AudioComponent->OnAudioPlaybackPercent.Clear();
+					}
 				}
-			}
 
-			if (EvaluationData)
-			{
 				EvaluationData->VolumeMultiplier = VolumeMultiplier;
 				EvaluationData->PitchMultiplier = PitchMultiplier;
 
@@ -1255,6 +1288,19 @@ UMovieSceneAudioSystem::FAudioComponentEvaluationData* UMovieSceneAudioSystem::A
 		ExistingData->AudioComponent = NewComponent;
 		ExistingData->bAudioComponentHasBeenPlayed = false;
 
+		if (Section->GetOnQueueSubtitles().IsBound())
+		{
+			NewComponent->OnQueueSubtitles = Section->GetOnQueueSubtitles();
+		}
+		if (Section->GetOnAudioFinished().IsBound())
+		{
+			NewComponent->OnAudioFinished = Section->GetOnAudioFinished();
+		}
+		if (Section->GetOnAudioPlaybackPercent().IsBound())
+		{
+			NewComponent->OnAudioPlaybackPercent = Section->GetOnAudioPlaybackPercent();
+		}
+
 #if WITH_EDITOR
 		static int32 ScrubSoundCounter = 0;
 		UScrubbedSound* ScrubbedSound = NewObject<UScrubbedSound>(Actor, FName(*FString::Printf(TEXT("ScrubbedSound_Bound_%i"), ScrubSoundCounter++)));
@@ -1307,6 +1353,19 @@ UMovieSceneAudioSystem::FAudioComponentEvaluationData* UMovieSceneAudioSystem::A
 		ExistingData->AudioComponent = NewComponent;
 		ExistingData->bAudioComponentHasBeenPlayed = false;
 
+		if (Section->GetOnQueueSubtitles().IsBound())
+		{
+			NewComponent->OnQueueSubtitles = Section->GetOnQueueSubtitles();
+		}
+		if (Section->GetOnAudioFinished().IsBound())
+		{
+			NewComponent->OnAudioFinished = Section->GetOnAudioFinished();
+		}
+		if (Section->GetOnAudioPlaybackPercent().IsBound())
+		{
+			NewComponent->OnAudioPlaybackPercent = Section->GetOnAudioPlaybackPercent();
+		}
+
 #if WITH_EDITOR
 		static int32 ScrubSoundCounter = 0;
 		UObject* OuterObject = GetTransientPackage();
