Sequencer: Fix old path bone matching for blend child of root. We were getting gimbal lock related flips when decomposing quats to eulers and then back again. Fix by decomposing the delta rather than the transform eulers themselves.

#Synced-CL 51336034

diff --git a/Engine/Source/Runtime/MovieSceneTracks/Private/Tracks/MovieSceneCommonAnimationTrack.cpp b/Engine/Source/Runtime/MovieSceneTracks/Private/Tracks/MovieSceneCommonAnimationTrack.cpp
--- a/Engine/Source/Runtime/MovieSceneTracks/Private/Tracks/MovieSceneCommonAnimationTrack.cpp
+++ b/Engine/Source/Runtime/MovieSceneTracks/Private/Tracks/MovieSceneCommonAnimationTrack.cpp
@@ -782,13 +782,8 @@ void UMovieSceneCommonAnimationTrack::MatchSectionByBoneTransform(bool bMatchWit
 			FVector SecondTransformTranslation = SecondTransform.GetTranslation();
 			FQuat FirstTransformQuat = FirstTransform.GetRotation();
 			FQuat SecondTransformQuat = SecondTransform.GetRotation();
-
 			FirstTransformQuat.EnforceShortestArcWith(SecondTransformQuat);
 
-			FRotator FirstTransformRotation(FRotator::MakeFromEuler(EulerFromQuat(FirstTransformQuat, RotationOrder)));
-			FRotator SecondTransformRotation(FRotator::MakeFromEuler(EulerFromQuat(SecondTransformQuat, RotationOrder)));
-			SecondTransformRotation.SetClosestToMe(FirstTransformRotation);
-
 			if (!CurrentSection->bMatchTranslation)
 			{
 				FirstTransformTranslation.X = SecondTransformTranslation.X;
@@ -801,19 +796,28 @@ void UMovieSceneCommonAnimationTrack::MatchSectionByBoneTransform(bool bMatchWit
 			}
 			FirstTransform.SetTranslation(FirstTransformTranslation);
 
+			// Selective rotation matching: decompose the delta between the two
+			// quaternions to Euler (the delta is small, avoiding gimbal lock),
+			// zero unmatched axes, and recompose.
+			FQuat DeltaQuat = SecondTransformQuat.Inverse() * FirstTransformQuat;
+			DeltaQuat.Normalize();
+			FVector DeltaEuler = EulerFromQuat(DeltaQuat, RotationOrder);
+
 			if (!CurrentSection->bMatchRotationYaw)
 			{
-				FirstTransformRotation.Yaw = SecondTransformRotation.Yaw;
+				DeltaEuler.Z = 0.0;
 			}
 			if (!CurrentSection->bMatchRotationPitch)
 			{
-				FirstTransformRotation.Pitch = SecondTransformRotation.Pitch;
+				DeltaEuler.Y = 0.0;
 			}
 			if (!CurrentSection->bMatchRotationRoll)
 			{
-				FirstTransformRotation.Roll = SecondTransformRotation.Roll;
+				DeltaEuler.X = 0.0;
 			}
-			FirstTransformQuat = QuatFromEuler(FirstTransformRotation.Euler(), RotationOrder);
+			FQuat AdjustedDelta = QuatFromEuler(DeltaEuler, RotationOrder);
+			FirstTransformQuat = SecondTransformQuat * AdjustedDelta;
+			FirstTransformQuat.Normalize();
 			FirstTransform.SetRotation(FirstTransformQuat);
 
 			// Below is the match but we need to use GetRelativeTransformReverse since Inverse doesn't work as expected.
@@ -849,13 +853,8 @@ void UMovieSceneCommonAnimationTrack::MatchSectionByBoneTransform(bool bMatchWit
 			FVector SecondTransformTranslation = SecondTransform.GetTranslation();
 			FQuat FirstTransformQuat = FirstTransform.GetRotation();
 			FQuat SecondTransformQuat = SecondTransform.GetRotation();
-
 			SecondTransformQuat.EnforceShortestArcWith(FirstTransformQuat);
 
-			FRotator FirstTransformRotation(FRotator::MakeFromEuler(EulerFromQuat(FirstTransformQuat, RotationOrder)));
-			FRotator SecondTransformRotation(FRotator::MakeFromEuler(EulerFromQuat(SecondTransformQuat, RotationOrder)));
-			FirstTransformRotation.SetClosestToMe(SecondTransformRotation);
-
 			if (!CurrentSection->bMatchTranslation)
 			{
 				SecondTransformTranslation.X = FirstTransformTranslation.X;
@@ -868,19 +867,25 @@ void UMovieSceneCommonAnimationTrack::MatchSectionByBoneTransform(bool bMatchWit
 			}
 			SecondTransform.SetTranslation(SecondTransformTranslation);
 
+			FQuat DeltaQuat = FirstTransformQuat.Inverse() * SecondTransformQuat;
+			DeltaQuat.Normalize();
+			FVector DeltaEuler = EulerFromQuat(DeltaQuat, RotationOrder);
+
 			if (!CurrentSection->bMatchRotationYaw)
 			{
-				SecondTransformRotation.Yaw = FirstTransformRotation.Yaw;
+				DeltaEuler.Z = 0.0;
 			}
 			if (!CurrentSection->bMatchRotationPitch)
 			{
-				SecondTransformRotation.Pitch = FirstTransformRotation.Pitch;
+				DeltaEuler.Y = 0.0;
 			}
 			if (!CurrentSection->bMatchRotationRoll)
 			{
-				SecondTransformRotation.Roll = FirstTransformRotation.Roll;
+				DeltaEuler.X = 0.0;
 			}
-			SecondTransformQuat = QuatFromEuler(SecondTransformRotation.Euler(), RotationOrder);
+			FQuat AdjustedDelta = QuatFromEuler(DeltaEuler, RotationOrder);
+			SecondTransformQuat = FirstTransformQuat * AdjustedDelta;
+			SecondTransformQuat.Normalize();
 			SecondTransform.SetRotation(SecondTransformQuat);
 
 			//GetRelativeTransformReverse returns this(-1)* Other, and parameter is Other.
