Sequencer: Bone matching- Fix EulerFromQuat transposed matrix elements causing rotation sign flip in bone matching.

#Synced-CL 51609808

diff --git a/Engine/Source/Runtime/MovieSceneTracks/Public/MovieSceneAnimationRotationUtils.h b/Engine/Source/Runtime/MovieSceneTracks/Public/MovieSceneAnimationRotationUtils.h
--- a/Engine/Source/Runtime/MovieSceneTracks/Public/MovieSceneAnimationRotationUtils.h
+++ b/Engine/Source/Runtime/MovieSceneTracks/Public/MovieSceneAnimationRotationUtils.h
@@ -101,23 +101,33 @@ inline FVector EulerFromQuat(const FQuat& Rotation, ERotationOrder RotationOrder
 	double XX2 = X * X2;
 	double XY2 = X * Y2;
 	double XZ2 = X * Z2;
+	double YX2 = Y * X2;
 	double YY2 = Y * Y2;
 	double YZ2 = Y * Z2;
+	double ZX2 = Z * X2;
+	double ZY2 = Z * Y2;
 	double ZZ2 = Z * Z2;
 	double WX2 = W * X2;
 	double WY2 = W * Y2;
 	double WZ2 = W * Z2;
 
-	FVector AxisX(1.0 - (YY2 + ZZ2), XY2 + WZ2, XZ2 - WY2);
-	FVector AxisY(XY2 - WZ2, 1.0 - (XX2 + ZZ2), YZ2 + WX2);
-	FVector AxisZ(XZ2 + WY2, YZ2 - WX2, 1.0 - (XX2 + YY2));
+	FVector AxisX, AxisY, AxisZ;
+	AxisX.X = (1.0 - (YY2 + ZZ2));
+	AxisY.X = (XY2 + WZ2);
+	AxisZ.X = (XZ2 - WY2);
+	AxisX.Y = (XY2 - WZ2);
+	AxisY.Y = (1.0 - (XX2 + ZZ2));
+	AxisZ.Y = (YZ2 + WX2);
+	AxisX.Z = (XZ2 + WY2);
+	AxisY.Z = (YZ2 - WX2);
+	AxisZ.Z = (1.0 - (XX2 + YY2));
 
 	FVector Result = FVector::ZeroVector;
 
-	switch (RotationOrder)
+	if (RotationOrder == ERotationOrder::XYZ)
 	{
-	case ERotationOrder::XYZ:
-		Result.Y = FMath::Asin(-FMath::Clamp(AxisZ.X, -1.0, 1.0));
+		Result.Y = FMath::Asin(-FMath::Clamp<double>(AxisZ.X, -1.0, 1.0));
+
 		if (FMath::Abs(AxisZ.X) < 1.0 - SMALL_NUMBER)
 		{
 			Result.X = FMath::Atan2(AxisZ.Y, AxisZ.Z);
@@ -128,10 +138,12 @@ inline FVector EulerFromQuat(const FQuat& Rotation, ERotationOrder RotationOrder
 			Result.X = 0.0;
 			Result.Z = FMath::Atan2(-AxisX.Y, AxisY.Y);
 		}
-		break;
+	}
+	else if (RotationOrder == ERotationOrder::XZY)
+	{
+
+		Result.Z = FMath::Asin(FMath::Clamp<double>(AxisY.X, -1.0, 1.0));
 
-	case ERotationOrder::XZY:
-		Result.Z = FMath::Asin(FMath::Clamp(AxisY.X, -1.0, 1.0));
 		if (FMath::Abs(AxisY.X) < 1.0 - SMALL_NUMBER)
 		{
 			Result.X = FMath::Atan2(-AxisY.Z, AxisY.Y);
@@ -142,10 +154,11 @@ inline FVector EulerFromQuat(const FQuat& Rotation, ERotationOrder RotationOrder
 			Result.X = 0.0;
 			Result.Y = FMath::Atan2(AxisX.Z, AxisZ.Z);
 		}
-		break;
+	}
+	else if (RotationOrder == ERotationOrder::YXZ)
+	{
+		Result.X = FMath::Asin(FMath::Clamp<double>(AxisZ.Y, -1.0, 1.0));
 
-	case ERotationOrder::YXZ:
-		Result.X = FMath::Asin(FMath::Clamp(AxisZ.Y, -1.0, 1.0));
 		if (FMath::Abs(AxisZ.Y) < 1.0 - SMALL_NUMBER)
 		{
 			Result.Y = FMath::Atan2(-AxisZ.X, AxisZ.Z);
@@ -156,10 +169,11 @@ inline FVector EulerFromQuat(const FQuat& Rotation, ERotationOrder RotationOrder
 			Result.Y = 0.0;
 			Result.Z = FMath::Atan2(AxisY.X, AxisX.X);
 		}
-		break;
+	}
+	else if (RotationOrder == ERotationOrder::YZX)
+	{
+		Result.Z = FMath::Asin(-FMath::Clamp<double>(AxisX.Y, -1.0, 1.0));
 
-	case ERotationOrder::YZX:
-		Result.Z = FMath::Asin(-FMath::Clamp(AxisX.Y, -1.0, 1.0));
 		if (FMath::Abs(AxisX.Y) < 1.0 - SMALL_NUMBER)
 		{
 			Result.X = FMath::Atan2(AxisZ.Y, AxisY.Y);
@@ -170,10 +184,11 @@ inline FVector EulerFromQuat(const FQuat& Rotation, ERotationOrder RotationOrder
 			Result.X = FMath::Atan2(-AxisY.Z, AxisZ.Z);
 			Result.Y = 0.0;
 		}
-		break;
+	}
+	else if (RotationOrder == ERotationOrder::ZXY)
+	{
+		Result.X = FMath::Asin(-FMath::Clamp<double>(AxisY.Z, -1.0, 1.0));
 
-	case ERotationOrder::ZXY:
-		Result.X = FMath::Asin(-FMath::Clamp(AxisY.Z, -1.0, 1.0));
 		if (FMath::Abs(AxisY.Z) < 1.0 - SMALL_NUMBER)
 		{
 			Result.Y = FMath::Atan2(AxisX.Z, AxisZ.Z);
@@ -184,10 +199,11 @@ inline FVector EulerFromQuat(const FQuat& Rotation, ERotationOrder RotationOrder
 			Result.Y = FMath::Atan2(-AxisZ.X, AxisX.X);
 			Result.Z = 0.0;
 		}
-		break;
+	}
+	else if (RotationOrder == ERotationOrder::ZYX)
+	{
+		Result.Y = FMath::Asin(FMath::Clamp<double>(AxisX.Z, -1.0, 1.0));
 
-	case ERotationOrder::ZYX:
-		Result.Y = FMath::Asin(FMath::Clamp(AxisX.Z, -1.0, 1.0));
 		if (FMath::Abs(AxisX.Z) < 1.0 - SMALL_NUMBER)
 		{
 			Result.X = FMath::Atan2(-AxisY.Z, AxisZ.Z);
@@ -198,10 +214,9 @@ inline FVector EulerFromQuat(const FQuat& Rotation, ERotationOrder RotationOrder
 			Result.X = FMath::Atan2(AxisZ.Y, AxisY.Y);
 			Result.Z = 0.0;
 		}
-		break;
 	}
 
-	return Result * (180.0 / UE_DOUBLE_PI);
+	return Result * 180.0 / DOUBLE_PI;
 }
 
 } // namespace UE::MovieScene
