Sequencer: Remove redundant RelevantComponent from UMovieSceneSkeletalAnimationSystem

RelevantComponent was set to SkeletalAnimation in the constructor, which caused
IsRelevant() to short-circuit to true whenever any entity had a SkeletalAnimation
component, bypassing IsRelevantImpl. This meant the system would get linked even
when all skeletal animation entities had the AnimMixerPoseProducer tag, which
IsRelevantImpl correctly filters out.

#Synced-CL 51609808

diff --git a/Engine/Source/Runtime/MovieSceneTracks/Private/Systems/MovieSceneSkeletalAnimationSystem.cpp b/Engine/Source/Runtime/MovieSceneTracks/Private/Systems/MovieSceneSkeletalAnimationSystem.cpp
--- a/Engine/Source/Runtime/MovieSceneTracks/Private/Systems/MovieSceneSkeletalAnimationSystem.cpp
+++ b/Engine/Source/Runtime/MovieSceneTracks/Private/Systems/MovieSceneSkeletalAnimationSystem.cpp
@@ -1097,8 +1097,6 @@ UMovieSceneSkeletalAnimationSystem::UMovieSceneSkeletalAnimationSystem(const FOb
 {
 	using namespace UE::MovieScene;
 
-	FMovieSceneTracksComponentTypes* TrackComponents = FMovieSceneTracksComponentTypes::Get();
-	RelevantComponent = TrackComponents->SkeletalAnimation;
 	Phase = ESystemPhase::Instantiation | ESystemPhase::Scheduling;
 
 	SystemCategories |= FSystemInterrogator::GetExcludedFromInterrogationCategory();
