Sequencer: Fix PostLoad bug where the display name is getting emptied even if already empty

43433169 fixed an issue where the MPC track name wouldn't be derived from the MPC by emptying the DisplayName in PostLoad. Unfortunately, it unnecessarily empties it since GetDisplayName() returns GetDefaultDisplayName() if empty. The fix is to check whether the internally stored DisplayName is already empty.

#Synced-CL 51670918

diff --git a/Engine/Source/Runtime/MovieScene/Public/MovieSceneNameableTrack.h b/Engine/Source/Runtime/MovieScene/Public/MovieSceneNameableTrack.h
--- a/Engine/Source/Runtime/MovieScene/Public/MovieSceneNameableTrack.h
+++ b/Engine/Source/Runtime/MovieScene/Public/MovieSceneNameableTrack.h
@@ -67,6 +67,9 @@ public:
 
 	MOVIESCENE_API virtual void OnRowIndicesChanged(const TMap<int32, int32>& NewToOldRowIndices) override;
 
+protected:
+	const FText& GetDisplayName_Internal() const { return DisplayName; }
+
 private:
 
 	/** The track's human readable display name. */
diff --git a/Engine/Source/Runtime/MovieSceneTracks/Private/Tracks/MovieSceneMaterialParameterCollectionTrack.cpp b/Engine/Source/Runtime/MovieSceneTracks/Private/Tracks/MovieSceneMaterialParameterCollectionTrack.cpp
--- a/Engine/Source/Runtime/MovieSceneTracks/Private/Tracks/MovieSceneMaterialParameterCollectionTrack.cpp
+++ b/Engine/Source/Runtime/MovieSceneTracks/Private/Tracks/MovieSceneMaterialParameterCollectionTrack.cpp
@@ -114,7 +114,7 @@ FText UMovieSceneMaterialParameterCollectionTrack::GetDefaultDisplayName() const
 
 void UMovieSceneMaterialParameterCollectionTrack::PostLoad()
 {
-	if (MPC && GetDisplayName().EqualTo(FText::FromString(MPC->GetName())))
+	if (!GetDisplayName_Internal().IsEmpty() && MPC && GetDisplayName_Internal().EqualTo(FText::FromString(MPC->GetName())))
 	{
 		UMovieSceneSequence* Sequence = GetTypedOuter<UMovieSceneSequence>();
 		UE_LOG(LogMovieScene, Display, TEXT("Material Parameter Collection Track: Removing display name %s from %s because it is identical to the material name and track names are now derived from the material name"), *MPC->GetName(), *GetNameSafe(Sequence));
