Sequencer: Null check for valid animation

#Synced-CL 51953805

diff --git a/Engine/Source/Runtime/MovieSceneTracks/Private/Sections/MovieSceneSkeletalAnimationSection.cpp b/Engine/Source/Runtime/MovieSceneTracks/Private/Sections/MovieSceneSkeletalAnimationSection.cpp
--- a/Engine/Source/Runtime/MovieSceneTracks/Private/Sections/MovieSceneSkeletalAnimationSection.cpp
+++ b/Engine/Source/Runtime/MovieSceneTracks/Private/Sections/MovieSceneSkeletalAnimationSection.cpp
@@ -529,7 +529,7 @@ void UMovieSceneSkeletalAnimationSection::ImportEntityImpl(UMovieSceneEntitySyst
 	const FMovieSceneTracksComponentTypes* TrackComponents = FMovieSceneTracksComponentTypes::Get();
 
 	const FGuid ObjectBindingID = InParams.GetObjectBindingID();
-	const bool bIsAdditive = Params.Animation->IsValidAdditive();
+	const bool bIsAdditive = Params.Animation ? Params.Animation->IsValidAdditive() : false;
 
 	if (!ObjectBindingID.IsValid())
 	{
