[Sequencer] Explicitly set a Widget Blueprint Animation's DefaultCompletionState to Restore State.
Consistently detect if a sound has changed to refresh EvalOptions.CompletionMode, and ensure EMovieSceneCompletionMode::RestoreState is always used for looping assets and sections.
UE-368964

#Synced-CL 50895136

diff --git a/Engine/Source/Runtime/MovieSceneTracks/Private/Sections/MovieSceneAudioSection.cpp b/Engine/Source/Runtime/MovieSceneTracks/Private/Sections/MovieSceneAudioSection.cpp
--- a/Engine/Source/Runtime/MovieSceneTracks/Private/Sections/MovieSceneAudioSection.cpp
+++ b/Engine/Source/Runtime/MovieSceneTracks/Private/Sections/MovieSceneAudioSection.cpp
@@ -77,7 +77,7 @@ UMovieSceneAudioSection::UMovieSceneAudioSection( const FObjectInitializer& Obje
 	AudioStartTime_DEPRECATED = AudioDeprecatedMagicNumber;
 	AudioDilationFactor_DEPRECATED = AudioDeprecatedMagicNumber;
 	AudioVolume_DEPRECATED = AudioDeprecatedMagicNumber;
-	bLooping = true;
+	bLooping = true; // for legacy compatibility, this must default to true, updated only for new content or edits via SetSound()
 	bSuppressSubtitles = false;
 	bOverrideAttenuation = false;
 	BlendType = EMovieSceneBlendType::Absolute;
@@ -293,7 +293,13 @@ void UMovieSceneAudioSection::SetupSoundInputParameters(USoundBase* InSoundBase)
 
 void UMovieSceneAudioSection::SetCompletionMode()
 {
-	if (bLooping) // All looping sounds are set to restore state to prevent "stuck on" sounds
+	bool bSoundIsLooping = false;
+	if (Sound)
+	{
+		bSoundIsLooping = !Sound->IsOneShot();
+	}
+	
+	if (bLooping || bSoundIsLooping) // All looping sounds are set to restore state to prevent "stuck on" sounds
 	{
 		EvalOptions.CompletionMode = EMovieSceneCompletionMode::RestoreState;
 	}
@@ -482,14 +488,14 @@ UMovieSceneSection* UMovieSceneAudioSection::SplitSection(FQualifiedFrameTime Sp
 
 void UMovieSceneAudioSection::SetSound(USoundBase* InSound)
 {
-	const bool bIsNewSound = (Sound != InSound);
+	// Sound is a UPROPERTY and can be set directly via FObjectPropertyBase::SetObjectPtrPropertyValue
+	// This function is called after that fact via PostEditChangeProperty().
+	// For new sections, it's called from UMovieSceneAudioTrack::AddNewSoundOnRow().
 
 	Sound = InSound;
-
-	if (Sound && bIsNewSound)
+	if (InSound)
 	{
 		bLooping = !Sound->IsOneShot();
-
 		SetCompletionMode();
 	}
 
@@ -724,4 +730,4 @@ void UMovieSceneAudioSection::ImportEntityImpl(UMovieSceneEntitySystemLinker* En
 			.Add(TrackComponents->AudioTriggerName, AudioTriggerNames[EntityIndex])
 		);
 	}
-}
+}
\ No newline at end of file
