On UNDO, update "when finished" property on audio section in Widget Blueprint Animation's sequence
On UNDO, UMovieSceneSignedObject::PostEditUndo() calls UMovieSceneAudioSection::PostEditChangeProperty() with a null MemberProperty -- update completion mode etc. via SetSound for this codepath.
UE-371218

#Synced-CL 50895136

diff --git a/Engine/Source/Runtime/MovieSceneTracks/Private/Sections/MovieSceneAudioSection.cpp b/Engine/Source/Runtime/MovieSceneTracks/Private/Sections/MovieSceneAudioSection.cpp
--- a/Engine/Source/Runtime/MovieSceneTracks/Private/Sections/MovieSceneAudioSection.cpp
+++ b/Engine/Source/Runtime/MovieSceneTracks/Private/Sections/MovieSceneAudioSection.cpp
@@ -77,7 +77,7 @@ UMovieSceneAudioSection::UMovieSceneAudioSection( const FObjectInitializer& Obje
 	AudioStartTime_DEPRECATED = AudioDeprecatedMagicNumber;
 	AudioDilationFactor_DEPRECATED = AudioDeprecatedMagicNumber;
 	AudioVolume_DEPRECATED = AudioDeprecatedMagicNumber;
-	bLooping = true; // for legacy compatibility, this must default to true, updated only for new content or edits via SetSound()
+	bLooping = true; // for legacy compatibility, this must default to true, updated only for new content or edits via SetSoundInternal()
 	bSuppressSubtitles = false;
 	bOverrideAttenuation = false;
 	BlendType = EMovieSceneBlendType::Absolute;
@@ -376,18 +376,31 @@ void UMovieSceneAudioSection::PostEditChangeProperty(FPropertyChangedEvent& Prop
 {
 	Super::PostEditChangeProperty(PropertyChangedEvent);
 
-	if (PropertyChangedEvent.MemberProperty && Sound)
+
+	if (PropertyChangedEvent.MemberProperty)
 	{
-		const FName bLoopingName = GET_MEMBER_NAME_CHECKED(UMovieSceneAudioSection, bLooping);
+		const FName PropertyName = PropertyChangedEvent.MemberProperty->GetFName();
 		const FName EvalOptionsName = GET_MEMBER_NAME_CHECKED(UMovieSceneSection, EvalOptions);
+		if (PropertyName == EvalOptionsName)
+		{
+			// Cache user intended changes to the completion mode
+			CachedCompletionMode = EvalOptions.CompletionMode;
+		}
+	}
 
+	// Update completion Mode to prevent looping sounds from playing indefinitely.
+	// Also, if the sound is changed, invalidate channel proxy to regenerate input tracks.
+	// This needs to be done even if Sound is nullptr.
+	// Also called for undo via UMovieSceneSignedObject::PostEditUndo, where PropertyChangedEvent.MemberProperty is nullptr.
+	const bool bInvalidateChannelProxy = ( !PropertyChangedEvent.MemberProperty || PropertyChangedEvent.MemberProperty->GetFName() == GET_MEMBER_NAME_CHECKED(UMovieSceneAudioSection, Sound) );
+	SetSoundInternal(Sound, bInvalidateChannelProxy); // calls SetCompletionMode() to prevent the user from setting a looping sound to Keep State
+
+	if (Sound)
+	{
 		// Clamp the StartOffset to the length of the SoundWave when not looping
 		// Setting the StartOffset to the size of the SoundWave will result in silence, irrespective of the size of the section
-		const FName PropertyName = PropertyChangedEvent.MemberProperty->GetFName();
 		const bool bHasDeterministicDuration = !bLooping && !Sound->IsProcedurallyGenerated() && Sound->IsOneShot();
-
-		if ((PropertyName == GET_MEMBER_NAME_CHECKED(UMovieSceneAudioSection, StartFrameOffset) || (PropertyName == bLoopingName))
-			&& bHasDeterministicDuration)
+		if (bHasDeterministicDuration)
 		{
 			// SoundDuration is local-specifc, meaning it can vary based on the value of GetCurrentLanguage()
 			const float SoundDuration = MovieSceneHelpers::GetSoundDuration(Sound);
@@ -396,27 +409,6 @@ void UMovieSceneAudioSection::PostEditChangeProperty(FPropertyChangedEvent& Prop
 
 			StartFrameOffset = FMath::Min(StartFrameOffset, MaxFrame);
 		}
-
-		// Change Completion Mode to prevent looping sounds from playing indefinitely.
-		if (PropertyName == bLoopingName)
-		{
-			SetCompletionMode();
-		}
-		else if (PropertyName == EvalOptionsName)
-		{
-			// Cache user intended changes to the completion mode
-			CachedCompletionMode = EvalOptions.CompletionMode;
-
-			// Prevent the user from setting a looping sound to Keep State
-			SetCompletionMode();
-		}
-	}
-
-	// If the sound is changed, invalidate channel proxy to regenerate input tracks
-	// Outside of main if statement as this needs to handle Sound being a nullptr as well
-	if (PropertyChangedEvent.MemberProperty && PropertyChangedEvent.MemberProperty->GetFName() == GET_MEMBER_NAME_CHECKED(UMovieSceneAudioSection, Sound))
-	{
-		SetSound(Sound);
 	}
 }
 #endif // WITH_EDITOR
@@ -487,6 +479,11 @@ UMovieSceneSection* UMovieSceneAudioSection::SplitSection(FQualifiedFrameTime Sp
 }
 
 void UMovieSceneAudioSection::SetSound(USoundBase* InSound)
+{
+	SetSoundInternal(InSound, true); // For new sections, it's called from UMovieSceneAudioTrack::AddNewSoundOnRow().
+}
+
+void UMovieSceneAudioSection::SetSoundInternal(USoundBase* InSound, bool bInvalidateChannelProxy)
 {
 	// Sound is a UPROPERTY and can be set directly via FObjectPropertyBase::SetObjectPtrPropertyValue
 	// This function is called after that fact via PostEditChangeProperty().
@@ -496,10 +493,13 @@ void UMovieSceneAudioSection::SetSound(USoundBase* InSound)
 	if (InSound)
 	{
 		bLooping = !Sound->IsOneShot();
-		SetCompletionMode();
 	}
+	SetCompletionMode();
 
-	InvalidateChannelProxy();
+	if (bInvalidateChannelProxy)
+	{
+		InvalidateChannelProxy();
+	}
 }
 
 USceneComponent* UMovieSceneAudioSection::GetAttachComponent(const AActor* InParentActor, const FMovieSceneActorReferenceKey& Key) const
diff --git a/Engine/Source/Runtime/MovieSceneTracks/Public/Sections/MovieSceneAudioSection.h b/Engine/Source/Runtime/MovieSceneTracks/Public/Sections/MovieSceneAudioSection.h
--- a/Engine/Source/Runtime/MovieSceneTracks/Public/Sections/MovieSceneAudioSection.h
+++ b/Engine/Source/Runtime/MovieSceneTracks/Public/Sections/MovieSceneAudioSection.h
@@ -224,6 +224,8 @@ private:
 
 	MOVIESCENETRACKS_API void SetCompletionMode();
 
+	void SetSoundInternal(class USoundBase* InSound, bool bInvalidateChannelProxy);
+
 	// Cache the users intended CompletionMode. 
 	// Used to revert the CompletionMode back to the users intent when not looping (when looping, completion mode is always RestoreState)
 	UPROPERTY()
