Sequencer: Fix crash introduced by 49072649

Section's Outer could be null (or not a UMovieScene) on paste

#Synced-CL 52445895

diff --git a/Engine/Source/Runtime/MovieSceneTracks/Private/Sections/MovieSceneAudioSection.cpp b/Engine/Source/Runtime/MovieSceneTracks/Private/Sections/MovieSceneAudioSection.cpp
--- a/Engine/Source/Runtime/MovieSceneTracks/Private/Sections/MovieSceneAudioSection.cpp
+++ b/Engine/Source/Runtime/MovieSceneTracks/Private/Sections/MovieSceneAudioSection.cpp
@@ -399,12 +399,13 @@ void UMovieSceneAudioSection::PostEditChangeProperty(FPropertyChangedEvent& Prop
 	{
 		// Clamp the StartOffset to the length of the SoundWave when not looping
 		// Setting the StartOffset to the size of the SoundWave will result in silence, irrespective of the size of the section
-		const bool bHasDeterministicDuration = !bLooping && !Sound->IsProcedurallyGenerated() && Sound->IsOneShot();
+		UMovieScene* MovieScene = GetTypedOuter<UMovieScene>();
+		const bool bHasDeterministicDuration = !bLooping && !Sound->IsProcedurallyGenerated() && Sound->IsOneShot() && MovieScene;
 		if (bHasDeterministicDuration)
 		{
 			// SoundDuration is local-specifc, meaning it can vary based on the value of GetCurrentLanguage()
 			const float SoundDuration = MovieSceneHelpers::GetSoundDuration(Sound);
-			const FFrameRate FrameRate = GetTypedOuter<UMovieScene>()->GetTickResolution();
+			const FFrameRate FrameRate = MovieScene->GetTickResolution();
 			const FFrameNumber MaxFrame = (SoundDuration * FrameRate).FloorToFrame();
 
 			StartFrameOffset = FMath::Min(StartFrameOffset, MaxFrame);
