Movie Scene: fixed issue where zero-scaled components would get their render state dirty called every time UE::MovieScene::SetComponentTransform is called.

#Synced-CL 52400047

diff --git a/Engine/Source/Runtime/MovieSceneTracks/Private/MovieSceneTracksCustomAccessors.cpp b/Engine/Source/Runtime/MovieSceneTracks/Private/MovieSceneTracksCustomAccessors.cpp
--- a/Engine/Source/Runtime/MovieSceneTracks/Private/MovieSceneTracksCustomAccessors.cpp
+++ b/Engine/Source/Runtime/MovieSceneTracks/Private/MovieSceneTracksCustomAccessors.cpp
@@ -64,7 +64,7 @@ void SetComponentTransform(USceneComponent* SceneComponent, const FIntermediate3
 		return;
 	}
 
-	const bool bFlipped = NewTransform.GetDeterminant() * SceneComponent->GetComponentTransform().GetDeterminant() <= 0;
+	const bool bFlipped = FMath::Sign(NewTransform.GetDeterminant()) != FMath::Sign(SceneComponent->GetComponentTransform().GetDeterminant());
 
 	SceneComponent->SetRelativeTransform(NewTransform, false, nullptr, bIsSimulatingPhysics ? ETeleportType::ResetPhysics : ETeleportType::None);
 
