MRQ - Audio - Fix Popping/Gaps/Repeats between Camera Cuts
Solution: Pause/resume the AudioComponent between shots rather than stopping and attempting to restart with a non-sample-accurate start time.
Now limited to only MRQ scenarios (bIsRecordingAudio) AND only when the audio section spans camera cuts -- it's stopped/destroyed as before in the camera cut containing AudioSectionEndFrame.
UE-235104

#Synced-CL 52605363

diff --git a/Engine/Plugins/MovieScene/MovieRenderPipeline/Source/MovieRenderPipelineCore/Private/MoviePipelineAudioRendering.cpp b/Engine/Plugins/MovieScene/MovieRenderPipeline/Source/MovieRenderPipelineCore/Private/MoviePipelineAudioRendering.cpp
--- a/Engine/Plugins/MovieScene/MovieRenderPipeline/Source/MovieRenderPipelineCore/Private/MoviePipelineAudioRendering.cpp
+++ b/Engine/Plugins/MovieScene/MovieRenderPipeline/Source/MovieRenderPipelineCore/Private/MoviePipelineAudioRendering.cpp
@@ -8,6 +8,8 @@
 #include "Engine/Engine.h"
 #include "AudioMixerBlueprintLibrary.h"
 #include "LevelSequence.h"
+#include "LevelSequencePlayer.h"
+#include "LevelSequenceActor.h"
 #include "AudioDevice.h"
 #include "AudioMixerDevice.h"
 #include "MoviePipelinePrimaryConfig.h"
@@ -69,6 +71,15 @@ void UMoviePipeline::SetupAudioRendering()
 
 void UMoviePipeline::TeardownAudioRendering()
 {
+	if (LevelSequenceActor && LevelSequenceActor->GetSequencePlayer())
+	{
+		TSharedRef<UE::MovieScene::FSharedPlaybackState> SharedPlaybackState = LevelSequenceActor->GetSequencePlayer()->GetSharedPlaybackState();
+		if (UE::MovieScene::FMovieSceneAudioRecordingCapability* AudioRecordingCapability = SharedPlaybackState->FindCapability<UE::MovieScene::FMovieSceneAudioRecordingCapability>())
+		{
+			AudioRecordingCapability->bIsRecordingAudio = false; // flag to notify MovieScene::FEvaluateAudio::Evaluate that MRQ is NOT active.
+		}
+	}
+
 	// Restore previous unfocused audio multiplier, to no longer force audio when unfocused
 	FApp::SetUnfocusedVolumeMultiplier(AudioState.PrevUnfocusedAudioMultiplier);
 
@@ -172,16 +183,17 @@ void UMoviePipeline::ProcessAudioTick()
 	if (CurrentShot->ShotInfo.State == EMovieRenderShotState::Rendering && !AudioState.bIsRecordingAudio)
 	{
 		// There can be time in the Unreal world that passes when MRQ is not yet active. When MRQ isn't active,
-		// then the NRT audio platform doesn't process anything and it builds up for when we do start to actually
-		// process it for a shot. This causes our actual audio to become offset from the sequence once lined up
-		// afterwards. So for now we'll just flush out anything that needed to be processed before we start capturing
-			
-		MixerDevice->Update(true);
-		
-		// We don't have a great way of determining how much outstanding audio there is, so we're just going to 
-		// process an arbitrary amount of time.
-		NRTPlatform->RenderAudio(30.f);
-		
+		// then the NRT audio platform doesn't process anything as all the AudioComponents that would have been playing are paused.
+		// On capturing the next camera, The AudioComponents resume playback, ensuring sample-accurate continuity across camera cuts.
+		if (LevelSequenceActor && LevelSequenceActor->GetSequencePlayer())
+		{
+			TSharedRef<UE::MovieScene::FSharedPlaybackState> SharedPlaybackState = LevelSequenceActor->GetSequencePlayer()->GetSharedPlaybackState();
+			if (UE::MovieScene::FMovieSceneAudioRecordingCapability* AudioRecordingCapability = SharedPlaybackState->FindCapability<UE::MovieScene::FMovieSceneAudioRecordingCapability>())
+			{
+				AudioRecordingCapability->bIsRecordingAudio = true; // flag to notify MovieScene::FEvaluateAudio::Evaluate that MRQ is active.
+			}
+		}
+
 		StartAudioRecording();
 	}
 
diff --git a/Engine/Source/Runtime/LevelSequence/Private/LevelSequencePlayer.cpp b/Engine/Source/Runtime/LevelSequence/Private/LevelSequencePlayer.cpp
--- a/Engine/Source/Runtime/LevelSequence/Private/LevelSequencePlayer.cpp
+++ b/Engine/Source/Runtime/LevelSequence/Private/LevelSequencePlayer.cpp
@@ -278,6 +278,7 @@ void ULevelSequencePlayer::InitializeLevelSequenceRootInstance(TSharedRef<UE::Mo
 {
 	NewSharedPlaybackState->SetOrAddCapabilityRaw<FMovieSceneSpawnRegister>(SpawnRegister.Get());
 	NewSharedPlaybackState->SetOrAddCapabilityRaw<FCameraCutPlaybackCapability>((FCameraCutPlaybackCapability*)this);
+	NewSharedPlaybackState->SetOrAddCapability<UE::MovieScene::FMovieSceneAudioRecordingCapability>(); // constructs inline in NewSharedPlaybackState
 }
 
 void ULevelSequencePlayer::TakeFrameSnapshot(FLevelSequencePlayerSnapshot& OutSnapshot) const
diff --git a/Engine/Source/Runtime/MovieScene/Private/EntitySystem/MovieSceneSharedPlaybackState.cpp b/Engine/Source/Runtime/MovieScene/Private/EntitySystem/MovieSceneSharedPlaybackState.cpp
--- a/Engine/Source/Runtime/MovieScene/Private/EntitySystem/MovieSceneSharedPlaybackState.cpp
+++ b/Engine/Source/Runtime/MovieScene/Private/EntitySystem/MovieSceneSharedPlaybackState.cpp
@@ -8,6 +8,7 @@
 
 namespace UE::MovieScene
 {
+	UE_DEFINE_MOVIESCENE_PLAYBACK_CAPABILITY(FMovieSceneAudioRecordingCapability)
 
 FSharedPlaybackState::FSharedPlaybackState(UMovieSceneEntitySystemLinker* InLinker)
 	: WeakLinker(InLinker)
diff --git a/Engine/Source/Runtime/MovieScene/Public/EntitySystem/MovieSceneSharedPlaybackState.h b/Engine/Source/Runtime/MovieScene/Public/EntitySystem/MovieSceneSharedPlaybackState.h
--- a/Engine/Source/Runtime/MovieScene/Public/EntitySystem/MovieSceneSharedPlaybackState.h
+++ b/Engine/Source/Runtime/MovieScene/Public/EntitySystem/MovieSceneSharedPlaybackState.h
@@ -22,6 +22,17 @@ struct FMovieSceneSequenceHierarchy;
 namespace UE::MovieScene
 {
 
+/**
+* Playback capability for recording audio from sequences.
+*/
+struct FMovieSceneAudioRecordingCapability
+{
+	UE_DECLARE_MOVIESCENE_PLAYBACK_CAPABILITY_API(MOVIESCENE_API, FMovieSceneAudioRecordingCapability)
+
+	// True when recording audio from scene, e.g. via MRQ.
+	bool bIsRecordingAudio = false;
+};
+
 /**
  * Parameter structure for initializing a new shared playback state.
  */
diff --git a/Engine/Source/Runtime/MovieSceneTracks/Private/Systems/MovieSceneAudioSystem.cpp b/Engine/Source/Runtime/MovieSceneTracks/Private/Systems/MovieSceneAudioSystem.cpp
--- a/Engine/Source/Runtime/MovieSceneTracks/Private/Systems/MovieSceneAudioSystem.cpp
+++ b/Engine/Source/Runtime/MovieSceneTracks/Private/Systems/MovieSceneAudioSystem.cpp
@@ -90,6 +90,13 @@ enum class EPreAnimatedAudioStateType
 {
 	/** Pre-animated state manages the lifespan of the audio component */
 	AudioComponentLifespan,
+	/** Pre-animated state manages the lifespan of the audio component.
+	* When MRQ is active, it will process UE time outside of recording to let secondary animations settle etc.
+	* For this case only, if the section extends beyond the current camera cut, the lifespan of the AudioComponent also spans the camera cut.
+	* This ensures audio capture has no gaps between shots as pause/resume (rather than stop/play) enables sample-accurate continuity across camera cuts.
+	* On capturing the next camera, playback will resume using the same AudioComponent.
+	**/
+	AudioComponentLifespanBeyondCurrentCameraCut,
 	/** Pre-animated state manages whether the audio component is playing */
 	AudioPlaying
 };
@@ -110,15 +117,22 @@ struct FPreAnimatedAudioStateTraits : BaseTraits
 	{
 		if (UAudioComponent* AudioComponent = Cast<UAudioComponent>(InKey.ResolveObjectPtr()))
 		{
+			check( InStateType == EPreAnimatedAudioStateType::AudioPlaying 
+				|| InStateType == EPreAnimatedAudioStateType::AudioComponentLifespan
+				|| InStateType == EPreAnimatedAudioStateType::AudioComponentLifespanBeyondCurrentCameraCut );
 			switch (InStateType)
 			{
+				case EPreAnimatedAudioStateType::AudioComponentLifespanBeyondCurrentCameraCut:
+					AudioComponent->SetPaused(true);
+					break;
 				case EPreAnimatedAudioStateType::AudioPlaying:
 					AudioComponent->Stop();
 					break;
 				case EPreAnimatedAudioStateType::AudioComponentLifespan:
 					AudioComponent->DestroyComponent();
+					// NOTE: if not destroyed here, it will be via ULevel::ClearLevelComponents.
 					break;
-			}
+			}			
 		}
 	}
 };
@@ -330,6 +344,66 @@ struct FEvaluateAudio
 	}
 
 private:
+	void UpdatePreAnimatedStorage(	const FMovieSceneEntityID& EntityID,
+									const FSequenceInstance& Instance,
+									const FRootInstanceHandle& RootInstanceHandle,
+									bool bWantsRestoreState,
+									const FMovieSceneContext& Context,
+									FAudioComponentEvaluationData* EvaluationData,
+									UMovieSceneAudioSection* AudioSection,
+									bool bRootMovieSceneIsRecordingAudio,
+									UAudioComponent* AudioComponent ) const
+	{
+		check(EvaluationData != nullptr);
+		check(AudioSection != nullptr);
+		check(AudioComponent != nullptr);
+
+		// When bRootMovieSceneIsRecordingAudio (e.g. MRQ is active) and the audio section doesn't end in this camera cut,
+		// use EPreAnimatedAudioStateType::AudioComponentLifespanBeyondCurrentCameraCut.
+		// This is applied on the first frame of each camera cut, flagged by setting EvaluationData->AudioComponentLifespanControlSection
+		// to the current camera cut.
+		// When not bRootMovieSceneIsRecordingAudio, EvaluationData->AudioComponentLifespanControlSection is set to the AudioSection itself,
+		// and is thus applied once for the life of the audio component / evaluation data.
+		UMovieSceneSection* CurrentLifespanControlSection = AudioSection;
+		EPreAnimatedAudioStateType LifespanType = EPreAnimatedAudioStateType::AudioComponentLifespan;
+		if (EvaluationData->AudioComponentLifespanControlSection != AudioSection && AudioSection->HasEndFrame() && bRootMovieSceneIsRecordingAudio)
+		{
+			check(Instance.GetSharedPlaybackState()->GetRootSequence() != nullptr);
+			const UMovieScene* RootMovieScene = Instance.GetSharedPlaybackState()->GetRootSequence()->GetMovieScene();
+			const UMovieSceneTrack* CameraCutTrack = RootMovieScene->GetCameraCutTrack();
+			if (CameraCutTrack)
+			{
+				const FFrameNumber CurrentFrame = Context.GetTime().FrameNumber;
+				const FFrameNumber AudioSectionEndFrame = AudioSection->GetExclusiveEndFrame();
+				for (UMovieSceneSection* CameraCutSection : CameraCutTrack->GetAllSections())
+				{
+					if (CameraCutSection && CameraCutSection->GetRange().Contains(CurrentFrame)) // determine current camera cut
+					{
+						CurrentLifespanControlSection = CameraCutSection;
+						if (!CameraCutSection->GetRange().Contains(AudioSectionEndFrame))
+						{
+							LifespanType = EPreAnimatedAudioStateType::AudioComponentLifespanBeyondCurrentCameraCut;
+						}
+						break;
+					}
+				}
+			}
+		}
+
+		if (EvaluationData->AudioComponentLifespanControlSection != CurrentLifespanControlSection)
+		{
+			if (LifespanType == EPreAnimatedAudioStateType::AudioComponentLifespanBeyondCurrentCameraCut)
+			{
+				ensure(CurrentLifespanControlSection != AudioSection); // AudioComponentLifespanBeyondCurrentCameraCut can only be used when the CurrentLifespanControlSection is a camera cut.
+			}
+			EvaluationData->AudioComponentLifespanControlSection = CurrentLifespanControlSection;
+
+			AudioSystem->PreAnimatedStorage->BeginTrackingEntity(EntityID, bWantsRestoreState, RootInstanceHandle, AudioComponent);
+			AudioSystem->PreAnimatedStorage->CachePreAnimatedValue(
+				FCachePreAnimatedValueParams(), AudioComponent,
+				[LifespanType](FObjectKey InKey) { return LifespanType; });
+		}
+	}
 
 	void Evaluate(
 			const FMovieSceneEntityID& EntityID,
@@ -355,13 +429,20 @@ private:
 		FObjectKey SectionKey(AudioSection);
 
 		const EAudioEvaluationType EvalType = GetAudioEvaluationType(Context);
+
+		bool bRootMovieSceneIsRecordingAudio = false;
+		if (FMovieSceneAudioRecordingCapability* AudioRecordingCapability = Instance.GetSharedPlaybackState()->FindCapability<FMovieSceneAudioRecordingCapability>()) // AudioRecordingCapability will be null if MRQ isn't active, leaving bRootMovieSceneIsRecordingAudio false
+		{
+			bRootMovieSceneIsRecordingAudio = AudioRecordingCapability->bIsRecordingAudio;
+		}
+
 		if (EvalType == EAudioEvaluationType::StopAndPlay)
 		{
-			AudioSystem->StopSound(InstanceHandle, ActorKey, AudioData.Section);
+			AudioSystem->StopSound(InstanceHandle, ActorKey, AudioData.Section, bRootMovieSceneIsRecordingAudio);
 		}
 		else if (EvalType == EAudioEvaluationType::Stop)
 		{
-			AudioSystem->StopSound(InstanceHandle, ActorKey, AudioData.Section);
+			AudioSystem->StopSound(InstanceHandle, ActorKey, AudioData.Section, bRootMovieSceneIsRecordingAudio);
 			return;
 		}
 		else if (EvalType == EAudioEvaluationType::Skip)
@@ -403,15 +484,7 @@ private:
 				if (World)
 				{
 					EvaluationData = AudioSystem->AddRootAudioComponent(InstanceHandle, AudioSection, World);
-					UAudioComponent* AudioComponent = EvaluationData ? EvaluationData->AudioComponent.Get() : nullptr;
-
-					if (ensure(AudioComponent))
-					{
-						AudioSystem->PreAnimatedStorage->BeginTrackingEntity(EntityID, bWantsRestoreState, RootInstanceHandle, AudioComponent);
-						AudioSystem->PreAnimatedStorage->CachePreAnimatedValue(
-							FCachePreAnimatedValueParams(), AudioComponent,
-							[](FObjectKey InKey) { return EPreAnimatedAudioStateType::AudioComponentLifespan; });
-					}
+					ensure(EvaluationData && EvaluationData->AudioComponent.Get() != nullptr);
 				}
 			}
 
@@ -421,6 +494,16 @@ private:
 
 				if (AudioComponent)
 				{
+					UpdatePreAnimatedStorage(EntityID,
+											Instance,
+											RootInstanceHandle,
+											bWantsRestoreState,
+											Context,
+											EvaluationData,
+											AudioSection,
+											bRootMovieSceneIsRecordingAudio,
+											AudioComponent);
+
 					if (AttachComponent && (AudioComponent->GetAttachParent() != AttachComponent || AudioComponent->GetAttachSocketName() != AttachKey.SocketName))
 					{
 						AudioComponent->AttachToComponent(AttachComponent, FAttachmentTransformRules::KeepRelativeTransform, AttachKey.SocketName);
@@ -473,15 +556,6 @@ private:
 			{
 				// Initialize the sound
 				EvaluationData = AudioSystem->AddBoundObjectAudioComponent(InstanceHandle, AudioSection, BoundObject);
-				UAudioComponent* AudioComponent = EvaluationData ? EvaluationData->AudioComponent.Get() : nullptr;
-
-				if (AudioComponent)
-				{
-					AudioSystem->PreAnimatedStorage->BeginTrackingEntity(EntityID, bWantsRestoreState, RootInstanceHandle, AudioComponent);
-					AudioSystem->PreAnimatedStorage->CachePreAnimatedValue(
-						FCachePreAnimatedValueParams(), AudioComponent,
-						[](FObjectKey InKey) { return EPreAnimatedAudioStateType::AudioComponentLifespan; });
-				}
 			}
 
 			if (EvaluationData)
@@ -490,6 +564,16 @@ private:
 
 				if (AudioComponent)
 				{
+					UpdatePreAnimatedStorage(EntityID,
+											Instance,
+											RootInstanceHandle,
+											bWantsRestoreState,
+											Context,
+											EvaluationData,
+											AudioSection,
+											bRootMovieSceneIsRecordingAudio,
+											AudioComponent);
+
 					if (AudioSection->GetOnQueueSubtitles().IsBound())
 					{
 						AudioComponent->OnQueueSubtitles = AudioSection->GetOnQueueSubtitles();
@@ -660,7 +744,7 @@ private:
 					{
 						// If we're not a looping section and the AC is done, return
 						// otherwise a looping section will restart the sound if it's not playing
-						if (!AudioComponent.IsPlaying())
+						if (!AudioComponent.IsPlaying()) // !IsPlaying() == not playing, fading, or paused
 						{
 							UE_LOGF(LogMovieScene, Verbose, "Procedural sound Audio Component reached end of playback. Component: %ls Sound: %ls", *AudioComponent.GetName(), *GetNameSafe(AudioComponent.Sound));
 							return;
@@ -680,7 +764,7 @@ private:
 				if (!AudioSection.GetLooping() && AudioTime >= Duration && Duration != 0.f)
 				{
 					// If this is Non-ProcedurallyGenerated audio, and it's not looping then check to see if it needs to be stopped
-					if (AudioComponent.IsPlaying())
+					if (AudioComponent.IsPlaying()) // IsPlaying() == playing, fading, or paused
 					{
 						UE_LOGF(LogMovieScene, Verbose, "Audio Component reached end of playback. Component: %ls Sound: %ls", *AudioComponent.GetName(), *GetNameSafe(AudioComponent.Sound));
 						AudioComponent.Stop();
@@ -782,7 +866,7 @@ private:
 			{
 				// Only call stop on the sound if it is actually playing. This prevents spamming
 				// stop calls when a sound cue with a duration of zero is played.
-				if (AudioComponent.IsPlaying() || bSoundNeedsTimeSync)
+				if ((AudioComponent.IsPlaying() && !bAudioComponentIsPaused) || bSoundNeedsTimeSync)
 				{
 					UE_LOGF(LogMovieScene, Verbose, "Audio Component stopped due to needing a state change bIsPlaying: %d bNeedsTimeSync: %d. Component: %ls Sound: %ls", AudioComponent.IsPlaying(), bSoundNeedsTimeSync, *AudioComponent.GetName(), *GetNameSafe(AudioComponent.Sound));
 					AudioComponent.Stop();
@@ -824,7 +908,7 @@ private:
 				if (bPerformGranularScrubbing)
 				{
 					// If we're not playing the audio component, then set up the audio component to use the scrubbed sound wave and play it
-					if (!AudioComponent.IsPlaying())
+					if (!AudioComponent.IsPlaying()) // !IsPlaying() == not playing, fading, or paused
 					{
 						ScrubbedSound->SetSoundWave(SoundWave);
 						ScrubbedSound->ModulationSettings = SoundWave->ModulationSettings;
@@ -862,7 +946,14 @@ private:
 						AudioComponent.SetSound(Sound);
 					}
 					
-					AudioComponent.Play(AudioTime);
+					if (!AudioComponent.IsPlaying() || !bAudioComponentIsPaused) // !IsPlaying() == not playing, fading, or paused
+					{
+						AudioComponent.Play(AudioTime);
+					}
+					else
+					{
+						AudioComponent.SetPaused(false); // no-op if already unpaused
+					}
 					
 					if (Context.GetStatus() == EMovieScenePlayerStatus::Playing)
 					{
@@ -917,7 +1008,7 @@ private:
 			AudioComponent.StopDelayed(ScrubDuration);
 		}
 
-		if (AudioComponent.IsPlaying() && AudioInputs)
+		if ((AudioComponent.IsPlaying() && !bAudioComponentIsPaused) && AudioInputs)
 		{
 			SetAudioInputTriggers(AudioInputs->Inputs_Trigger, AudioComponent);
 		}
@@ -1390,7 +1481,7 @@ UMovieSceneAudioSystem::FAudioComponentEvaluationData* UMovieSceneAudioSystem::A
 	return ExistingData;
 }
 
-void UMovieSceneAudioSystem::StopSound(FInstanceHandle InstanceHandle, FObjectKey ActorKey, FObjectKey SectionKey)
+void UMovieSceneAudioSystem::StopSound(FInstanceHandle InstanceHandle, FObjectKey ActorKey, FObjectKey SectionKey, bool bIsRecordingAudio)
 {
 	if (FAudioComponentBySectionKey* Map = AudioComponentsByActorKey.Find(ActorKey))
 	{
@@ -1399,7 +1490,14 @@ void UMovieSceneAudioSystem::StopSound(FInstanceHandle InstanceHandle, FObjectKe
 		{
 			if (UAudioComponent* AudioComponent = Data->AudioComponent.Get())
 			{
-				AudioComponent->Stop();
+				if (!bIsRecordingAudio)
+				{
+					AudioComponent->Stop();
+				}
+				else
+				{
+					AudioComponent->SetPaused(true);
+				}
 			}
 
 			Data->bAudioComponentHasBeenPlayed = false;
diff --git a/Engine/Source/Runtime/MovieSceneTracks/Private/Systems/MovieSceneAudioSystem.h b/Engine/Source/Runtime/MovieSceneTracks/Private/Systems/MovieSceneAudioSystem.h
--- a/Engine/Source/Runtime/MovieSceneTracks/Private/Systems/MovieSceneAudioSystem.h
+++ b/Engine/Source/Runtime/MovieSceneTracks/Private/Systems/MovieSceneAudioSystem.h
@@ -15,6 +15,7 @@
 
 class IAudioParameterControllerInterface;
 class UAudioComponent;
+class UMovieSceneSection;
 class UMovieSceneAudioSection;
 class USoundBase;
 struct FMoveSceneAudioTriggerState;
@@ -40,6 +41,9 @@ namespace UE::MovieScene
 		/** The audio component that was created to play audio */
 		TWeakObjectPtr<UAudioComponent> AudioComponent;
 
+		/** The movie section managing the lifespan of the audio component */
+		TObjectPtr<UMovieSceneSection> AudioComponentLifespanControlSection = nullptr;
+
 #if WITH_EDITOR
 		/** While in editor, we can scrub the audio in the audio component. */
 		TObjectPtr<UScrubbedSound> ScrubbedSound;
@@ -111,9 +115,9 @@ public:
 	FAudioComponentEvaluationData* AddRootAudioComponent(FInstanceHandle InstanceHandle, UMovieSceneAudioSection* Section, UWorld* World);
 
 	/**
-	 * Stop the audio on the audio component associated with the given audio section.
+	 * Stop (or just pause if bIsRecordingAudio) the audio on the audio component associated with the given audio section.
 	 */
-	void StopSound(FInstanceHandle InstanceHandle, FObjectKey ActorKey, FObjectKey SectionKey);
+	void StopSound(FInstanceHandle InstanceHandle, FObjectKey ActorKey, FObjectKey SectionKey, bool bIsRecordingAudio);
 
 	/**
 	 * Reset shared accumulation data required every evaluation frame
