[Sequencer] Fix directorBP condition endpoints not evaluating in PIE when CallInEditor is unchecked

#rnx
#Synced-CL 52723273

diff --git a/Engine/Source/Runtime/MovieSceneTracks/Private/Conditions/MovieSceneDirectorBlueprintCondition.cpp b/Engine/Source/Runtime/MovieSceneTracks/Private/Conditions/MovieSceneDirectorBlueprintCondition.cpp
--- a/Engine/Source/Runtime/MovieSceneTracks/Private/Conditions/MovieSceneDirectorBlueprintCondition.cpp
+++ b/Engine/Source/Runtime/MovieSceneTracks/Private/Conditions/MovieSceneDirectorBlueprintCondition.cpp
@@ -74,9 +74,13 @@ bool FMovieSceneDirectorBlueprintConditionInvoker::InvokeDirectorBlueprintCondit
 	// Not sure why this isn't being checked further down, but check manually here
 	if (!ConditionFunc->HasMetaData(FBlueprintMetadata::MD_CallInEditor) || ConditionFunc->GetMetaData(FBlueprintMetadata::MD_CallInEditor) != TEXT("true"))
 	{
-		if (ConditionContext.WorldContext && ConditionContext.WorldContext->GetWorld()->IsEditorWorld())
+		if (ConditionContext.WorldContext)
 		{
-			return Result;
+			const UWorld* World = ConditionContext.WorldContext->GetWorld();
+			if (World && World->IsEditorWorld() && !World->IsPlayInEditor())
+			{
+				return Result;
+			}
 		}
 	}
 #endif
